Fixing build - Integrating CL#11106093 ,11107532 and 11108073

diff --git a/Engine/Plugins/Experimental/Text3D/Source/Text3D/Private/Text3DPrivate.h b/Engine/Plugins/Experimental/Text3D/Source/Text3D/Private/Text3DPrivate.h
--- a/Engine/Plugins/Experimental/Text3D/Source/Text3D/Private/Text3DPrivate.h
+++ b/Engine/Plugins/Experimental/Text3D/Source/Text3D/Private/Text3DPrivate.h
@@ -4,12 +4,20 @@
 #include "CoreMinimal.h"
 #include "Modules/ModuleManager.h"
 
+#if PLATFORM_COMPILER_HAS_GENERIC_KEYWORD
+#define generic __identifier(generic)
+#endif	//PLATFORM_COMPILER_HAS_GENERIC_KEYWORD
+
 THIRD_PARTY_INCLUDES_START
 #include "ft2build.h"
 #include FT_FREETYPE_H
 #include FT_ADVANCES_H
 THIRD_PARTY_INCLUDES_END
 
+#if PLATFORM_COMPILER_HAS_GENERIC_KEYWORD
+#undef generic
+#endif	//PLATFORM_COMPILER_HAS_GENERIC_KEYWORD
+
 const int32 FontSize = 64;
 const float FontInverseScale = 1.0f / FontSize;
 
diff --git a/Engine/Source/Developer/MeshBuilder/Private/BuildStatisticManager.h b/Engine/Source/Developer/MeshBuilder/Private/BuildStatisticManager.h
--- a/Engine/Source/Developer/MeshBuilder/Private/BuildStatisticManager.h
+++ b/Engine/Source/Developer/MeshBuilder/Private/BuildStatisticManager.h
@@ -3,7 +3,7 @@
 
 #include "CoreMinimal.h"
 #include "CoreTypes.h"
-#include "GenericPlatform/GenericPlatformTime.h"
+#include "HAL/PlatformTime.h"
 
 DECLARE_LOG_CATEGORY_EXTERN(LogMeshDescriptionBuildStatistic, Log, All);
 
diff --git a/Engine/Source/Runtime/Datasmith/DatasmithCore/Private/DatasmithUtils.cpp b/Engine/Source/Runtime/Datasmith/DatasmithCore/Private/DatasmithUtils.cpp
--- a/Engine/Source/Runtime/Datasmith/DatasmithCore/Private/DatasmithUtils.cpp
+++ b/Engine/Source/Runtime/Datasmith/DatasmithCore/Private/DatasmithUtils.cpp
@@ -31,7 +31,7 @@ void FDatasmithUtils::SanitizeNameInplace(FString& InString)
 {
 	static const TCHAR Original[] = TEXT("ÀÁÂÃÄÅÆÇÈÉÊËÌÍÎÏÐÑÒÓÔÕÖ×ØÙÚÛÜÝÞßàáâãäåæçèéêëìíîïðñòóôõö÷øùúûüýþÿБбВвГгДдЁёЖжЗзИиЙйКкЛлМмНнОоПпРрСсТтУуФфХхЦцЧчШшЩщЪъЫыЬьЭэЮюЯя'\"");
 	static const TCHAR Modified[] = TEXT("AAAAAAECEEEEIIIIDNOOOOOx0UUUUYPsaaaaaaeceeeeiiiiOnoooood0uuuuypyBbVvGgDdEeJjZzIiYyKkLlMmNnOoPpRrSsTtUuFfJjTtCcSsSs__ii__EeYyYy__");
-	static_assert(GetArrayLength(Original) == GetArrayLength(Modified), "array size mismatch");
+	static_assert(!PLATFORM_DESKTOP || GetArrayLength(Original) == GetArrayLength(Modified), "array size mismatch");
 
 	for (int32 i = 0; i < GetArrayLength(Original); i++)
 	{
