The problem was a post process of the meta data to remove unprintable characteres. Process that removed the end character from the string...
So configurations of the parts became different for each instance of part...

diff --git a/Engine/Source/Runtime/Datasmith/DatasmithCore/Private/DatasmithUtils.cpp b/Engine/Source/Runtime/Datasmith/DatasmithCore/Private/DatasmithUtils.cpp
--- a/Engine/Source/Runtime/Datasmith/DatasmithCore/Private/DatasmithUtils.cpp
+++ b/Engine/Source/Runtime/Datasmith/DatasmithCore/Private/DatasmithUtils.cpp
@@ -47,7 +47,7 @@ void FDatasmithUtils::SanitizeNameInplace(FString& InString)
 
 void FDatasmithUtils::SanitizeStringInplace(FString& InString)
 {
-	for (TCHAR& Char : InString.GetCharArray())
+	for (TCHAR& Char : InString)
 	{
 		if (!FChar::IsPrint(Char) && !FChar::IsWhitespace(Char))
 		{
