DirectLink: linkage fix for ps4 modular build

diff --git a/Engine/Source/Runtime/Datasmith/DirectLink/Private/DirectLinkParameterStore.cpp b/Engine/Source/Runtime/Datasmith/DirectLink/Private/DirectLinkParameterStore.cpp
--- a/Engine/Source/Runtime/Datasmith/DirectLink/Private/DirectLinkParameterStore.cpp
+++ b/Engine/Source/Runtime/Datasmith/DirectLink/Private/DirectLinkParameterStore.cpp
@@ -14,6 +14,19 @@
 namespace DirectLink
 {
 
+// -- ps4 fix --
+// Fix for a modular build on PS4 platform, where orbis-ld requires the following symbols to be
+// accessible when compiling DatasmithCore as a shared object.
+// Generated methods are not enough somehow, so they are explicitely defined here, in the cpp,
+// which fixes the visibility issue.
+FParameterStore::FParameterStore() = default;
+FParameterStore::~FParameterStore() = default;
+FParameterStore::FParameterStore(const FParameterStore&) = default;
+FParameterStore& FParameterStore::operator=(const FParameterStore&) = default;
+FParameterStore::FParameterStore(FParameterStore&&) = default;
+FParameterStore& FParameterStore::operator=(FParameterStore&&) = default;
+// -- end ps4 fix --
+
 uint32 FParameterStore::GetParameterCount() const
 {
 	return Parameters.Num();
diff --git a/Engine/Source/Runtime/Datasmith/DirectLink/Public/DirectLinkParameterStore.h b/Engine/Source/Runtime/Datasmith/DirectLink/Public/DirectLinkParameterStore.h
--- a/Engine/Source/Runtime/Datasmith/DirectLink/Public/DirectLinkParameterStore.h
+++ b/Engine/Source/Runtime/Datasmith/DirectLink/Public/DirectLinkParameterStore.h
@@ -103,6 +103,13 @@ private:
 class DIRECTLINK_API FParameterStore
 {
 public:
+	FParameterStore();
+	~FParameterStore();
+	FParameterStore(const FParameterStore&);
+	FParameterStore& operator=(const FParameterStore&);
+	FParameterStore(FParameterStore&&);
+	FParameterStore& operator=(FParameterStore&&);
+
 	template<typename T, typename S>
 	TStoreKey<T, S>& RegisterParameter(TStoreKey<T, S>& Key, FName Name)
 	{
