CADKernel Source Code v0.01
The review of the source code is in progress.

diff --git a/Engine/Source/Runtime/Datasmith/CADKernel/Private/CADKernel/Core/Database.cpp b/Engine/Source/Runtime/Datasmith/CADKernel/Private/CADKernel/Core/Database.cpp
new file mode 100644
--- /dev/null
+++ b/Engine/Source/Runtime/Datasmith/CADKernel/Private/CADKernel/Core/Database.cpp
@@ -0,0 +1,477 @@
+// Copyright Epic Games, Inc. All Rights Reserved.
+
+#include "CADKernel/Core/Database.h"
+
+#include "CADKernel/Core/CADKernelArchive.h"
+#include "CADKernel/Core/Group.h"
+#include "CADKernel/Core/Session.h"
+#include "CADKernel/Mesh/Structure/ModelMesh.h"
+#include "CADKernel/Topo/Body.h"
+#include "CADKernel/Topo/TopologicalEdge.h"
+#include "CADKernel/Topo/Linkable.h"
+#include "CADKernel/Topo/Model.h"
+#include "CADKernel/Topo/Shell.h"
+#include "CADKernel/Topo/TopologicalVertex.h"
+#include "CADKernel/Ui/Message.h"
+
+#include "HAL/FileManager.h"
+
+using namespace CADKernel;
+
+static const int32 DataBaseInitialSize = 10000;
+
+FDatabase::FDatabase()
+{
+	Model = FEntity::MakeShared<FModel>();
+
+	DatabaseEntities.Reserve(DataBaseInitialSize);
+	AvailableIdents.Reserve(DataBaseInitialSize);
+	DatabaseEntities.Add(TSharedPtr<FEntity>());
+}
+
+TSharedRef<FModel> FDatabase::GetModel()
+{
+	if (!Model.IsValid())
+	{
+		Model = FEntity::MakeShared<FModel>();
+	}
+	return Model.ToSharedRef();
+}
+
+FIdent FDatabase::CreateId()
+{
+	if (AvailableIdents.Num() > 0)
+	{
+		return  AvailableIdents.Pop(false);
+	}
+	else
+	{
+		FIdent NewIdent = (FIdent)DatabaseEntities.Num();
+		DatabaseEntities.Add(TSharedPtr<FEntity>());
+		return NewIdent;
+	}
+}
+
+void FDatabase::RemoveEntity(FIdent EntityId)
+{
+	if (EntityId >= (FIdent)DatabaseEntities.Num())
+	{
+		FMessage::Printf(Debug, TEXT("Warning in FDataBasse::RemoveEntity: the Id=%d is not yet defined, the entity has never exist.\n"), EntityId);
+	}
+	else if (DatabaseEntities[EntityId].IsValid())
+	{
+		DatabaseEntities[EntityId] = TSharedPtr<FEntity>();
+		AvailableIdents.Add(EntityId);
+	}
+	else
+	{
+		FMessage::Printf(Debug, TEXT("Warning in FDataBasse::RemoveEntity: the entity with Id=%d is already deleted.\n"), EntityId);
+	}
+}
+
+TSharedPtr<FEntity> FDatabase::GetEntity(FIdent EntityId) const
+{
+	if (EntityId >= (FIdent)DatabaseEntities.Num())
+	{
+		FMessage::Printf(Debug, TEXT("Warning in FDataBasse::GetEntity: the Id=%d is not yet defined\n"), EntityId);
+		return TSharedPtr<FEntity>();
+	}
+	if (!DatabaseEntities[EntityId].IsValid())
+	{
+		FMessage::Printf(Debug, TEXT("Warning in FDataBasse::GetEntity: the entity with Id=%d no longer exists.\n"), EntityId);
+	}
+	return DatabaseEntities[EntityId];
+}
+
+void FDatabase::GetEntities(const TArray<FIdent>& EntityIds, TArray<TSharedPtr<FEntity>>& Entities) const
+{
+	Entities.Empty(EntityIds.Num());
+
+	// to avoid duplicated components
+	TSet<FIdent> AddedComponents;
+	AddedComponents.Reserve(EntityIds.Num());
+
+	for (FIdent EntityId : EntityIds)
+	{
+		TSharedPtr<FEntity> Entity = GetEntity(EntityId);
+		if (AddedComponents.Find(EntityId) == nullptr)
+		{
+			AddedComponents.Add(EntityId);
+			Entities.Add(Entity);
+		}
+	}
+}
+
+void FCADKernelArchive::SetReferencedEntityOrAddToWaitingList(FIdent ArchiveId, TWeakPtr<FEntity>& Entity)
+{
+	Session.Database.SetReferencedEntityOrAddToWaitingList(ArchiveId, Entity);
+}
+
+void FCADKernelArchive::SetReferencedEntityOrAddToWaitingList(FIdent ArchiveId, TSharedPtr<FEntity>& Entity)
+{
+	Session.Database.SetReferencedEntityOrAddToWaitingList(ArchiveId, Entity);
+}
+
+void FCADKernelArchive::AddEntityToSave(FIdent Id)
+{
+	Session.Database.AddEntityToSave(Id);
+}
+
+void FCADKernelArchive::AddEntityFromArchive(TSharedRef<FEntity>& Entity)
+{
+	Session.Database.AddEntityFromArchive(Entity);
+}
+
+void FDatabase::SerializeSelection(FCADKernelArchive& Ar, const TArray<FIdent>& SelectionIds)
+{
+	TRACE_CPUPROFILER_EVENT_SCOPE(FDatabase::Serialize);
+
+	ensure(Ar.IsSaving());
+
+	bIsRecursiveSerialization = true;
+
+	Ar.Session.Serialize(Ar);
+
+	int32 DatabaseSize = DatabaseEntities.Num();
+	Ar << DatabaseSize;
+
+	EntitiesToBeSerialized.Empty();
+	NotYetSerialized.Empty();
+	for (FIdent EntityId : SelectionIds)
+	{
+		AddEntityToSave(EntityId);
+	}
+
+	int32 Index = 0;
+	FIdent EntityIdToSave;
+	while(NotYetSerialized.Dequeue(EntityIdToSave))
+	{
+		TSharedPtr<FEntity> Entity = GetEntity(EntityIdToSave);
+		if (!Entity.IsValid())
+		{
+			continue;
+		}
+
+		EEntity Type = Entity->GetEntityType();
+		Ar << Type;
+		Entity->Serialize(Ar);
+		Index++;
+	}
+
+	EntitiesToBeSerialized.Empty();
+	NotYetSerialized.Empty();
+	bIsRecursiveSerialization = false;
+
+	FMessage::Printf(Log, TEXT("End Serialisation of %d entities\n"), Index);
+}
+
+void FDatabase::Serialize(FCADKernelArchive& Ar)
+{
+	TRACE_CPUPROFILER_EVENT_SCOPE(FDatabase::Serialize);
+
+	ensure(Ar.IsSaving());
+
+	SpawnEntityIdent((TSharedPtr<FEntity>&) Model, true);
+
+	Ar.Session.Serialize(Ar);
+
+	int32 DatabaseSize = DatabaseEntities.Num();
+	Ar << DatabaseSize;
+
+	int32 Index = 0;
+	FProgress Progess(DatabaseEntities.Num());
+	for (TSharedPtr<FEntity> Entity : DatabaseEntities)
+	{
+		Progess.Increase();
+
+		if (!Entity.IsValid())
+		{
+			EEntity Type = EEntity::NullEntity;
+			Ar << Type;
+			continue;
+		}
+
+		EEntity Type = Entity->GetEntityType();
+		Ar << Type;
+		Entity->Serialize(Ar);
+		Index++;
+	}
+	FMessage::Printf(Log, TEXT("End Serialisation of %d %d entity\n"), DatabaseEntities.Num(), Index);
+}
+
+void FDatabase::Deserialize(FCADKernelArchive& Ar)
+{
+	TRACE_CPUPROFILER_EVENT_SCOPE(FDatabase::Deserialize);
+
+	ensure(Ar.IsLoading());
+
+	Ar.Session.Serialize(Ar);
+
+	int32 ArchiveSize = 0;
+	Ar << ArchiveSize;
+
+	DatabaseEntities.Reserve((int32)(DatabaseEntities.Num() + ArchiveSize * 1.5));
+	ArchiveIdToWaitingPointers.SetNum(ArchiveSize);
+	ArchiveIdToWaitingWeakPointers.SetNum(ArchiveSize);
+	ArchiveEntities.Init(TSharedPtr<FEntity>(), ArchiveSize);
+
+	int64 TotalSize = Ar.TotalSize();
+	FProgress Progress(ArchiveSize);
+	int32 Index = 0;
+	for( ; Ar.Tell() < TotalSize; ++Index)
+	{
+		Progress.Increase();
+		TSharedPtr<FEntity> Entity = FEntity::Deserialize(Ar);
+		if (!Entity.IsValid())
+		{
+			FMessage::Printf(Log, TEXT("Failed at index %d\n"), Index);
+		}
+	}
+
+	// this is a partial archive, a clean of entities has to be performed to remove invalid references in entity descriptions
+	if(Index < ArchiveSize)
+	{
+		CleanArchiveEntities();
+	}
+
+	ArchiveIdToWaitingPointers.Empty();
+	ArchiveEntities.Empty();
+	
+	FMessage::Printf(Log, TEXT("End Deserialisation of %d %d entity\n"), DatabaseEntities.Num(), ArchiveSize);
+}
+
+void FDatabase::CleanArchiveEntities()
+{
+	for (TSharedPtr<FEntity> Entity : ArchiveEntities)
+	{
+		if (Entity.IsValid())
+		{
+			switch (Entity->GetEntityType())
+			{
+			case EEntity::EdgeLink:
+			{
+				StaticCastSharedPtr<FEdgeLink>(Entity)->CleanLink();
+				break;
+			}
+			case EEntity::VertexLink:
+			{
+				StaticCastSharedPtr<FVertexLink>(Entity)->CleanLink();
+				break;
+			}
+			}
+		}
+	}
+}
+
+void FDatabase::Empty()
+{
+	Model.Reset();
+	DatabaseEntities.Empty(DataBaseInitialSize);
+	DatabaseEntities.Add(TSharedPtr<FEntity>());
+	AvailableIdents.Empty(DataBaseInitialSize);
+}
+
+uint32 FDatabase::SpawnEntityIdent(const TArray<TSharedPtr<FEntity>>& SelectedEntities, bool bInForceSpawning)
+{
+	bForceSpawning = bInForceSpawning;
+	EntityCount = 0;
+	for (TSharedPtr<FEntity> Entity : SelectedEntities)
+	{
+		Entity->SpawnIdent(*this);
+	}
+	bForceSpawning = false;
+	return EntityCount;
+}
+
+uint32 FDatabase::SpawnEntityIdent(TSharedPtr<FEntity>& Entity, bool bInForceSpawning)
+{
+	bForceSpawning = bInForceSpawning;
+	EntityCount = 0;
+	Entity->SpawnIdent(*this);
+	bForceSpawning = false;
+	return EntityCount;
+}
+
+
+// =========================================================================================================================================================================================================
+// =========================================================================================================================================================================================================
+// =========================================================================================================================================================================================================
+//
+//
+//                                                                            NOT YET REVIEWED
+//
+//
+// =========================================================================================================================================================================================================
+// =========================================================================================================================================================================================================
+// =========================================================================================================================================================================================================
+
+void FDatabase::ExpandSelection(const TArray<TSharedPtr<FEntity>>& Entities, const TSet<EEntity>& Filters, TArray<TSharedPtr<FEntity>>& Selection) const
+{
+	TSet<TSharedPtr<FEntity>> EntitySet;
+	ExpandSelection(Entities, Filters, EntitySet);
+	for (TSharedPtr<FEntity> Entity : EntitySet)
+	{
+		Selection.Add(Entity);
+	}
+}
+
+void FDatabase::ExpandSelection(const TArray<TSharedPtr<FEntity>>& Entities, const TSet<EEntity>& Filter, TSet<TSharedPtr<FEntity>>& Selection) const
+{
+	for (TSharedPtr<FEntity> Entity : Entities)
+	{
+		ExpandSelection(Entity, Filter, Selection);
+	}
+}
+
+void FDatabase::ExpandSelection(TSharedPtr<FEntity> Entity, const TSet<EEntity>& Filter, TSet<TSharedPtr<FEntity>>& Selection) const
+{
+	EEntity Type = Entity->GetEntityType();
+
+	if (Filter.Contains(Type))
+	{
+		Selection.Add(Entity);
+	}
+
+	switch (Entity->GetEntityType())
+	{
+
+	case EEntity::TopologicalEdge:
+	{
+		ensureCADKernel(false);
+		TSharedPtr<FTopologicalEdge>Edge = StaticCastSharedPtr<FTopologicalEdge>(Entity);
+		ExpandSelection(Edge->GetCurve(), Filter, Selection);
+		break;
+	}
+
+	case EEntity::Group:
+	{
+		ensureCADKernel(false);
+		TSharedPtr<FGroup>Group = StaticCastSharedPtr<FGroup>(Entity);
+		for (TSharedPtr<FEntity> GroupEntity : Group->GetEntities())
+		{
+			ExpandSelection(GroupEntity, Filter, Selection);
+		}
+		break;
+	}
+
+	case EEntity::Model:
+	{
+		TSharedPtr<FModel> InModel = StaticCastSharedPtr<FModel>(Entity);
+		for (TSharedPtr<FBody> Body : InModel->GetBodyList())
+		{
+			ExpandSelection(Body, Filter, Selection);
+		}
+
+		for (TSharedPtr<FTopologicalFace> Domain : InModel->GetFaces())
+		{
+			ExpandSelection(Domain, Filter, Selection);
+		}
+		break;
+	}
+
+	case EEntity::Body:
+	{
+		TSharedPtr<FBody> Body = StaticCastSharedPtr<FBody>(Entity);
+		for (const TSharedPtr<FShell> Shell : Body->GetShells())
+		{
+			for (const FOrientedFace& Face : Shell->GetFaces())
+			{
+				ExpandSelection(Face.Entity, Filter, Selection);
+			}
+		}
+		break;
+	}
+
+	case EEntity::MeshModel:
+	{
+		TSharedPtr<FModelMesh> MeshModel = StaticCastSharedPtr<FModelMesh>(Entity);
+		for (TSharedPtr<FMesh> Mesh : MeshModel->GetMeshes())
+		{
+			ExpandSelection(Mesh, Filter, Selection);
+		}
+		break;
+	}
+
+	default:
+	{
+		break;
+	}
+	}
+
+}
+
+void FDatabase::TopologicalEntitiesSelection(const TArray<TSharedPtr<FEntity>>& Entities, TArray<TSharedPtr<FTopologicalEntity>>& OutTopoEntities, bool bFindSurfaces, bool bFindEdges, bool bFindVertices) const
+{
+	TSet<EEntity> Filter;
+	if (bFindSurfaces)	Filter.Add(EEntity::TopologicalFace);
+	if (bFindEdges)		Filter.Add(EEntity::TopologicalEdge);
+	if (bFindVertices) 	Filter.Add(EEntity::TopologicalVertex);
+
+	OutTopoEntities.Empty(FMath::Max(Entities.Num(), (int32)100));
+	ExpandSelection(Entities, Filter, (TArray<TSharedPtr<FEntity>>&) OutTopoEntities);
+}
+
+void FDatabase::CadEntitiesSelection(const TArray<TSharedPtr<FEntity>>& Entities, TArray<TSharedPtr<FEntity>>& OutCADEntities) const
+{
+	TSet<EEntity> Filter;
+	Filter.Add(EEntity::TopologicalFace);
+	Filter.Add(EEntity::TopologicalEdge);
+	Filter.Add(EEntity::TopologicalVertex);
+	Filter.Add(EEntity::Curve);
+	Filter.Add(EEntity::Surface);
+
+	OutCADEntities.Empty(FMath::Max(Entities.Num(), (int32)100));
+	ExpandSelection(Entities, Filter, OutCADEntities);
+}
+
+void FDatabase::MeshEntitiesSelection(const TArray<TSharedPtr<FEntity>>& Entities, TArray<TSharedPtr<FEntity>>& OutMeshEntities) const
+{
+	TSet<EEntity> Filter;
+	Filter.Add(EEntity::Mesh);
+
+	OutMeshEntities.Empty(FMath::Max(Entities.Num(), (int32)100));
+	ExpandSelection(Entities, Filter, OutMeshEntities);
+}
+
+void FDatabase::GetEntitiesOfType(EEntity Type, TArray<TSharedPtr<FEntity>>& OutEntities) const
+{
+	TSet<EEntity> Filter;
+	Filter.Add(Type);
+
+	GetEntitiesOfTypes(Filter, OutEntities);
+}
+
+void FDatabase::GetEntitiesOfTypes(const TSet<EEntity>& Filter, TArray<TSharedPtr<FEntity>>& OutEntities) const
+{
+	OutEntities.Empty(100);
+	for (TSharedPtr<FEntity> Entity : DatabaseEntities)
+	{
+		if (!Entity.IsValid())
+		{
+			continue;
+		}
+
+		if (Filter.Contains(Entity->GetEntityType()))
+		{
+			OutEntities.Add(Entity);
+		}
+	}
+}
+
+TSharedPtr<FEntity>FDatabase::GetFirstEntityOfType(EEntity Type) const
+{
+	for (TSharedPtr<FEntity> Entity : DatabaseEntities)
+	{
+		if (!Entity.IsValid())
+		{
+			continue;
+		}
+
+		if (Entity->GetEntityType() == Type)
+		{
+			return Entity;
+		}
+	}
+	return nullptr;
+}
+
diff --git a/Engine/Source/Runtime/Datasmith/CADKernel/Private/CADKernel/Core/Entity.cpp b/Engine/Source/Runtime/Datasmith/CADKernel/Private/CADKernel/Core/Entity.cpp
new file mode 100644
--- /dev/null
+++ b/Engine/Source/Runtime/Datasmith/CADKernel/Private/CADKernel/Core/Entity.cpp
@@ -0,0 +1,291 @@
+// Copyright Epic Games, Inc. All Rights Reserved.
+
+#include "CADKernel/Core/Entity.h"
+
+#include "CADKernel/Core/EntityGeom.h"
+
+#include "CADKernel/Core/Database.h"
+#include "CADKernel/Core/Group.h"
+#include "CADKernel/Core/Session.h"
+#include "CADKernel/Core/System.h"
+#include "CADKernel/Geo/Curves/Curve.h"
+#include "CADKernel/Geo/Surfaces/Surface.h"
+#include "CADKernel/Mesh/Criteria/Criterion.h"
+#include "CADKernel/Topo/Body.h"
+#include "CADKernel/Topo/Model.h"
+#include "CADKernel/Topo/Shell.h"
+#include "CADKernel/Topo/TopologicalEdge.h"
+#include "CADKernel/Topo/TopologicalFace.h"
+#include "CADKernel/Topo/TopologicalLink.h"
+#include "CADKernel/Topo/TopologicalLoop.h"
+#include "CADKernel/Topo/TopologicalVertex.h"
+
+namespace CADKernel
+{
+	const TCHAR* FEntity::TypesNames[] = 
+	{
+		TEXT("NullEntity"),
+		TEXT("Curve"),
+		TEXT("Surface"),
+
+		TEXT("Edge Link"),
+		TEXT("Vertex Link"),
+		TEXT("Edge"),
+		TEXT("Face"),
+ 		TEXT("Link"),
+		TEXT("Loop"),
+		TEXT("Vertex"),
+		TEXT("Shell"),
+		TEXT("Body"),
+		TEXT("Model"),
+
+		TEXT("Mesh Model"),
+		TEXT("Mesh"),
+
+		TEXT("Group"),
+		TEXT("Criterion"),
+		TEXT("Property"),
+		nullptr
+	};
+
+	const TCHAR* FEntity::GetTypeName(EEntity Type)
+	{
+		if (Type >= EEntity::EntityTypeEnd)
+		{
+			return TypesNames[0];
+		}
+		return TypesNames[(int32) Type];
+	}
+
+	FEntity::~FEntity()
+	{
+		TRACE_CPUPROFILER_EVENT_SCOPE(FEntity::~FEntity);
+	}
+
+	TSharedPtr<FEntity> FEntity::Deserialize(FCADKernelArchive& Archive)
+	{
+		TRACE_CPUPROFILER_EVENT_SCOPE(FEntity::Deserialize);
+
+		ensureCADKernel(Archive.IsLoading());
+
+		EEntity Type = EEntity::NullEntity;
+		Archive << Type;
+
+		switch (Type)
+		{
+		case EEntity::Body:
+			return FEntity::MakeShared<FBody>(Archive);
+		case EEntity::Curve:
+			return FCurve::Deserialize(Archive);
+		case EEntity::Criterion:
+			return FCriterion::Deserialize(Archive);
+		case EEntity::EdgeLink:
+			return FEntity::MakeShared<FEdgeLink>(Archive);
+		case EEntity::Group:
+			return FEntity::MakeShared<FGroup>(Archive);
+		case EEntity::Model:
+			return FEntity::MakeShared<FModel>(Archive);
+		case EEntity::Shell:
+			return FEntity::MakeShared<FShell>(Archive);
+		case EEntity::Surface:
+			return FSurface::Deserialize(Archive);
+		case EEntity::TopologicalEdge:
+			return FEntity::MakeShared<FTopologicalEdge>(Archive);
+		case EEntity::TopologicalFace:
+			return FEntity::MakeShared<FTopologicalFace>(Archive);
+		case EEntity::TopologicalLoop:
+			return FEntity::MakeShared<FTopologicalLoop>(Archive);
+		case EEntity::TopologicalVertex:
+			return FEntity::MakeShared<FTopologicalVertex>(Archive);
+		case EEntity::VertexLink:
+			return FEntity::MakeShared<FVertexLink>(Archive);
+		default:
+			return TSharedPtr<FEntity>();
+		}
+	}
+
+#ifdef CADKERNEL_DEV
+	void FEntity::InfoEntity() const
+	{
+		FInfoEntity Info;
+		GetInfo(Info);
+		Info.Display();
+	}
+
+	FInfoEntity& FEntity::GetInfo(FInfoEntity& Info) const
+	{
+		Info.Set(*this);
+		Info.Add(TEXT("Id"), GetId())
+			.Add(TEXT("Type"), TypesNames[(uint32)GetEntityType()]);
+		return Info;
+	}
+
+#ifdef CADKERNEL_DEV
+	FInfoEntity& FEntityGeom::GetInfo(FInfoEntity& Info) const
+	{
+		return FEntity::GetInfo(Info)
+			.Add(TEXT("Kio"), CtKioId);
+	}
+#endif
+
+#endif
+
+	void FEntity::SerializeIdents(FCADKernelArchive& Ar, TArray<TOrientedEntity<FEntity>>& Array)
+	{
+		if (Ar.IsLoading())
+		{
+			int32 ArrayNum = 0;
+			Ar << ArrayNum;
+			Array.Init(TOrientedEntity<FEntity>(), ArrayNum);
+			for (TOrientedEntity<FEntity>& OrientedEntity : Array)
+			{
+				FIdent OldId = 0;
+				Ar << OldId;
+				Ar.SetReferencedEntityOrAddToWaitingList(OldId, OrientedEntity.Entity);
+				Ar.Serialize(&OrientedEntity.Direction, sizeof(EOrientation));
+			}
+		}
+		else
+		{
+			int32 ArrayNum = Array.Num();
+			Ar << ArrayNum;
+			for (TOrientedEntity<FEntity>& OrientedEntity : Array)
+			{
+				FIdent Id = OrientedEntity.Entity->GetId();
+				Ar << Id;
+				Ar.Serialize(&OrientedEntity.Direction, sizeof(EOrientation));
+				Ar.AddEntityToSave(Id);
+			}
+		}
+	}
+
+	void FEntity::SerializeIdents(FCADKernelArchive& Ar, TArray<TWeakPtr<FEntity>>& EntityArray, bool bSaveSelection)
+	{
+		if (Ar.IsLoading())
+		{
+			int32 ArrayNum = 0;
+			Ar << ArrayNum;
+			EntityArray.Init(TSharedPtr<FEntity>(), ArrayNum);
+			for (TWeakPtr<FEntity>& Entity : EntityArray)
+			{
+				FIdent OldId = 0;
+				Ar << OldId;
+				Ar.SetReferencedEntityOrAddToWaitingList(OldId, Entity);
+			}
+		}
+		else
+		{
+			int32 ArrayNum = EntityArray.Num();
+			Ar << ArrayNum;
+			for (TWeakPtr<FEntity>& Entity : EntityArray)
+			{
+				if (Entity.IsValid())
+				{
+					FIdent Id = Entity.Pin()->GetId();
+					Ar << Id;
+				}
+			}
+
+			if (bSaveSelection)
+			{
+				for (TWeakPtr<FEntity>& Entity : EntityArray)
+				{
+					if (Entity.IsValid())
+					{
+						Ar.AddEntityToSave(Entity.Pin()->GetId());
+					}
+				}
+			}
+		}
+	}
+
+	void FEntity::SerializeIdents(FCADKernelArchive& Ar, TArray<TSharedPtr<FEntity>>& EntityArray, bool bSaveSelection)
+	{
+		if (Ar.IsLoading())
+		{
+			int32 ArrayNum = 0;
+			Ar << ArrayNum;
+			EntityArray.Init(TSharedPtr<FEntity>(), ArrayNum);
+			for (TSharedPtr<FEntity>& Entity : EntityArray)
+			{
+				FIdent OldId = 0;
+				Ar << OldId;
+				Ar.SetReferencedEntityOrAddToWaitingList(OldId, Entity);
+			}
+		}
+		else
+		{
+			int32 ArrayNum = EntityArray.Num();
+			Ar << ArrayNum;
+			for (TSharedPtr<FEntity>& Entity : EntityArray)
+			{
+				if (Entity.IsValid())
+				{
+					FIdent Id = Entity->GetId();
+					Ar << Id;
+				}
+			}
+
+			if (bSaveSelection)
+			{
+				for (TSharedPtr<FEntity>& Entity : EntityArray)
+				{
+					if(Entity.IsValid())
+					{ 
+						Ar.AddEntityToSave(Entity->GetId());
+					}
+				}
+			}
+		}
+	}
+
+	void FEntity::SerializeIdent(FCADKernelArchive& Ar, TSharedPtr<FEntity>& Entity, bool bSaveSelection)
+	{
+		if (Ar.IsLoading())
+		{
+			FIdent OldId = 0;
+			Ar << OldId;
+			Ar.SetReferencedEntityOrAddToWaitingList(OldId, Entity);
+		}
+		else
+		{
+			FIdent Id = Entity.IsValid() ? Entity->GetId() : 0;
+			Ar << Id;
+			if (bSaveSelection && Id)
+			{
+				Ar.AddEntityToSave(Id);
+			}
+		}
+	}
+
+	void FEntity::SerializeIdent(FCADKernelArchive& Ar, TWeakPtr<FEntity>& Entity, bool bSaveSelection)
+	{
+		if (Ar.Archive.IsLoading())
+		{
+			FIdent OldId;
+			Ar.Archive << OldId;
+			Ar.SetReferencedEntityOrAddToWaitingList(OldId, Entity);
+		}
+		else
+		{
+			FIdent Id = Entity.IsValid() ? Entity.Pin()->GetId() : 0;
+			Ar.Archive << Id;
+			if (bSaveSelection && Id)
+			{
+				Ar.AddEntityToSave(Id);
+			}
+		}
+	}
+
+	bool FEntity::SetId(FDatabase& Database)
+	{
+		if (Id < 1)
+		{
+			Database.AddEntity(AsShared());
+			++Database.EntityCount;
+			return true;
+		}
+		return Database.bForceSpawning || false;
+	}
+
+} // namespace CADKernel
diff --git a/Engine/Source/Runtime/Datasmith/CADKernel/Private/CADKernel/Core/KernelParameters.h b/Engine/Source/Runtime/Datasmith/CADKernel/Private/CADKernel/Core/KernelParameters.h
new file mode 100644
--- /dev/null
+++ b/Engine/Source/Runtime/Datasmith/CADKernel/Private/CADKernel/Core/KernelParameters.h
@@ -0,0 +1,20 @@
+// Copyright Epic Games, Inc. All Rights Reserved.
+#pragma once
+
+#include "CADKernel/Core/Parameters.h"
+
+
+namespace CADKernel
+{
+	class CADKERNEL_API FKernelParameters : public FParameters
+	{
+	public:
+		FParameter GeometricalTolerance;
+
+		FKernelParameters()
+			: FParameters(3)
+			, GeometricalTolerance(TEXT("GeometricalTolerance"), 0.02, *this)
+		{}
+	};
+} // namespace CADKernel
+
diff --git a/Engine/Source/Runtime/Datasmith/CADKernel/Private/CADKernel/Core/Parameters.cpp b/Engine/Source/Runtime/Datasmith/CADKernel/Private/CADKernel/Core/Parameters.cpp
new file mode 100644
--- /dev/null
+++ b/Engine/Source/Runtime/Datasmith/CADKernel/Private/CADKernel/Core/Parameters.cpp
@@ -0,0 +1,199 @@
+// Copyright Epic Games, Inc. All Rights Reserved.
+
+#include "CADKernel/Core/Parameters.h"
+
+#include "CADKernel/UI/Message.h"
+
+using namespace CADKernel;
+using namespace std;
+
+FString FParameterValue::ToString() const
+{
+	switch (Type)
+	{
+	case EValue::Double:
+		return FString::Printf(TEXT("%f"), Value.Real);
+		break;
+
+	case EValue::Boolean:
+		return Value.Boolean ? TEXT("true") : TEXT("false");
+		break;
+
+	case EValue::Integer:
+		return FString::Printf(TEXT("%d"), Value.Integer);
+		break;
+
+	case EValue::String:
+		return StringValue;
+		break;
+	}
+	return TEXT("");
+}
+
+FString FParameter::GetTypeString() const
+{
+	switch (GetType())
+	{
+	case EValue::Double:
+		return TEXT("Double");
+
+	case EValue::Boolean:
+		return TEXT("Boolean");
+
+	case EValue::Integer:
+		return TEXT("Integer");
+
+	case EValue::String:
+		return TEXT("String");
+
+	default:
+		break;
+	}
+	return TEXT("");
+}
+
+bool FParameterValue::operator==(const FParameterValue& Other) const
+{
+	if (Type != Other.Type) 
+	{
+		return false;
+	}
+
+	switch (Type)
+	{
+	case EValue::Double:
+		return FMath::IsNearlyEqual(Value.Real, Other.Value.Real);
+
+	case EValue::Boolean:
+		return Value.Boolean == Other.Value.Boolean;
+
+	case EValue::Integer:
+		return Value.Integer == Other.Value.Integer;
+
+	case EValue::String:
+		return StringValue == Other.StringValue;
+
+	default:
+		break;
+	}
+	return false;
+}
+
+void FParameter::AddToParameterMap(FParameter& Parameter, FParameters& Parameters)
+{
+	Parameters.Add(Parameter);
+}
+
+void FParameter::SetFromString(const FString& String)
+{
+	switch (GetType())
+	{
+	case EValue::Double:
+		{
+			double TmpValue = _wtof(*String);
+			SetValue(TmpValue);
+		}
+		break;
+
+	case EValue::Integer:
+		{
+			int32 TmpValue = _wtoi(*String);
+			SetValue(TmpValue);
+		}
+		break;
+
+	case EValue::Boolean:
+		{
+			FString lowerStr = String.ToLower();
+
+			if (lowerStr == TEXT("true"))
+			{
+				SetValue(true);
+			}
+			else if (lowerStr == TEXT("false"))
+			{
+				SetValue(false);
+			}
+			else if (lowerStr == TEXT("0"))
+			{
+				SetValue(false);
+			}
+			else {
+				SetValue(true);
+			}
+		}
+		break;
+
+	case EValue::String:
+		{
+			SetValue(String);
+		}
+		break;
+
+	default:
+		break;
+	}
+}
+
+void FParameters::SetFromString(const FString& ParameterStr)
+{
+	if (ParameterStr.Len() > 0 && ParameterStr[0] == TEXT('?')) 
+	{
+		PrintParameterList();
+	}
+
+	TArray<FString> ParameterArray;
+	ParameterStr.ParseIntoArray(ParameterArray, TEXT(";"));
+
+	for (FString& ParameterString : ParameterArray)
+	{
+		FString Name;
+		FString Value;
+		if (ParameterString.Split(TEXT("="), &Name, &Value))
+		{
+			FParameter* Parameter = GetByName(Name);
+			if (Parameter == nullptr)
+			{
+				continue;
+			}
+			Parameter->SetFromString(Value);
+		}
+	}
+}
+
+void FParameters::PrintParameterList()
+{
+	FMessage::Printf(Log, TEXT("Extra parameters:\n"));
+	for (const TPair<FString, FParameter*>& Parameter : Map)
+	{
+		FMessage::Printf(Log, TEXT(" - %s : %s (default = %s)\n"), *Parameter.Key, *Parameter.Value->GetValue().ToString(), *Parameter.Value->GetDefaultValue().ToString());
+	}
+}
+
+FString FParameters::ToString(bool bOnlyChanged) const
+{
+	FString String;
+	bool bIsFirst = true;
+	for (const TPair<FString, FParameter*>& Pair : Map)
+	{
+		const FParameter* Parameter = Pair.Value;
+		if (bOnlyChanged && Parameter->IsDefault()) 
+		{
+			continue;
+		}
+
+		if (!bIsFirst) 
+		{
+			String += TEXT(";");
+		}
+
+		String += Parameter->ToString();
+		bIsFirst = false;
+	}
+	return String;
+}
+
+FString FParameter::ToString() const
+{
+	return GetName() + FString(TEXT("=")) + GetValue().ToString();
+}
diff --git a/Engine/Source/Runtime/Datasmith/CADKernel/Private/CADKernel/Core/Session.cpp b/Engine/Source/Runtime/Datasmith/CADKernel/Private/CADKernel/Core/Session.cpp
new file mode 100644
--- /dev/null
+++ b/Engine/Source/Runtime/Datasmith/CADKernel/Private/CADKernel/Core/Session.cpp
@@ -0,0 +1,54 @@
+// Copyright Epic Games, Inc. All Rights Reserved.
+#include "CADKernel/Core/Session.h"
+
+#include "CADKernel/UI/Message.h"
+
+using namespace CADKernel;
+
+
+void FSession::SaveDatabase(const TCHAR* FileName)
+{
+	TSharedPtr<FCADKernelArchive> Archive = FCADKernelArchive::CreateArchiveWriter(*this, FileName);
+	if (!Archive.IsValid())
+	{
+		FMessage::Printf(Log, TEXT("The archive file %s is corrupted\n"), FileName);
+	}
+
+	Database.Serialize(*Archive.Get());
+	Archive->Close();
+}
+
+TSharedRef<FModel> FSession::GetModel()
+{
+	return Database.GetModel();
+}
+
+
+void FSession::SaveDatabase(const TCHAR* FileName, const TArray<TSharedPtr<FEntity>>& SelectedEntities)
+{
+	TArray<FIdent> EntityIds;
+	EntityIds.Reserve(SelectedEntities.Num());
+
+	SpawnEntityIdent(SelectedEntities, true);
+
+	for (const TSharedPtr<FEntity>& Entity : SelectedEntities)
+	{
+		EntityIds.Add(Entity->GetId());
+	}
+
+	TSharedPtr<FCADKernelArchive> Archive = FCADKernelArchive::CreateArchiveWriter(*this, FileName);
+
+	Database.SerializeSelection(*Archive.Get(), EntityIds);
+	Archive->Close();
+}
+
+void FSession::LoadDatabase(const TCHAR* FilePath)
+{
+	TSharedPtr<FCADKernelArchive> Archive = FCADKernelArchive::CreateArchiveReader(*this, FilePath);
+	if (!Archive.IsValid())
+	{
+		FMessage::Printf(Log, TEXT("The archive file %s is corrupted\n"), FilePath);
+	}
+
+	Database.Deserialize(*Archive.Get());
+}
diff --git a/Engine/Source/Runtime/Datasmith/CADKernel/Private/CADKernel/Core/System.cpp b/Engine/Source/Runtime/Datasmith/CADKernel/Private/CADKernel/Core/System.cpp
new file mode 100644
--- /dev/null
+++ b/Engine/Source/Runtime/Datasmith/CADKernel/Private/CADKernel/Core/System.cpp
@@ -0,0 +1,149 @@
+// Copyright Epic Games, Inc. All Rights Reserved.
+#include "CADKernel/Core/System.h"
+
+#include "CADKernel/Core/Database.h"
+#include "CADKernel/Core/Entity.h"
+#include "CADKernel/Core/KernelParameters.h"
+#include "CADKernel/Core/Version.h"
+#include "CADKernel/UI/Visu.h"
+#include "CADKernel/Utils/Util.h"
+
+#include "HAL/FileManager.h"
+
+#ifdef CADKERNEL_DEV
+#include <stdlib.h>
+#include <signal.h>
+#endif 
+
+#pragma warning(disable:4996) // unsafe sprintf
+
+using namespace CADKernel;
+
+TUniquePtr<FSystem> FSystem::Instance = nullptr;
+
+FSystem::FSystem()
+	: Parameters(MakeShared<FKernelParameters>())
+	, DefaultVisu()
+	, Viewer(&DefaultVisu)
+	, Console(&DefaultConsole)
+	, ProgressManager(&DefaultProgressManager)
+{
+	LogFile = nullptr;
+	LogLevel = Log;
+	SpyFile = nullptr;
+
+	VerboseLevel = Log;
+}
+
+void FSystem::Initialize(bool bIsDll, const FString& LogFilePath, const FString& SpyFilePath)
+{
+	SetVerboseLevel(Log);
+
+	TCHAR* DebugFolderPath = _wgetenv(TEXT("DS_DBG"));
+	if (DebugFolderPath) 
+	{
+		DefineLogFile(FString(DebugFolderPath) + TEXT("\\") + ProductName + TEXT(".log"));
+		DefineSpyFile(FString(DebugFolderPath) + TEXT("\\") + ProductName + TEXT(".py"));
+	}
+	else 
+	{
+		if (LogFilePath.Len() > 0)
+		{
+			DefineLogFile(LogFilePath);
+		}
+		if (SpyFilePath.Len() > 0)
+		{
+			DefineSpyFile(SpyFilePath);
+		}
+	}
+
+	PrintHeader();
+
+	fflush(stdout);
+
+	if(bIsDll) 
+	{
+		SetVerboseLevel(EVerboseLevel::NoVerbose);
+	} 
+	else 
+	{
+		SetVerboseLevel(EVerboseLevel::Log);
+	}
+}
+
+void FSystem::CloseLogFiles()
+{
+	if (LogFile)
+	{
+		LogFile->Close();
+		LogFile.Reset();
+	}
+	if (SpyFile)
+	{
+		SpyFile->Close();
+		SpyFile.Reset();
+	}
+}
+
+void FSystem::Shutdown()
+{
+	CloseLogFiles();
+	Instance.Reset();
+}
+
+void FSystem::DefineLogFile(const FString& InLogFilePath, EVerboseLevel InLevel)
+{
+	if(LogFile) 
+	{
+		LogFile->Close();
+		LogFile.Reset();
+	}
+
+	LogFile = MakeShareable<FArchive>(IFileManager::Get().CreateFileWriter(*InLogFilePath, IO_WRITE));
+	LogLevel = InLevel;
+}
+
+void FSystem::DefineSpyFile(const FString& InSpyFilePath)
+{
+	if(SpyFile) 
+	{
+		SpyFile->Close();
+		SpyFile.Reset();
+	}
+	SpyFile = MakeShareable<FArchive>(IFileManager::Get().CreateFileWriter(*InSpyFilePath, IO_WRITE));
+}
+
+void FSystem::InitializeCADKernel()
+{
+	FSystem::Get().Initialize();
+	FSystem::Get().SetVerboseLevel(EVerboseLevel::Log);
+}
+
+FString FSystem::GetToolkitVersion() const
+{
+	return TOOLKIT_VERSION_ASCII;
+}
+
+FString FSystem::GetCompilationDate() const
+{
+	return UTF8_TO_TCHAR(__DATE__);
+}
+
+void FSystem::PrintHeader()
+{
+	FMessage::Printf(Log, TEXT("_______________________________________________________________________________\n"));
+	FMessage::Printf(Log, TEXT("\n"));
+	FMessage::Printf(Log, TEXT("\tDatasmith CAD Kernel Toolkit release %s (%s)\n"),* GetToolkitVersion(),* GetCompilationDate());
+	FMessage::Printf(Log, TEXT("\t" EPIC_COPYRIGHT "\n"));
+	FMessage::Printf(Log, TEXT("\n"));
+	FMessage::Printf(Log, TEXT("_______________________________________________________________________________\n"));
+	FMessage::Printf(Log, TEXT("\n"));
+}
+
+extern "C"
+{
+#ifdef COMPILER_MSC
+#pragma warning( disable : 4297)
+#endif
+}
+
diff --git a/Engine/Source/Runtime/Datasmith/CADKernel/Private/CADKernel/Core/Version.h b/Engine/Source/Runtime/Datasmith/CADKernel/Private/CADKernel/Core/Version.h
new file mode 100644
--- /dev/null
+++ b/Engine/Source/Runtime/Datasmith/CADKernel/Private/CADKernel/Core/Version.h
@@ -0,0 +1,20 @@
+// Copyright Epic Games, Inc. All Rights Reserved.
+#pragma once
+
+//** ** Version du Toolkit** **//
+#define TOOLKIT_VERSION 7, 0, 4
+#define TOOLKIT_VERSION_ASCII L"7.0.4"
+
+//** ** Numero du TAG** **//
+#define TAG_VERSION 2000
+#define TAG_VERSION_ASCII L".2000" 
+
+#define RELEASE_DATE	 __DATE__
+
+//** ** Version de CADKernel** **//
+#define DATASMITH_CAD_KERNEL_VERSION TOOLKIT_VERSION, TAG_VERSION
+#define DATASMITH_CAD_KERNEL_VERSION_ASCII TOOLKIT_VERSION_ASCII TAG_VERSION_ASCII
+
+//** ** CADKernel Copyright** **//
+#define EPIC_COPYRIGHT "Copyright Epic Games, Inc. All Rights Reserved."
+
diff --git a/Engine/Source/Runtime/Datasmith/CADKernel/Private/CADKernel/Geo/Curves/BezierCurve.cpp b/Engine/Source/Runtime/Datasmith/CADKernel/Private/CADKernel/Geo/Curves/BezierCurve.cpp
new file mode 100644
--- /dev/null
+++ b/Engine/Source/Runtime/Datasmith/CADKernel/Private/CADKernel/Geo/Curves/BezierCurve.cpp
@@ -0,0 +1,69 @@
+// Copyright Epic Games, Inc. All Rights Reserved.
+#include "CADKernel/Geo/Curves/BezierCurve.h"
+
+#include "CADKernel/Geo/GeoPoint.h"
+#include "CADKernel/Geo/Sampling/PolylineTools.h"
+#include "CADKernel/Math/BSpline.h"
+
+
+using namespace CADKernel;
+
+void FBezierCurve::EvaluatePoint(double Coordinate, FCurvePoint& OutPoint, int32 DerivativeOrder) const
+{
+	OutPoint.DerivativeOrder = DerivativeOrder;
+	OutPoint.Init();
+
+	TArray<double> Bernstein;
+	TArray<double> BernsteinD1;
+	TArray<double> BernsteinD2;
+
+	BSpline::Bernstein(Poles.Num()-1, Coordinate, Bernstein, BernsteinD1, BernsteinD2);
+
+	for (int32 Index = 0; Index < Poles.Num(); Index++)
+	{
+		OutPoint.Point += Poles[Index] * Bernstein[Index];
+	}
+
+	if (DerivativeOrder > 0)
+	{
+		for (int32 Index = 0; Index < Poles.Num(); Index++)
+		{
+			OutPoint.Gradient += Poles[Index] * BernsteinD1[Index];
+		}
+	}
+
+	if (DerivativeOrder > 1)
+	{
+		for (int32 Index = 0; Index < Poles.Num(); Index++)
+		{
+			OutPoint.Laplacian += Poles[Index] * BernsteinD2[Index];
+		}
+	}
+}
+
+TSharedPtr<FEntityGeom> FBezierCurve::ApplyMatrix(const FMatrixH& InMatrix) const
+{
+	TArray<FPoint> TransformedPoles;
+	TransformedPoles.Reserve(Poles.Num());
+
+	for (const FPoint& Pole : Poles)
+	{
+		TransformedPoles.Emplace(InMatrix.Multiply(Pole));
+	}
+
+	return FEntity::MakeShared<FBezierCurve>(Tolerance, TransformedPoles);
+}
+
+#ifdef CADKERNEL_DEV
+FInfoEntity& FBezierCurve::GetInfo(FInfoEntity& Info) const
+{
+	return FCurve::GetInfo(Info)
+		.Add(TEXT("degre"), Poles.Num() - 1)
+		.Add(TEXT("poles"), Poles);
+}
+#endif
+
+void FBezierCurve::ExtendTo(const FPoint& Point)
+{
+	PolylineTools::ExtendTo(Poles, Point);
+}
diff --git a/Engine/Source/Runtime/Datasmith/CADKernel/Private/CADKernel/Geo/Curves/Curve.cpp b/Engine/Source/Runtime/Datasmith/CADKernel/Private/CADKernel/Geo/Curves/Curve.cpp
new file mode 100644
--- /dev/null
+++ b/Engine/Source/Runtime/Datasmith/CADKernel/Private/CADKernel/Geo/Curves/Curve.cpp
@@ -0,0 +1,184 @@
+// Copyright Epic Games, Inc. All Rights Reserved.
+#include "CADKernel/Geo/Curves/Curve.h"
+
+#include "CADKernel/Core/KernelParameters.h"
+#include "CADKernel/Core/System.h"
+#include "CADKernel/Geo/Curves/BoundedCurve.h"
+#include "CADKernel/Geo/Sampling/Polyline.h"
+#include "CADKernel/Geo/Sampling/PolylineTools.h"
+#include "CADKernel/Geo/Sampling/SurfacicPolyline.h"
+#include "CADKernel/Geo/Sampler/SamplerOnChord.h"
+#include "CADKernel/Geo/Sampler/SamplerOnParam.h"
+#include "CADKernel/UI/Display.h"
+
+namespace CADKernel
+{
+	double FCurve::GetLength() const
+	{
+		if (!GlobalLength.IsValid())
+		{
+			GlobalLength = ComputeLength(Boundary);
+		}
+		return GlobalLength;
+	}
+
+	double FCurve::GetParametricTolerance() const
+	{
+		if (!Tolerance1D.IsValid())
+		{
+			double Tolerance2DMax = (GetUMax() - GetUMin()) * 0.01;
+			double Length = GetLength();
+			if (Length < SMALL_NUMBER)
+			{
+				Tolerance1D = Tolerance2DMax;
+				return Tolerance1D;
+			}
+
+			Tolerance1D = Tolerance * (GetUMax() - GetUMin()) / Length;
+			Tolerance1D = FMath::Max(Tolerance2DMax, (double)Tolerance1D);
+		}
+		return Tolerance1D;
+	}
+
+	void FCurve::EvaluatePoints(const TArray<double>& Coordinates, TArray<FCurvePoint>& OutPoints, int32 DerivativeOrder) const
+	{
+		OutPoints.SetNum(Coordinates.Num());
+		for (int32 iPoint = 0; iPoint < Coordinates.Num(); iPoint++)
+		{
+			EvaluatePoint(Coordinates[iPoint], OutPoints[iPoint], DerivativeOrder);
+		}
+	}
+
+	void FCurve::EvaluatePoints(const TArray<double>& Coordinates, TArray<FPoint>& OutPoints) const
+	{
+		OutPoints.Empty(Coordinates.Num());
+		for (double Coordinate : Coordinates)
+		{
+			FCurvePoint Point;
+			EvaluatePoint(Coordinate, Point, 0);
+			OutPoints.Emplace(Point.Point);
+		}
+	}
+
+	void FCurve::Evaluate2DPoints(const TArray<double>& Coordinates, TArray<FCurvePoint2D>& OutPoints, int32 DerivativeOrder) const
+	{
+		OutPoints.SetNum(Coordinates.Num());
+		for (int32 iPoint = 0; iPoint < Coordinates.Num(); iPoint++)
+		{
+			Evaluate2DPoint(Coordinates[iPoint], OutPoints[iPoint], DerivativeOrder);
+		}
+	}
+
+	/**
+	 * Evaluate exact 2D points of the curve at the input Coordinates
+	 * The function can only be used with 2D curve (Dimension == 2)
+	 */
+	void FCurve::Evaluate2DPoints(const TArray<double>& Coordinates, TArray<FPoint2D>& OutPoints) const
+	{
+		OutPoints.SetNum(Coordinates.Num());
+		for (int32 iPoint = 0; iPoint < Coordinates.Num(); iPoint++)
+		{
+			Evaluate2DPoint(Coordinates[iPoint], OutPoints[iPoint]);
+		}
+	}
+
+#ifdef CADKERNEL_DEV
+	FInfoEntity& FCurve::GetInfo(FInfoEntity& Info) const
+	{
+		return FEntity::GetInfo(Info)
+			.Add(TEXT("Curve type"), CurvesTypesNames[(uint8)GetCurveType()])
+			.Add(TEXT("Dimension"), (int32) Dimension)
+			.Add(TEXT("Boundary"), Boundary)
+			.Add(TEXT("Length"), GetLength());
+	}
+#endif
+
+	void FCurve::FindNotDerivableCoordinates(int32 DerivativeOrder, TArray<double>& OutNotDerivableCoordinates) const
+	{
+		FindNotDerivableCoordinates(Boundary, DerivativeOrder, OutNotDerivableCoordinates);
+	}
+
+	void FCurve::FindNotDerivableCoordinates(const FLinearBoundary& InBoundary, int32 DerivativeOrder, TArray<double>& OutNotDerivableCoordinates) const
+	{
+	}
+
+	TSharedPtr<FCurve> FCurve::ReboundCurve(const FLinearBoundary& InBoundary)
+	{
+		FLinearBoundary NewBoundary = InBoundary;
+		if (NewBoundary.Min < GetUMin())
+		{
+			NewBoundary.Min = GetUMin();
+		}
+
+		if (NewBoundary.Max > GetUMax())
+		{
+			NewBoundary.Max = GetUMax();
+		}
+
+		if (NewBoundary.IsDegenerated())
+		{
+			FMessage::Printf(Log, TEXT("Invalid bounds (u1=%f u2=%f) on curve %d\n"), NewBoundary.Min, NewBoundary.Max, GetId());
+			return TSharedPtr<FCurve>();
+		}
+
+		if (FMath::IsNearlyEqual(NewBoundary.Min, GetUMin()) && FMath::IsNearlyEqual(NewBoundary.Max, GetUMax()))
+		{
+			FMessage::Printf(Debug, TEXT("Invalid rebound (UMin and UMax are nearly equal) on curve %d\n"), GetId());
+			return TSharedPtr<FCurve>();
+		}
+
+		return FEntity::MakeShared<FBoundedCurve>(Tolerance, StaticCastSharedRef<FCurve>(AsShared()), NewBoundary, Dimension);
+	}
+
+	//double FCurve::ComputeLength(const FLinearBoundary& InBoundary) const
+	//{
+	//	if (GetDimension() == 3)
+	//	{
+	//		return ComputeLength3D(InBoundary);
+	//	}
+	//	else
+	//	{
+	//		return ComputeLength2D(InBoundary);
+	//	}
+	//}
+
+	double FCurve::ComputeLength(const FLinearBoundary& InBoundary) const
+	{
+		FPolyline3D Polyline;
+		FCurveSamplerOnChord Sampler(*this, Boundary, Tolerance, Polyline);
+		Sampler.Sample();
+		return Polyline.GetLength(Boundary);
+	}
+
+	//double FCurve::ComputeLength2D(const FLinearBoundary& InBoundary) const
+	//{
+	//	double Gap = Boundary.Length() / 30;
+	//	int32 PointCount = (int32) (InBoundary.Length() / Gap + 1);
+	//	Gap = InBoundary.Length() / PointCount;
+	//	TArray<double> Coordinates;
+	//	PointCount++;
+	//	Coordinates.Reserve(PointCount);
+	//	double Coordinate = InBoundary.GetMin();
+	//	for (int32 Index = 0; Index < PointCount; ++Index)
+	//	{
+	//		Coordinates.Add(Coordinate);
+	//		Coordinate += Gap;
+	//	}
+
+	//	TArray<FPoint2D> Polyline;
+	//	Evaluate2DPoints(Coordinates, Polyline);
+
+	//	TPolylineApproximator<FPoint2D> Approximator(Coordinates, Polyline);
+	//	return Approximator.ComputeLength();
+	//}
+
+	void FCurve::Presample(const FLinearBoundary& InBoundary, TArray<double>& OutSampling) const
+	{
+		FPolyline3D Presampling;
+		FCurveSamplerOnParam Sampler(*this, Boundary, Tolerance * 10., Tolerance, Presampling);
+		Sampler.Sample();
+
+		Presampling.SwapCoordinates(OutSampling);
+	}
+
+} // namespace CADKernel
diff --git a/Engine/Source/Runtime/Datasmith/CADKernel/Private/CADKernel/Geo/Curves/EllipseCurve.cpp b/Engine/Source/Runtime/Datasmith/CADKernel/Private/CADKernel/Geo/Curves/EllipseCurve.cpp
new file mode 100644
--- /dev/null
+++ b/Engine/Source/Runtime/Datasmith/CADKernel/Private/CADKernel/Geo/Curves/EllipseCurve.cpp
@@ -0,0 +1,22 @@
+// Copyright Epic Games, Inc. All Rights Reserved.
+#include "CADKernel/Geo/Curves/EllipseCurve.h"
+
+
+using namespace CADKernel;
+
+
+TSharedPtr<FEntityGeom> FEllipseCurve::ApplyMatrix(const FMatrixH& InMatrix) const
+{
+	FMatrixH NewMatrix = InMatrix * Matrix;
+	return FEntity::MakeShared<FEllipseCurve>(Tolerance, NewMatrix, RadiusU, RadiusV, Boundary);
+}
+
+#ifdef CADKERNEL_DEV
+FInfoEntity& FEllipseCurve::GetInfo(FInfoEntity& Info) const
+{
+	return FCurve::GetInfo(Info).Add(TEXT("Matrix"), Matrix)
+		.Add(TEXT("radius U"), RadiusU)
+		.Add(TEXT("radius V"), RadiusV);
+}
+#endif
+
diff --git a/Engine/Source/Runtime/Datasmith/CADKernel/Private/CADKernel/Geo/Curves/HyperbolaCurve.cpp b/Engine/Source/Runtime/Datasmith/CADKernel/Private/CADKernel/Geo/Curves/HyperbolaCurve.cpp
new file mode 100644
--- /dev/null
+++ b/Engine/Source/Runtime/Datasmith/CADKernel/Private/CADKernel/Geo/Curves/HyperbolaCurve.cpp
@@ -0,0 +1,21 @@
+// Copyright Epic Games, Inc. All Rights Reserved.
+#include "CADKernel/Geo/Curves/HyperbolaCurve.h"
+
+
+using namespace CADKernel;
+
+TSharedPtr<FEntityGeom> FHyperbolaCurve::ApplyMatrix(const FMatrixH& InMatrix) const
+{
+	FMatrixH NewMatrix = InMatrix * Matrix;
+	return FEntity::MakeShared<FHyperbolaCurve>(Tolerance, NewMatrix, SemiMajorAxis, SemiImaginaryAxis, Boundary);
+}
+
+#ifdef CADKERNEL_DEV
+FInfoEntity& FHyperbolaCurve::GetInfo(FInfoEntity& Info) const
+{
+	return FCurve::GetInfo(Info).Add(TEXT("Matrix"), Matrix)
+		.Add(TEXT("semi axis"), SemiMajorAxis)
+		.Add(TEXT("semi imag axis"), SemiImaginaryAxis);
+}
+#endif
+
diff --git a/Engine/Source/Runtime/Datasmith/CADKernel/Private/CADKernel/Geo/Curves/NURBSCurve.cpp b/Engine/Source/Runtime/Datasmith/CADKernel/Private/CADKernel/Geo/Curves/NURBSCurve.cpp
new file mode 100644
--- /dev/null
+++ b/Engine/Source/Runtime/Datasmith/CADKernel/Private/CADKernel/Geo/Curves/NURBSCurve.cpp
@@ -0,0 +1,139 @@
+// Copyright Epic Games, Inc. All Rights Reserved.
+
+#include "CADKernel/Geo/Curves/NURBSCurve.h"
+
+#include "CADKernel/Geo/GeoPoint.h"
+#include "CADKernel/Geo/Sampling/PolylineTools.h"
+#include "CADKernel/Math/BSpline.h"
+#include "CADKernel/Utils/Util.h"
+
+#include "Algo/ForEach.h"
+#include "Algo/Reverse.h"
+
+using namespace CADKernel;
+
+FNURBSCurve::FNURBSCurve(const double InTolerance, int32 InDegre, const TArray<double>& InNodalVector, const TArray<FPoint>& InPoles, int8 InDimension)
+	: FCurve(InTolerance, InDimension)
+	, Degree(InDegre)
+	, NodalVector(InNodalVector)
+	, Poles(InPoles)
+	, bIsRational(false)
+{
+	Weights.Init(1.0, InPoles.Num());
+	Finalize();
+}
+
+FNURBSCurve::FNURBSCurve(const double InTolerance, int32 InDegre, const TArray<double>& InNodalVector, const TArray<FPoint>& InPoles, const TArray<double>& InWeights, int8 InDimension)
+	: FCurve(InTolerance, InDimension)
+	, Degree(InDegre)
+	, NodalVector(InNodalVector)
+	, Weights(InWeights)
+	, Poles(InPoles)
+	, bIsRational(true)
+{
+	Finalize();
+}
+
+TSharedPtr<FEntityGeom> FNURBSCurve::ApplyMatrix(const FMatrixH& InMatrix) const
+{
+	TArray<FPoint> TransformedPoles;
+
+	TransformedPoles.Reserve(Poles.Num());
+	for (FPoint Pole : Poles) 
+	{
+		TransformedPoles.Emplace(InMatrix.Multiply(Pole));
+	}
+
+	return FEntity::MakeShared<FNURBSCurve>(Tolerance, Degree, NodalVector, TransformedPoles, Weights, Dimension);
+}
+
+#ifdef CADKERNEL_DEV
+FInfoEntity& FNURBSCurve::GetInfo(FInfoEntity& Info) const
+{
+	return FCurve::GetInfo(Info).Add(TEXT("Degre"), Degree)
+		.Add(TEXT("Nodal vector"), NodalVector)
+		.Add(TEXT("Poles"), Poles)
+		.Add(TEXT("Weights"), Weights);
+}
+#endif
+
+void FNURBSCurve::Finalize()
+{
+	// Check if the curve is really rational, otherwise remove the weights
+	if (bIsRational)
+	{
+		const double FirstWeigth = Weights[0];
+
+		bool bIsReallyRational = false;
+		for (const double& Weight : Weights)
+		{
+			if (!FMath::IsNearlyEqual(Weight, FirstWeigth))
+			{
+				bIsReallyRational = true;
+				break;
+			}
+		}
+
+		if (!bIsReallyRational)
+		{
+			if (!FMath::IsNearlyEqual(1., FirstWeigth))
+			{
+				for (FPoint& Pole : Poles)
+				{
+					Pole /= FirstWeigth;
+				}
+			}
+			bIsRational = false;
+		}
+	}
+
+	PoleDimension = Dimension + (bIsRational ? 1 : 0);
+	HomogeneousPoles.SetNum(Poles.Num() * PoleDimension);
+
+	if (bIsRational)
+	{
+		if (Dimension == 2)
+		{
+			for (int32 Index = 0, Jndex = 0; Index < Poles.Num(); Index++)
+			{
+				HomogeneousPoles[Jndex++] = Poles[Index].X * Weights[Index];
+				HomogeneousPoles[Jndex++] = Poles[Index].Y * Weights[Index];
+				HomogeneousPoles[Jndex++] = Weights[Index];
+			}
+		}
+		else
+		{
+			for (int32 Index = 0, Jndex = 0; Index < Poles.Num(); Index++)
+			{
+				HomogeneousPoles[Jndex++] = Poles[Index].X * Weights[Index];
+				HomogeneousPoles[Jndex++] = Poles[Index].Y * Weights[Index];
+				HomogeneousPoles[Jndex++] = Poles[Index].Z * Weights[Index];
+				HomogeneousPoles[Jndex++] = Weights[Index];
+			}
+		}
+	}
+	else
+	{
+		int32 Jndex = 0;
+		if (Dimension == 2)
+		{
+			for (FPoint& Pole : Poles)
+			{
+				HomogeneousPoles[Jndex++] = Pole.X;
+				HomogeneousPoles[Jndex++] = Pole.Y;
+			}
+		}
+		else
+		{
+			memcpy(HomogeneousPoles.GetData(), Poles.GetData(), sizeof(FPoint) * Poles.Num());
+		}
+	}
+
+	Boundary.Set(NodalVector[Degree], NodalVector[NodalVector.Num() - 1 - Degree]);
+}
+
+void FNURBSCurve::ExtendTo(const FPoint& Point)
+{
+	PolylineTools::ExtendTo(Poles, Point);
+	Finalize();
+}
diff --git a/Engine/Source/Runtime/Datasmith/CADKernel/Private/CADKernel/Geo/Curves/ParabolaCurve.cpp b/Engine/Source/Runtime/Datasmith/CADKernel/Private/CADKernel/Geo/Curves/ParabolaCurve.cpp
new file mode 100644
--- /dev/null
+++ b/Engine/Source/Runtime/Datasmith/CADKernel/Private/CADKernel/Geo/Curves/ParabolaCurve.cpp
@@ -0,0 +1,24 @@
+// Copyright Epic Games, Inc. All Rights Reserved.
+#include "CADKernel/Geo/Curves/ParabolaCurve.h"
+
+#include "CADKernel/Math/Point.h"
+#include "CADKernel/Math/MatrixH.h"
+
+#include <math.h>
+
+using namespace CADKernel;
+
+TSharedPtr<FEntityGeom> FParabolaCurve::ApplyMatrix(const FMatrixH& InMatrix) const
+{
+	FMatrixH NewMatrix = InMatrix * Matrix;
+	return FEntity::MakeShared<FParabolaCurve>(Tolerance, NewMatrix, FocalDistance, Boundary, Dimension);
+}
+
+#ifdef CADKERNEL_DEV
+FInfoEntity& FParabolaCurve::GetInfo(FInfoEntity& Info) const
+{
+	return FCurve::GetInfo(Info)
+		.Add(TEXT("Matrix"), Matrix)
+		.Add(TEXT("focal dist"), FocalDistance);
+}
+#endif
diff --git a/Engine/Source/Runtime/Datasmith/CADKernel/Private/CADKernel/Geo/Curves/RestrictionCurve.cpp b/Engine/Source/Runtime/Datasmith/CADKernel/Private/CADKernel/Geo/Curves/RestrictionCurve.cpp
new file mode 100644
--- /dev/null
+++ b/Engine/Source/Runtime/Datasmith/CADKernel/Private/CADKernel/Geo/Curves/RestrictionCurve.cpp
@@ -0,0 +1,30 @@
+// Copyright Epic Games, Inc. All Rights Reserved.
+
+#include "CADKernel/Geo/Curves/RestrictionCurve.h"
+
+#include "CADKernel/math/SlopeUtils.h"
+
+using namespace CADKernel;
+
+
+#ifdef CADKERNEL_DEV
+FInfoEntity& FRestrictionCurve::GetInfo(FInfoEntity& Info) const
+{
+	return FSurfacicCurve::GetInfo(Info)
+		.Add(TEXT("2D polyline"), Polyline.Points2D)
+		.Add(TEXT("3D polyline"), Polyline.Points3D);
+}
+#endif
+
+void FRestrictionCurve::ExtendTo2DPoint(const FPoint2D& Point)
+{
+	Curve2D->ExtendTo(Point);
+	if(Polyline.bWithTangent)
+	{
+		EvaluateSurfacicPolylineWithNormalAndTangent(Polyline);
+	}
+	else
+	{
+		EvaluateSurfacicPolyline(Polyline);
+	}
+}
diff --git a/Engine/Source/Runtime/Datasmith/CADKernel/Private/CADKernel/Geo/Curves/SegmentCurve.cpp b/Engine/Source/Runtime/Datasmith/CADKernel/Private/CADKernel/Geo/Curves/SegmentCurve.cpp
new file mode 100644
--- /dev/null
+++ b/Engine/Source/Runtime/Datasmith/CADKernel/Private/CADKernel/Geo/Curves/SegmentCurve.cpp
@@ -0,0 +1,22 @@
+// Copyright Epic Games, Inc. All Rights Reserved.
+
+#include "CADKernel/Geo/Curves/SegmentCurve.h"
+
+
+using namespace CADKernel;
+
+TSharedPtr<FEntityGeom> FSegmentCurve::ApplyMatrix(const FMatrixH& InMatrix) const
+{
+	FPoint TransformedStartPoint = InMatrix.Multiply(StartPoint);
+	FPoint TransformedEndPoint = InMatrix.Multiply(EndPoint);
+
+	return FEntity::MakeShared<FSegmentCurve>(Tolerance, TransformedStartPoint, TransformedEndPoint, Dimension);
+}
+
+#ifdef CADKERNEL_DEV
+FInfoEntity& FSegmentCurve::GetInfo(FInfoEntity& Info) const
+{
+	return FCurve::GetInfo(Info).Add(TEXT("StartPoint"), StartPoint).Add(TEXT("pt2"), EndPoint);
+}
+#endif
+
diff --git a/Engine/Source/Runtime/Datasmith/CADKernel/Private/CADKernel/Geo/Sampler/SamplerOnParam.h b/Engine/Source/Runtime/Datasmith/CADKernel/Private/CADKernel/Geo/Sampler/SamplerOnParam.h
new file mode 100644
--- /dev/null
+++ b/Engine/Source/Runtime/Datasmith/CADKernel/Private/CADKernel/Geo/Sampler/SamplerOnParam.h
@@ -0,0 +1,428 @@
+// Copyright Epic Games, Inc. All Rights Reserved.
+#pragma once
+
+#include "CADKernel/Core/Types.h"
+#include "CADKernel/Geo/GeoEnum.h"
+#include "CADKernel/Geo/Surfaces/Surface.h"
+#include "CADKernel/Geo/Curves/Curve.h"
+#include "CADKernel/Geo/Curves/RestrictionCurve.h"
+#include "CADKernel/Geo/Sampler/SamplerAbstract.h"
+#include "CADKernel/Geo/Sampling/Polyline.h"   
+#include "CADKernel/Geo/Sampling/SurfacicPolyline.h"   
+#include "CADKernel/Geo/Sampling/SurfacicSampling.h"
+#include "CADKernel/Math/Boundary.h"
+#include "CADKernel/Math/Point.h"
+#include "CADKernel/Utils/ArrayUtils.h"
+
+
+//#define DEBUG_CURVE_SAMPLING
+namespace CADKernel
+{
+
+	template<typename PolylineType, typename PointType>
+	class TSamplerBasedOnParametrizationAndChordError : public TCurveSamplerAbstract<PolylineType, PointType>
+	{
+	public:
+		TSamplerBasedOnParametrizationAndChordError(const FLinearBoundary& InBoundary, double InMaxChordError, double InMaxParameterizationError, PolylineType& OutPolyline)
+			: TCurveSamplerAbstract<PolylineType, PointType>(InBoundary, OutPolyline, InMaxChordError)
+			,  DesiredParameterizationError(InMaxParameterizationError)
+		{
+		};
+
+	protected:
+
+		double  DesiredParameterizationError;
+
+		/**
+		 * This function check the chord error and the parametrization error
+		 *
+		 * Chord error:
+		 *   Check the chord error of the AB segment for each point of TmpPoints between FirstIndex and EndIndex
+		 *   If the MaxChord is biggest to the DesiredChordError, an estimation of the needed point number to insert is done.
+		 *   This estimation is bases on the approximation of local curve radius (R) :
+		 *   R = square(L) / 8Sag
+		 *   L2 = sqrt(square(L1) * S2 / S1)
+		 *   S2 = S1 * square(L2 / L1)
+		 *
+		 *   Sag & Angle criterion.pdf
+		 *   https://docs.google.com/presentation/d/1bUnrRFWCW3sDn9ngb9ftfQS-2JxNJaUZlh783hZMMEw/edit?usp=sharing
+		 *
+		 *   @return the number of point to add after PointA to PointB to respect the chord error.
+		 *   if PointNum = 0, PointB is not necessary
+		 *   if PointNum = 1, PointB has to be add to respect the chord error and the segment AB is respecting the chord error
+		 *   if PointNum = 2, Point at FirstIndex and PointB have to be add to respect the chord error and both segments are respecting the chord error
+		 *   if PointNum > 2, Point at FirstIndex and PointB have to be add but at least (PointNum - 2) points have to be added in both segments to respect the chord error
+		 *
+		 * Parametrization error:
+		 *   This function check that for inner point of the segment AB, point M(u) is near its approximation = A + AB * (u - uA) / (uB - uA) than the MaxAllowedError
+		 *   To Remove Chord error, M approximation is compare to the projection of M(u) on AB
+		 */
+
+		virtual int32 CheckSamplingError(int32 FirstIndex, int32 EndIndex) override
+		{
+			const FPoint& APoint = this->Sampling.GetPoints()[this->StartSamplingSegmentIndex];
+			double ACoordinate = this->Sampling.Coordinates[this->StartSamplingSegmentIndex];
+
+			const FPoint& BPoint = this->EndStudySegment.Polyline->GetPoints()[this->EndStudySegment.Index];
+			double BCoordinate = this->EndStudySegment.Polyline->Coordinates[this->EndStudySegment.Index];
+
+			double ABCoordinate = BCoordinate - ACoordinate;
+			FPoint ABVector = BPoint - APoint;
+
+			double MaxChord = 0;
+			double MaxParametrizationError = 0;
+
+			for (int32 Index = FirstIndex; Index < EndIndex; ++Index)
+			{
+				const FPoint& PointM = this->CandidatePoints.GetPoints()[Index];
+
+				double CurvilinearAbscise;
+				FPoint ProjectedPoint = ProjectPointOnSegment(PointM, APoint, BPoint, CurvilinearAbscise, /*bRestrictCoodinateToInside*/ true);
+
+				double Chord = ProjectedPoint.Distance(PointM);
+				double AMCoordinate = this->CandidatePoints.Coordinates[Index] - ACoordinate;
+				FPoint MApproximate = APoint + (AMCoordinate / ABCoordinate) * ABVector;
+				double ParametrizationError = MApproximate.Distance(ProjectedPoint);
+
+				if (Chord > MaxChord)
+				{
+					MaxChord = Chord;
+				}
+
+				if (ParametrizationError > MaxParametrizationError)
+				{
+					MaxParametrizationError = ParametrizationError;
+				}
+			}
+
+#ifdef DEBUG_CHECKCHORDERROR
+			if (bDisplay)
+			{
+				{
+					int32 Index = 0;
+					F3DDebugSession G(FString::Printf(TEXT("CheckChord %f %d"), MaxChord, FirstIndex));
+
+					CADKernel::Display(APoint, EVisuProperty::BluePoint);
+					CADKernel::Display(BPoint, EVisuProperty::BluePoint);
+
+					for (int32 Index = FirstIndex; Index < EndIndex; ++Index)
+					{
+						const FPoint& Middle = CandidatePoints.GetPoints()[Index];
+						CADKernel::Display(Middle, EVisuProperty::YellowPoint);
+					}
+
+					DisplaySegment(APoint, BPoint, 0, EVisuProperty::BlueCurve);
+					DisplaySegment(APoint, CandidatePoints.GetPoints()[FirstIndex], 0, EVisuProperty::YellowCurve);
+					DisplaySegment(BPoint, CandidatePoints.GetPoints()[EndIndex - 1], 0, EVisuProperty::YellowCurve);
+
+					for (int32 Index = FirstIndex; Index < EndIndex - 1; ++Index)
+					{
+						DisplaySegment(CandidatePoints.GetPoints()[Index], CandidatePoints.GetPoints()[Index + 1], 0, EVisuProperty::YellowCurve);
+					}
+				}
+			}
+#endif
+
+			int32 SegmentNumToRespectParametrizationError = (int32)(MaxParametrizationError /  DesiredParameterizationError);
+			if (SegmentNumToRespectParametrizationError > 2)
+			{
+				SegmentNumToRespectParametrizationError = (int32)((SegmentNumToRespectParametrizationError - 1) / 2) + 2;
+			}
+
+			int32 SegmentNumToRespectChordError = 0;
+			if (MaxChord < this->DesiredChordError / 4)
+			{
+				// S2 = S1 * sqare(L2 / L1)
+				// if the length is doubled, the chord error is still smaller than the desired chord error.
+				// the end point is not needed
+				return 0;
+			}
+			else if (MaxChord < this->DesiredChordError)
+			{
+				// the end point is needed because the Chord is smaller but nearly equal to the desired chord error.
+				SegmentNumToRespectChordError = 1;
+			}
+			else if (MaxChord < 4 * this->DesiredChordError)
+			{
+				// if the length is divided by two, the chord error is divided by 4 and becomes smaller than the desired chord error.
+				// Add First index and Index
+				SegmentNumToRespectChordError = this->CheckTangentError(APoint, ACoordinate, BPoint, BCoordinate, FirstIndex, EndIndex, this->StartSamplingSegmentIndex);
+			}
+			else
+			{
+				SegmentNumToRespectChordError = this->CountOfNeededPointsToRespectChordError(APoint, BPoint, MaxChord);
+			}
+
+			return FMath::Max(SegmentNumToRespectChordError, SegmentNumToRespectParametrizationError);
+		};
+	};
+
+	/**
+	 * Sampler of surfacic curve based on parametrization and chord error control
+	 */
+	class FSurfacicCurveSamplerOnParam : public TSamplerBasedOnParametrizationAndChordError<FSurfacicPolyline, FPoint>
+	{
+	public:
+		FSurfacicCurveSamplerOnParam(const FRestrictionCurve& InCurve, const FLinearBoundary& InBoundary, double InMaxSagError, double InMaxParameterizationError, FSurfacicPolyline& OutPolyline)
+			: TSamplerBasedOnParametrizationAndChordError<FSurfacicPolyline, FPoint>(InBoundary, InMaxSagError, InMaxParameterizationError, OutPolyline)
+			, Surface(*InCurve.GetCarrierSurface())
+			, Curve(*InCurve.Get2DCurve())
+		{
+		};
+
+		FSurfacicCurveSamplerOnParam(const FSurface& InSurface, const FCurve& InCurve, const FLinearBoundary& InBoundary, double InMaxSagError, double InMaxParameterizationError, FSurfacicPolyline& OutPolyline)
+			: TSamplerBasedOnParametrizationAndChordError<FSurfacicPolyline, FPoint>(InBoundary, InMaxSagError, InMaxParameterizationError, OutPolyline)
+			, Surface(InSurface)
+			, Curve(InCurve)
+		{
+		};
+
+	protected:
+
+		const FSurface& Surface;
+		const FCurve& Curve;
+		double MaxParameterizationError;
+
+		virtual void EvaluatesNewCandidatePoints() override
+		{
+			CandidatePoints.bWithNormals = Sampling.bWithNormals;
+			CandidatePoints.bWithTangent = Sampling.bWithTangent;
+			if (CandidatePoints.bWithTangent)
+			{
+				TArray<FCurvePoint2D> OutPoints;
+				Curve.Evaluate2DPoints(CandidatePoints.Coordinates, OutPoints, 1);
+				Surface.EvaluatePoints(OutPoints, CandidatePoints);
+			}
+			else
+			{
+				Curve.Evaluate2DPoints(CandidatePoints.Coordinates, CandidatePoints.Points2D);
+				Surface.EvaluatePoints(CandidatePoints);
+			}
+		};
+
+		virtual void GetNotDerivableCoordinates(TArray<double>& OutNotDerivableCoordinates) override
+		{
+			Curve.FindNotDerivableCoordinates(Boundary, 2, OutNotDerivableCoordinates);
+		}
+
+	};
+
+	/**
+	 * Sampler of curve based on parametrization and chord error control
+	 */
+	class FCurveSamplerOnParam : public TSamplerBasedOnParametrizationAndChordError<FPolyline3D, FPoint>
+	{
+	public:
+		FCurveSamplerOnParam(const FCurve& InCurve, const FLinearBoundary& InBoundary, double InMaxSagError, double InMaxParameterizationError, FPolyline3D& OutPolyline)
+			: TSamplerBasedOnParametrizationAndChordError<FPolyline3D, FPoint>(InBoundary, InMaxSagError, InMaxParameterizationError, OutPolyline)
+			, Curve(InCurve)
+		{
+		};
+
+	protected:
+
+		const FCurve& Curve;
+		double MaxParameterizationError;
+
+		virtual void EvaluatesNewCandidatePoints() override
+		{
+			Curve.EvaluatePoints(CandidatePoints.Coordinates, CandidatePoints.Points);
+		};
+
+		virtual void GetNotDerivableCoordinates(TArray<double>& OutNotDerivableCoordinates) override
+		{
+			Curve.FindNotDerivableCoordinates(Boundary, 2, OutNotDerivableCoordinates);
+		}
+	};
+
+	/**
+	 * Sampler of surface based on parametrization and chord error control
+	 */
+	class FSurfaceSamplerOnParam : public TSamplerBasedOnParametrizationAndChordError<FPolyline3D, FPoint>
+	{
+
+	public:
+		FSurfaceSamplerOnParam(const FSurface& InSurface, const FSurfacicBoundary& InBoundary, double InMaxSagError, double InMaxParameterizationError, FCoordinateGrid& OutSampling)
+			: TSamplerBasedOnParametrizationAndChordError<FPolyline3D, FPoint>(Boundary, InMaxSagError, InMaxParameterizationError, OutPolyline)
+			, Surface(InSurface)
+			, SurfaceBoundary(InBoundary)
+			, SurfaceSampling(OutSampling)
+		{
+		};
+
+		void Set(EIso InIsoType, double InIsoCoordinate, const FLinearBoundary& CurveBounds)
+		{
+			IsoType = InIsoType;
+			IsoCoordinate.Empty(1);
+			IsoCoordinate.Add(InIsoCoordinate);
+			Boundary = CurveBounds;
+		}
+
+		virtual void Sample() override
+		{
+			Sample(EIso::IsoU);
+			Sample(EIso::IsoV);
+		}
+
+		void Sample(EIso InIsoType)
+		{
+#ifdef DEBUG_CURVE_SAMPLING
+			CurveIndex = CurveToDisplay - 1;
+#endif
+
+			double Middle = SurfaceBoundary[InIsoType].GetMiddle();
+
+			Set(InIsoType, Middle, SurfaceBoundary[Other(InIsoType)]);
+
+			SamplingInitalizing();
+			RunSampling();
+
+			double Step = SurfaceBoundary[InIsoType].Length() / 6;
+			IsoCoordinate[0] = SurfaceBoundary[InIsoType].GetMin();
+			for (int32 IsoIndex = 1; IsoIndex < 6; ++IsoIndex)
+			{
+				IsoCoordinate[0] += Step;
+				if (IsoIndex == 3)
+				{
+					continue;
+				}
+
+				Swap(CandidatePoints.Coordinates, Sampling.Coordinates);
+				EvaluatesNewCandidatePoints();
+				Swap(CandidatePoints.Coordinates, Sampling.Coordinates);
+				Swap(CandidatePoints.Points, Sampling.Points);
+
+				// Check the result
+				CandidatePoints.Coordinates.Empty(Sampling.Coordinates.Num());
+
+				for (int32 Index = 0; Index < Sampling.Coordinates.Num()-1; ++Index)
+				{
+					double UCoord = Sampling.Coordinates[Index] + (Sampling.Coordinates[Index + 1] - Sampling.Coordinates[Index]) * 0.5;
+					CandidatePoints.Coordinates.Add(UCoord);
+				}
+
+				EvaluatesNewCandidatePoints();
+
+#ifdef DEBUG_CURVE_SAMPLING
+				CurveToDisplay = CurveIndex;
+				DisplaySampling(CurveIndex == CurveToDisplay, 0);
+#endif
+
+				RunSampling();
+			}
+
+			Swap(SurfaceSampling[Other(InIsoType)], Sampling.Coordinates);
+		}
+
+		virtual void SamplingInitalizing() override
+		{
+
+#ifdef DEBUG_CURVE_SAMPLING
+			++CurveIndex;
+#endif
+			IsOptimalSegments.Empty(100);
+			Sampling.Empty(100);
+
+			CandidatePoints.Empty(100);
+
+			// Initialization of the algorithm with the not derivable coordinates of the curve and at least 5 temporary points
+			int32 ComplementatyPointCount = 0;
+			{
+				TArray<double> LocalNotDerivableCoordinates;
+				GetNotDerivableCoordinates(LocalNotDerivableCoordinates);
+				LocalNotDerivableCoordinates.Insert(Boundary.Min, 0);
+				LocalNotDerivableCoordinates.Add(Boundary.Max);
+
+				ComplementatyPointCount = LocalNotDerivableCoordinates.Num() < 5 ? 5 : 1;
+
+				NextCoordinates.Empty(LocalNotDerivableCoordinates.Num() * (ComplementatyPointCount + 1));
+				NextCoordinates.Add(Boundary.Min);
+
+				for (int32 Index = 0; Index < LocalNotDerivableCoordinates.Num() - 1; ++Index)
+				{
+					AddIntermediateCoordinates(LocalNotDerivableCoordinates[Index], LocalNotDerivableCoordinates[Index + 1], ComplementatyPointCount);
+					NextCoordinates.Add(LocalNotDerivableCoordinates[Index + 1]);
+				}
+			}
+
+			CandidatePoints.SwapCoordinates(NextCoordinates);
+			EvaluatesNewCandidatePoints();
+
+			// Not Derivable point initialize the sampling
+			ComplementatyPointCount++;
+			for (int32 Index = 0, ISampling = 0; Index < CandidatePoints.Size(); Index += ComplementatyPointCount, ++ISampling)
+			{
+				Sampling.EmplaceAt(ISampling, CandidatePoints, Index);
+				IsOptimalSegments.Add(false);
+			}
+			IsOptimalSegments.Pop();
+
+			for (int32 Index = CandidatePoints.Size() - 1; Index >= 0; Index -= ComplementatyPointCount)
+			{
+				CandidatePoints.RemoveAt(Index);
+			}
+
+			// first segment is not optimal
+
+#ifdef DEBUG_CURVE_SAMPLING
+			int32 Step = 0;
+			DisplaySampling(CurveIndex == CurveToDisplay, Step);
+#endif
+		}
+
+	protected:
+
+		const FSurface& Surface;
+		const FSurfacicBoundary& SurfaceBoundary;
+
+		FCoordinateGrid& SurfaceSampling;
+
+		FCoordinateGrid NotDerivableCoordinates;
+
+		FLinearBoundary Boundary;
+		bool bNotDerivableFound = false;
+
+		FPolyline3D OutPolyline;
+
+		TArray<double> IsoCoordinate;
+		EIso IsoType;
+
+		virtual void GetNotDerivableCoordinates(TArray<double>& OutNotDerivableCoordinates) override
+		{
+			if (!bNotDerivableFound)
+			{
+				Surface.LinesNotDerivables(Surface.GetBoundary(), 2, NotDerivableCoordinates);
+				bNotDerivableFound = true;
+			}
+
+			EIso CurveIso = IsoType == EIso::IsoU ? EIso::IsoV : EIso::IsoU;
+			const TArray<double>& SurfaceIsoNotDerivableCoordinates = NotDerivableCoordinates[CurveIso];
+			double Tolerance = Surface.GetIsoTolerance(CurveIso);
+			ArrayUtils::SubArrayWithoutBoundary(SurfaceIsoNotDerivableCoordinates, Boundary, Tolerance, OutNotDerivableCoordinates);
+		}
+
+		virtual void EvaluatesNewCandidatePoints() override
+		{
+			if (CandidatePoints.Coordinates.Num() > 0)
+			{
+				TArray<double>& UCoordinates = IsoType == EIso::IsoU ? IsoCoordinate : CandidatePoints.Coordinates;
+				TArray<double>& VCoordinates = IsoType == EIso::IsoU ? CandidatePoints.Coordinates : IsoCoordinate;
+
+				FSurfacicSampling PointGrid;
+				FCoordinateGrid Grid;
+
+				Grid.Swap(UCoordinates, VCoordinates);
+				Surface.EvaluatePointGrid(Grid, PointGrid);
+				Grid.Swap(UCoordinates, VCoordinates);
+
+				Swap(CandidatePoints.Points, PointGrid.Points3D);
+			}
+			else
+			{
+				CandidatePoints.Points.SetNum(0);
+			}
+		}
+
+	};
+
+}
diff --git a/Engine/Source/Runtime/Datasmith/CADKernel/Private/CADKernel/Geo/Sampling/SurfacicPolyline.cpp b/Engine/Source/Runtime/Datasmith/CADKernel/Private/CADKernel/Geo/Sampling/SurfacicPolyline.cpp
new file mode 100644
--- /dev/null
+++ b/Engine/Source/Runtime/Datasmith/CADKernel/Private/CADKernel/Geo/Sampling/SurfacicPolyline.cpp
@@ -0,0 +1,55 @@
+// Copyright Epic Games, Inc. All Rights Reserved.
+
+#include "CADKernel/Geo/Sampling/SurfacicPolyline.h"
+
+#include "CADKernel/Core/KernelParameters.h"
+#include "CADKernel/Core/System.h"
+#include "CADKernel/Geo/Curves/Curve.h"
+#include "CADKernel/Geo/Surfaces/Surface.h"
+#include "CADKernel/Geo/Sampler/SamplerOnParam.h"
+
+using namespace CADKernel;
+
+FSurfacicPolyline::FSurfacicPolyline(TSharedRef<FSurface> InCarrierSurface, TSharedRef<FCurve> Curve2D, const double InTolerance)
+	: bWithNormals(true)
+	, bWithTangent(false)
+{
+	FSurfacicCurveSamplerOnParam Sampler(InCarrierSurface.Get(), Curve2D.Get(), Curve2D->GetBoundary(), InTolerance, InTolerance, *this);
+	Sampler.Sample();
+}
+
+FSurfacicPolyline::FSurfacicPolyline(TSharedRef<FSurface> InCarrierSurface, TSharedRef<FCurve> Curve2D, const double ChordTolerance, const double ParamTolerance, bool bInWithNormals, bool bInWithTangents)
+	: bWithNormals(bInWithNormals)
+	, bWithTangent(bInWithTangents)
+{
+	FSurfacicCurveSamplerOnParam Sampler(InCarrierSurface.Get(), Curve2D.Get(), Curve2D->GetBoundary(), ChordTolerance, ParamTolerance, *this);
+	Sampler.Sample();
+}
+
+void FSurfacicPolyline::CheckIfDegenerated(const double Tolerance3D, const double Tolerance2D, const FLinearBoundary& Boudary, bool& bDegeneration2D, bool& bDegeneration3D, double& Length3D) const
+{
+	TPolylineApproximator<FPoint2D> Approximator(Coordinates, Points2D);
+	double Length2D = Approximator.ComputeLengthOfSubPolyline(Boudary);
+	if (FMath::IsNearlyZero(Length2D, FMath::Max((double)KINDA_SMALL_NUMBER, Tolerance2D)))
+	{
+		bDegeneration2D = true;
+		bDegeneration3D = true;
+		Length3D = 0.;
+		return;
+	}
+
+	bDegeneration2D = false;
+
+	TPolylineApproximator<FPoint> Approximator3D(Coordinates, Points3D);
+	Length3D = Approximator3D.ComputeLengthOfSubPolyline(Boudary);
+
+	if (FMath::IsNearlyZero(Length3D, FMath::Max((double)KINDA_SMALL_NUMBER, Tolerance3D)))
+	{
+		bDegeneration3D = true;
+	}
+	else
+	{
+		bDegeneration3D = false;
+	}
+}
+
diff --git a/Engine/Source/Runtime/Datasmith/CADKernel/Private/CADKernel/Geo/Surfaces/BezierSurface.cpp b/Engine/Source/Runtime/Datasmith/CADKernel/Private/CADKernel/Geo/Surfaces/BezierSurface.cpp
new file mode 100644
--- /dev/null
+++ b/Engine/Source/Runtime/Datasmith/CADKernel/Private/CADKernel/Geo/Surfaces/BezierSurface.cpp
@@ -0,0 +1,189 @@
+// Copyright Epic Games, Inc. All Rights Reserved.
+
+#include "CADKernel/Geo/Surfaces/BezierSurface.h"
+
+#include "CADKernel/Geo/Sampling/PolylineTools.h"
+#include "CADKernel/Geo/GeoPoint.h"
+#include "CADKernel/Math/BSpline.h"
+
+using namespace CADKernel;
+
+
+void FBezierSurface::EvaluatePoint(const FPoint2D& InPoint2D, FSurfacicPoint& OutPoint3D, int32 InDerivativeOrder) const
+{
+	TArray<FPoint> VCurvePoles;
+	TArray<FPoint> VCurveUGradient;
+	TArray<FPoint> VCurveULaplacian;
+
+	VCurvePoles.SetNum(VPoleNum);
+
+	OutPoint3D.DerivativeOrder = InDerivativeOrder;
+	
+	if (InDerivativeOrder > 0)
+	{
+		VCurveUGradient.Init(FPoint::ZeroPoint, VPoleNum);
+	}
+
+	if (InDerivativeOrder > 1)
+	{
+		VCurveULaplacian.SetNum(VPoleNum);
+	}
+
+	{
+		double Coordinate = InPoint2D.U;
+
+		TArray<FPoint> UAuxiliaryPoles;
+		UAuxiliaryPoles.SetNum(UPoleNum);
+
+		TArray<FPoint> UAuxiliaryGradient;
+		TArray<FPoint> UAuxiliaryLaplacian;
+
+		// For each iso V Curve compute Point, Gradient and Laplacian at U coordinate
+		for (int32 Vndex = 0, PoleIndex = 0; Vndex < VPoleNum; Vndex++, PoleIndex += UPoleNum)
+		{
+			UAuxiliaryPoles.Empty(UPoleNum);
+			UAuxiliaryPoles.Append(Poles.GetData() + PoleIndex, UPoleNum);
+
+			if (InDerivativeOrder > 0)
+			{
+				UAuxiliaryGradient.Init(FPoint::ZeroPoint, UPoleNum);
+			}
+			if (InDerivativeOrder > 1)
+			{
+				UAuxiliaryLaplacian.Init(FPoint::ZeroPoint, UPoleNum);
+			}
+
+			// Compute Point, Gradient and Laplacian at U coordinate with De Casteljau's algorithm, 
+			for (int32 Undex = UPoleNum - 2; Undex >= 0; Undex--)
+			{
+				for (int32 Index = 0; Index <= Undex; Index++)
+				{
+					const FPoint PointI = UAuxiliaryPoles[Index];
+					const FPoint& PointB = UAuxiliaryPoles[Index + 1];
+					const FPoint VectorIB = PointB - PointI;
+
+					UAuxiliaryPoles[Index] = PointI + VectorIB * Coordinate;
+					if (InDerivativeOrder > 0)
+					{
+						const FPoint GradientI = UAuxiliaryGradient[Index];
+						const FPoint& GradientB = UAuxiliaryGradient[Index + 1];
+						const FPoint VectorGradientIB = GradientB - GradientI;
+
+						UAuxiliaryGradient[Index] = GradientI + VectorGradientIB * Coordinate + VectorIB;
+						if (InDerivativeOrder > 1)
+						{
+							//UAuxiliaryLaplacian[Index] = UAuxiliaryLaplacian[Index] + (UAuxiliaryLaplacian[Index + 1] - UAuxiliaryLaplacian[Index]) * Coordinate + 2.0 * VectorGradientIB;
+							UAuxiliaryLaplacian[Index] = PolylineTools::LinearInterpolation(UAuxiliaryLaplacian, Index, Coordinate) + 2.0 * VectorGradientIB;
+						}
+					}
+				}
+			}
+
+			// Point, Gradient and Laplacian of the iso v curve are saved to defined the VCurve
+			VCurvePoles[Vndex] = UAuxiliaryPoles[0];
+			if (InDerivativeOrder > 0)
+			{
+				VCurveUGradient[Vndex] = UAuxiliaryGradient[0];
+
+				if (InDerivativeOrder > 1)
+				{
+					VCurveULaplacian[Vndex] = UAuxiliaryLaplacian[0];
+				}
+			}
+		}
+	}
+
+	TArray<FPoint> VAuxiliaryGradient;
+	TArray<FPoint> UVAuxiliaryLaplacian;
+	TArray<FPoint> VVAuxiliaryLaplacian;
+	if (InDerivativeOrder > 0)
+	{
+		VAuxiliaryGradient.Init(FPoint::ZeroPoint, VPoleNum);
+		if (InDerivativeOrder > 1)
+		{
+			VVAuxiliaryLaplacian.Init(FPoint::ZeroPoint, VPoleNum);
+			UVAuxiliaryLaplacian.Init(FPoint::ZeroPoint, VPoleNum);
+		}
+	}
+
+	double Coordinate = InPoint2D.V;
+
+	// Compute Point, Gradient and Laplacian at V coordinate with De Casteljau's algorithm,
+	for (int32 Vndex = VPoleNum - 2; Vndex >= 0; Vndex--)
+	{
+		for (int32 Index = 0; Index <= Vndex; Index++)
+		{
+			const FPoint PointI = VCurvePoles[Index];
+			const FPoint& PointB = VCurvePoles[Index + 1];
+			const FPoint VectorIB = PointB - PointI;
+			VCurvePoles[Index] = PointI + VectorIB * Coordinate;
+
+			if (InDerivativeOrder > 0) 
+			{
+				const FPoint UGradientI = VCurveUGradient[Index];
+				const FPoint& UGradientB = VCurveUGradient[Index + 1];
+				const FPoint VectorUGradientIB = UGradientB - UGradientI;
+
+				VCurveUGradient[Index] = UGradientI + VectorUGradientIB * Coordinate;
+
+				const FPoint VGradientI = VAuxiliaryGradient[Index]; 
+				const FPoint& VGradientB = VAuxiliaryGradient[Index + 1];
+				const FPoint VectorVGradientIB = VGradientB - VGradientI;
+
+				VAuxiliaryGradient[Index] = VGradientI + VectorVGradientIB * Coordinate + VectorIB;
+
+				if (InDerivativeOrder > 1)
+				{
+					//UVAuxiliaryLaplacian[Index] = UVAuxiliaryLaplacian[Index] + (UVAuxiliaryLaplacian[Index + 1] - UVAuxiliaryLaplacian[Index]) * Coordinate + VectorUGradientIB;
+					UVAuxiliaryLaplacian[Index] = PolylineTools::LinearInterpolation(UVAuxiliaryLaplacian, Index, Coordinate) + VectorUGradientIB;
+
+					//VCurveULaplacian[Index] = VCurveULaplacian[Index] + (VCurveULaplacian[Index + 1] - VCurveULaplacian[Index]) * Coordinate;
+					VCurveULaplacian[Index] = PolylineTools::LinearInterpolation(VCurveULaplacian, Index, Coordinate);
+
+					//VVAuxiliaryLaplacian[Index] = VVAuxiliaryLaplacian[Index] + (VVAuxiliaryLaplacian[Index + 1] - VVAuxiliaryLaplacian[Index]) * Coordinate + 2.0 * VectorVGradientIB;
+					VVAuxiliaryLaplacian[Index] = PolylineTools::LinearInterpolation(VVAuxiliaryLaplacian, Index, Coordinate) + 2.0 * VectorVGradientIB;;
+				}
+			}
+		}
+	}
+
+	OutPoint3D.Point = VCurvePoles[0];
+
+	if (InDerivativeOrder > 0) 
+	{
+		OutPoint3D.GradientU = VCurveUGradient[0];
+		OutPoint3D.GradientV = VAuxiliaryGradient[0];
+	}
+
+	if (InDerivativeOrder > 1) 
+	{
+		OutPoint3D.LaplacianU = VCurveULaplacian[0];
+		OutPoint3D.LaplacianV = VVAuxiliaryLaplacian[0];
+		OutPoint3D.LaplacianUV = UVAuxiliaryLaplacian[0];
+	}
+}
+
+void FBezierSurface::Presample(const FSurfacicBoundary& InBoundaries, FCoordinateGrid& Coordinates)
+{
+	ensureCADKernel(false);
+}
+
+TSharedPtr<FEntityGeom> FBezierSurface::ApplyMatrix(const FMatrixH& InMatrix) const
+{
+	TArray<FPoint> TransformedPoles;
+	TransformedPoles.Reserve(Poles.Num());
+	for (const FPoint& Pole : Poles) 
+	{
+		TransformedPoles.Emplace(InMatrix.Multiply(Pole));
+	}
+	return FEntity::MakeShared<FBezierSurface>(Tolerance3D, UDegre, VDegre, TransformedPoles);
+}
+
+#ifdef CADKERNEL_DEV
+FInfoEntity& FBezierSurface::GetInfo(FInfoEntity& Info) const
+{
+	return FSurface::GetInfo(Info).Add(TEXT("degre U"), UDegre)
+		.Add(TEXT("degre V"), VDegre)
+		.Add(TEXT("poles"), Poles);
+}
+#endif
diff --git a/Engine/Source/Runtime/Datasmith/CADKernel/Private/CADKernel/Geo/Surfaces/ConeSurface.cpp b/Engine/Source/Runtime/Datasmith/CADKernel/Private/CADKernel/Geo/Surfaces/ConeSurface.cpp
new file mode 100644
--- /dev/null
+++ b/Engine/Source/Runtime/Datasmith/CADKernel/Private/CADKernel/Geo/Surfaces/ConeSurface.cpp
@@ -0,0 +1,123 @@
+// Copyright Epic Games, Inc. All Rights Reserved.
+
+#include "CADKernel/Geo/Surfaces/ConeSurface.h"
+
+#include "CADKernel/Geo/GeoPoint.h"
+#include "CADKernel/Geo/Sampling/SurfacicSampling.h"
+#include "CADKernel/Math/Boundary.h"
+
+using namespace CADKernel;
+
+
+void FConeSurface::EvaluatePoint(const FPoint2D& InPoint2D, FSurfacicPoint& OutPoint3D, int32 InDerivativeOrder) const
+{
+	OutPoint3D.DerivativeOrder = InDerivativeOrder;
+
+	double DeltaVR = tan(ConeAngle);
+	double Radius = StartRadius + InPoint2D.V * DeltaVR;
+
+	OutPoint3D.Point.Set(Radius * cos(InPoint2D.U), Radius * sin(InPoint2D.U), InPoint2D.V);
+
+	OutPoint3D.Point = Matrix.Multiply(OutPoint3D.Point);
+
+	if (InDerivativeOrder > 0)
+	{
+		OutPoint3D.GradientU = FPoint(-Radius * sin(InPoint2D.U), Radius * cos(InPoint2D.U), 0.);
+		OutPoint3D.GradientV = FPoint(DeltaVR * cos(InPoint2D.U), DeltaVR * sin(InPoint2D.U), 1.0);
+
+		OutPoint3D.GradientU = Matrix.MultiplyVector(OutPoint3D.GradientU);
+		OutPoint3D.GradientV = Matrix.MultiplyVector(OutPoint3D.GradientV);
+
+		if (InDerivativeOrder > 1)
+		{
+			OutPoint3D.LaplacianU = FPoint(-Radius * cos(InPoint2D.U), -Radius * sin(InPoint2D.U), 0.);
+			OutPoint3D.LaplacianV = FPoint(0., 0., 0.);
+
+			OutPoint3D.LaplacianUV = FPoint(-DeltaVR * sin(InPoint2D.U), DeltaVR * cos(InPoint2D.U), 0.);
+
+			OutPoint3D.LaplacianU = Matrix.MultiplyVector(OutPoint3D.LaplacianU);
+			OutPoint3D.LaplacianV = Matrix.MultiplyVector(OutPoint3D.LaplacianV);
+			OutPoint3D.LaplacianUV = Matrix.MultiplyVector(OutPoint3D.LaplacianUV);
+		}
+	}
+
+}
+
+void FConeSurface::EvaluatePointGrid(const FCoordinateGrid& Coordinates, FSurfacicSampling& OutPoints, bool bComputeNormals) const
+{
+	OutPoints.bWithNormals = bComputeNormals;
+
+	int32 PointNum = Coordinates.Count();
+	OutPoints.Reserve(PointNum);
+
+	OutPoints.Set2DCoordinates(Coordinates);
+
+	int32 UCount = Coordinates.IsoCount(EIso::IsoU);
+	int32 VCount = Coordinates.IsoCount(EIso::IsoV);
+
+	TArray<double> CosU;
+	TArray<double> SinU;
+	CosU.Reserve(UCount);
+	SinU.Reserve(UCount);
+
+	double DeltaVR = tan(ConeAngle);
+
+	for (double Angle : Coordinates[EIso::IsoU])
+	{
+		CosU.Emplace(cos(Angle));
+	}
+
+	for (double Angle : Coordinates[EIso::IsoU])
+	{
+		SinU.Emplace(sin(Angle));
+	}
+
+	for (int32 Vndex = 0, Index = 0; Vndex < VCount; ++Vndex)
+	{
+		double Radius = StartRadius + Coordinates[EIso::IsoV][Vndex] * DeltaVR;
+
+		for (int32 Undex = 0; Undex < UCount; Undex++)
+		{
+			FPoint& Point = OutPoints.Points3D.Emplace_GetRef(Radius * CosU[Undex], Radius * SinU[Undex], Coordinates[EIso::IsoV][Vndex]);
+			Point = Matrix.Multiply(Point);
+		}
+	}
+
+	if (bComputeNormals)
+	{
+		for (int32 Vndex = 0, Index = 0; Vndex < VCount; ++Vndex)
+		{
+			double Radius = StartRadius + Vndex * DeltaVR;
+			for (int32 Undex = 0; Undex < UCount; Undex++)
+			{
+				FPoint GradientU(-Radius * SinU[Undex], Radius * CosU[Undex], 0.);
+				FPoint GradientV(DeltaVR * CosU[Undex], DeltaVR * SinU[Undex], 1.0);
+				GradientU = Matrix.MultiplyVector(GradientU);
+				GradientV = Matrix.MultiplyVector(GradientV);
+				OutPoints.Normals.Emplace(GradientU ^ GradientV);
+			}
+		}
+
+		OutPoints.NormalizeNormals();
+	}
+}
+
+TSharedPtr<FEntityGeom> FConeSurface::ApplyMatrix(const FMatrixH& NewMatrix) const
+{
+	FMatrixH Mat = NewMatrix * Matrix;
+	return FEntity::MakeShared<FConeSurface>(Tolerance3D, Mat, StartRadius, ConeAngle, Boundary[EIso::IsoU].Min, Boundary[EIso::IsoU].Max, Boundary[EIso::IsoV].Min, Boundary[EIso::IsoV].Max);
+}
+
+#ifdef CADKERNEL_DEV
+FInfoEntity& FConeSurface::GetInfo(FInfoEntity& Info) const
+{
+	return FSurface::GetInfo(Info)
+		.Add(TEXT("Matrix"), Matrix)
+		.Add(TEXT("StartRadius"), StartRadius)
+		.Add(TEXT("ConeAngle"), ConeAngle)
+		.Add(TEXT("StartRuleLength"), Boundary[EIso::IsoU].Min)
+		.Add(TEXT("EndRuleLength"), Boundary[EIso::IsoU].Max)
+		.Add(TEXT("StartAngle"), Boundary[EIso::IsoV].Min)
+		.Add(TEXT("EndAngle"), Boundary[EIso::IsoV].Max);
+}
+#endif
diff --git a/Engine/Source/Runtime/Datasmith/CADKernel/Private/CADKernel/Geo/Surfaces/CoonsSurface.cpp b/Engine/Source/Runtime/Datasmith/CADKernel/Private/CADKernel/Geo/Surfaces/CoonsSurface.cpp
new file mode 100644
--- /dev/null
+++ b/Engine/Source/Runtime/Datasmith/CADKernel/Private/CADKernel/Geo/Surfaces/CoonsSurface.cpp
@@ -0,0 +1,155 @@
+// Copyright Epic Games, Inc. All Rights Reserved.
+
+#include "CADKernel/Geo/Surfaces/CoonsSurface.h"
+
+#include "CADKernel/Geo/Curves/Curve.h"
+#include "CADKernel/Utils/ArrayUtils.h"
+
+using namespace CADKernel;
+
+
+
+FCoonsSurface::FCoonsSurface(const double InToleranceGeometric, TSharedPtr<FCurve> InCurves[4])
+	: FCoonsSurface(InToleranceGeometric, InCurves[0], InCurves[1], InCurves[2], InCurves[3])
+{
+}
+
+FCoonsSurface::FCoonsSurface(const double InToleranceGeometric, TSharedPtr<FCurve> InCurve1, TSharedPtr<FCurve> InCurve2, TSharedPtr<FCurve> InCurve3, TSharedPtr<FCurve> InCurve4)
+	: FSurface(InToleranceGeometric)
+{
+	Curves[0] = InCurve1;
+	Curves[1] = InCurve2;
+	Curves[2] = InCurve3;
+	Curves[3] = InCurve4;
+
+	Corners.SetNum(4);
+	Corners[0] = Curves[0]->EvaluatePoint(Curves[0]->GetUMin()) * 0.5 + Curves[2]->EvaluatePoint(Curves[2]->GetUMin()) * 0.5;
+	Corners[1] = Curves[0]->EvaluatePoint(Curves[0]->GetUMax()) * 0.5 + Curves[3]->EvaluatePoint(Curves[3]->GetUMin()) * 0.5;
+	Corners[2] = Curves[1]->EvaluatePoint(Curves[1]->GetUMin()) * 0.5 + Curves[2]->EvaluatePoint(Curves[2]->GetUMax()) * 0.5;
+	Corners[3] = Curves[1]->EvaluatePoint(Curves[1]->GetUMax()) * 0.5 + Curves[3]->EvaluatePoint(Curves[3]->GetUMax()) * 0.5;
+}
+
+void FCoonsSurface::LinesNotDerivables(const FSurfacicBoundary& Bounds, int32 InDerivativeOrder, FCoordinateGrid& OutNotDerivableCoordinates) const
+{
+	TFunction<void(const TSharedPtr<FCurve>&, int32, TArray<double>&)> FindLinesNotDerivables = [&](const TSharedPtr<FCurve>& Curve, int32 CurveIndex, TArray<double>& NotDerivables)
+	{
+		TArray<double> TmpCoords;
+		Curve->FindNotDerivableCoordinates(InDerivativeOrder, TmpCoords);
+		double Length = Curve->GetUMax() - Curve->GetUMin();
+		for (int32 Index = 0; Index < TmpCoords.Num(); Index++)
+		{
+			NotDerivables.Add((TmpCoords[Index] - Curve->GetUMin()) / Length);
+		}
+	};
+
+	FindLinesNotDerivables(Curves[0], 0, OutNotDerivableCoordinates[EIso::IsoU]);
+	FindLinesNotDerivables(Curves[1], 1, OutNotDerivableCoordinates[EIso::IsoU]);
+	FindLinesNotDerivables(Curves[2], 2, OutNotDerivableCoordinates[EIso::IsoV]);
+	FindLinesNotDerivables(Curves[3], 3, OutNotDerivableCoordinates[EIso::IsoV]);
+
+	for (int32 Iso = EIso::IsoU; Iso <= EIso::IsoV; Iso++)
+	{
+		Algo::Sort(OutNotDerivableCoordinates[(EIso)Iso]);
+		ArrayUtils::RemoveDuplicates(OutNotDerivableCoordinates[(EIso)Iso], GetIsoTolerances()[Iso]);
+	}
+
+}
+
+void FCoonsSurface::EvaluatePoint(const FPoint2D& InPoint2D, FSurfacicPoint& OutPoint3D, int32 InDerivativeOrder) const
+{
+	TFunction<void(const TSharedPtr<FCurve>*, const FPoint2D&, EIso, FPoint&, FCurvePoint[])> \
+	ComputePointOnRuledSurface = [&](const TSharedPtr<FCurve>* Curves, const FPoint2D& Coord, EIso Iso, FPoint& OutPoint, FCurvePoint OutCurvePoints[2])
+	{
+		const int32 IndexCurve0 = Iso * 2;
+		const int32 IndexCurve1 = Iso * 2 + 1;
+
+		EIso OtherIso = Iso == EIso::IsoU ? EIso::IsoV : EIso::IsoU;
+
+		double CoordinateCurve0 = Curves[IndexCurve0]->GetUMin() + Coord[Iso] * (Curves[IndexCurve0]->GetUMax() - Curves[IndexCurve0]->GetUMin());
+		double CoordinateCurve1 = Curves[IndexCurve1]->GetUMin() + Coord[Iso] * (Curves[IndexCurve1]->GetUMax() - Curves[IndexCurve1]->GetUMin());
+
+		Curves[IndexCurve0]->EvaluatePoint(CoordinateCurve0, OutCurvePoints[0], InDerivativeOrder);
+		Curves[IndexCurve1]->EvaluatePoint(CoordinateCurve1, OutCurvePoints[1], InDerivativeOrder);
+
+		OutPoint = OutCurvePoints[0].Point + (OutCurvePoints[1].Point - OutCurvePoints[0].Point) * Coord[OtherIso];
+
+		if (InDerivativeOrder > 0)
+		{
+			OutCurvePoints[0].Gradient *= (Curves[IndexCurve0]->GetUMax() - Curves[IndexCurve0]->GetUMin());
+			OutCurvePoints[1].Gradient *= (Curves[IndexCurve1]->GetUMax() - Curves[IndexCurve1]->GetUMin());
+			if (InDerivativeOrder > 1)
+			{
+				OutCurvePoints[0].Laplacian *= FMath::Square(Curves[IndexCurve0]->GetUMax() - Curves[IndexCurve0]->GetUMin());
+				OutCurvePoints[1].Laplacian *= FMath::Square(Curves[IndexCurve1]->GetUMax() - Curves[IndexCurve1]->GetUMin());
+			}
+		}
+	};
+
+	// 1. First ruled Surface
+	FPoint RuledPoint1;
+	FCurvePoint Curve12Points[2];
+	ComputePointOnRuledSurface(Curves, InPoint2D, EIso::IsoU, RuledPoint1, Curve12Points);
+	
+	// 2. Second ruled Surface
+	FPoint RuledPoint2;
+	FCurvePoint Curve34Points[2];
+	ComputePointOnRuledSurface(Curves+2, InPoint2D, EIso::IsoV, RuledPoint2, Curve34Points);
+
+	// 3. bilinear surface
+	FPoint BilinearPoint = (Corners[0]*(1.0-InPoint2D.U) + Corners[1]*InPoint2D.U)*(1.0-InPoint2D.V) + (Corners[2]*(1.0-InPoint2D.U) + Corners[3]*InPoint2D.U)*InPoint2D.V;
+
+	OutPoint3D.DerivativeOrder = InDerivativeOrder;
+
+	// Mixed of ruled surfaces et bilinear Surface
+	OutPoint3D.Point = RuledPoint1 + RuledPoint2 - BilinearPoint;
+
+	if (InDerivativeOrder > 0)
+	{
+		FPoint DeltaU1 = Curve12Points[0].Gradient + InPoint2D.V * (Curve12Points[1].Gradient - Curve12Points[0].Gradient);
+		FPoint DeltaU2 = Curve34Points[1].Point - Curve34Points[0].Point;
+		FPoint DeltaU3 = (Corners[1] * (1.0 - InPoint2D.V) + Corners[3] * InPoint2D.V) - (Corners[0] * (1.0 - InPoint2D.V) + Corners[2] * InPoint2D.V);
+		OutPoint3D.GradientU = DeltaU1 + DeltaU2 - DeltaU3;
+
+		FPoint DeltaV1 = Curve12Points[1].Point - Curve12Points[0].Point;;
+		FPoint DeltaV2 = Curve34Points[0].Gradient + InPoint2D.U * (Curve34Points[1].Gradient - Curve34Points[0].Gradient);
+		FPoint DeltaV3 = (Corners[2] * (1.0 - InPoint2D.U) + Corners[3] * InPoint2D.U) - (Corners[0] * (1.0 - InPoint2D.U) + Corners[1] * InPoint2D.U);
+		OutPoint3D.GradientV = DeltaV1 + DeltaV2 - DeltaV3;
+
+		if (InDerivativeOrder > 1)
+		{
+			FPoint LaplacianUV1 = (Curve12Points[1].Laplacian - Curve12Points[0].Laplacian);
+			FPoint LaplacianUV2 = (Curve34Points[1].Laplacian - Curve34Points[0].Laplacian);
+
+			OutPoint3D.LaplacianU = Curve12Points[0].Laplacian + InPoint2D.V * LaplacianUV1;
+			OutPoint3D.LaplacianV = Curve34Points[0].Laplacian + InPoint2D.U * LaplacianUV2;
+
+			OutPoint3D.LaplacianUV = LaplacianUV1 + LaplacianUV2;
+		}
+	}
+}
+
+TSharedPtr<FEntityGeom> FCoonsSurface::ApplyMatrix(const FMatrixH& InMatrix) const
+{
+	TSharedPtr<FCurve> TransformedCurves[4];
+
+	for (int32 Index = 0; Index < 4; ++Index)
+	{
+		TransformedCurves[Index] = StaticCastSharedPtr<FCurve>(Curves[Index]->ApplyMatrix(InMatrix));
+		if (!TransformedCurves[Index].IsValid())
+		{
+			return TSharedPtr<FEntityGeom>();
+		}
+	}
+
+	return FEntity::MakeShared<FCoonsSurface>(Tolerance3D, TransformedCurves);
+}
+
+#ifdef CADKERNEL_DEV
+FInfoEntity& FCoonsSurface::GetInfo(FInfoEntity& Info) const
+{
+	return FSurface::GetInfo(Info).Add(TEXT("curve 1"), Curves[0])
+							      .Add(TEXT("curve 2"), Curves[1])
+							      .Add(TEXT("curve 3"), Curves[2])
+							      .Add(TEXT("curve 4"), Curves[3]);
+}
+#endif
diff --git a/Engine/Source/Runtime/Datasmith/CADKernel/Private/CADKernel/Geo/Surfaces/CylinderSurface.cpp b/Engine/Source/Runtime/Datasmith/CADKernel/Private/CADKernel/Geo/Surfaces/CylinderSurface.cpp
new file mode 100644
--- /dev/null
+++ b/Engine/Source/Runtime/Datasmith/CADKernel/Private/CADKernel/Geo/Surfaces/CylinderSurface.cpp
@@ -0,0 +1,64 @@
+// Copyright Epic Games, Inc. All Rights Reserved.
+
+#include "CADKernel/Geo/Surfaces/CylinderSurface.h"
+
+#include "CADKernel/Geo/GeoPoint.h"
+#include "CADKernel/Geo/Sampling/SurfacicSampling.h"
+
+
+using namespace CADKernel;
+
+FCylinderSurface::FCylinderSurface(const double InToleranceGeometric, const FMatrixH& INMatrix, double InRadius, double InStartLength, double InEndLength, double InStartAngle, double InEndAngle)
+	: FSurface(InToleranceGeometric, InStartAngle, InEndAngle, InStartLength, InEndLength)
+	, Matrix(INMatrix)
+	, Radius(InRadius)
+{
+}
+
+void FCylinderSurface::InitBoundary() const
+{
+}
+
+void FCylinderSurface::EvaluatePoint(const FPoint2D& InPoint2D, FSurfacicPoint& OutPoint3D, int32 InDerivativeOrder) const
+{
+	OutPoint3D.DerivativeOrder = InDerivativeOrder;
+
+	OutPoint3D.Point.Set(Radius*(cos(InPoint2D.U)), Radius*(sin(InPoint2D.U)), InPoint2D.V);
+	OutPoint3D.Point = Matrix.Multiply(OutPoint3D.Point);
+
+	if(InDerivativeOrder>0) 
+	{
+		OutPoint3D.GradientU = FPoint(Radius*(-sin(InPoint2D.U)), Radius*(cos(InPoint2D.U)), 0.0);
+		OutPoint3D.GradientV = FPoint(0.0, 0.0, 1.0);
+
+		OutPoint3D.GradientU = Matrix.MultiplyVector(OutPoint3D.GradientU);
+		OutPoint3D.GradientV = Matrix.MultiplyVector(OutPoint3D.GradientV);
+
+		if (InDerivativeOrder > 1) 
+		{
+			OutPoint3D.LaplacianU = FPoint(Radius * (-cos(InPoint2D.U)), Radius * (-sin(InPoint2D.U)), 0.0);
+			OutPoint3D.LaplacianU = Matrix.MultiplyVector(OutPoint3D.LaplacianU);
+
+			OutPoint3D.LaplacianV = FPoint(0.0, 0.0, 0.0);
+			OutPoint3D.LaplacianUV = FPoint(0.0, 0.0, 0.0);
+		}
+	}
+}
+
+TSharedPtr<FEntityGeom> FCylinderSurface::ApplyMatrix(const FMatrixH& InMatrix) const
+{
+	FMatrixH NewMatrix = InMatrix * Matrix;
+	return FEntity::MakeShared<FCylinderSurface>(Tolerance3D, NewMatrix, Radius, Boundary[EIso::IsoV].Min, Boundary[EIso::IsoU].Max, Boundary[EIso::IsoU].Min, Boundary[EIso::IsoU].Max);
+}
+
+#ifdef CADKERNEL_DEV
+FInfoEntity& FCylinderSurface::GetInfo(FInfoEntity& Info) const
+{
+	return FSurface::GetInfo(Info).Add(TEXT("Matrix"), Matrix)
+							 .Add(TEXT("Radius"), Radius)
+							 .Add(TEXT("StartLength"), Boundary[EIso::IsoV].Min)
+							 .Add(TEXT("EndLength"), Boundary[EIso::IsoU].Max)
+							 .Add(TEXT("StartAngle"), Boundary[EIso::IsoU].Min)
+							 .Add(TEXT("EndAngle"), Boundary[EIso::IsoU].Max);
+}
+#endif
diff --git a/Engine/Source/Runtime/Datasmith/CADKernel/Private/CADKernel/Geo/Surfaces/OffsetSurface.cpp b/Engine/Source/Runtime/Datasmith/CADKernel/Private/CADKernel/Geo/Surfaces/OffsetSurface.cpp
new file mode 100644
--- /dev/null
+++ b/Engine/Source/Runtime/Datasmith/CADKernel/Private/CADKernel/Geo/Surfaces/OffsetSurface.cpp
@@ -0,0 +1,88 @@
+// Copyright Epic Games, Inc. All Rights Reserved.
+
+#include "CADKernel/Geo/Surfaces/OffsetSurface.h"
+
+#include "CADKernel/Core/KernelParameters.h"
+#include "CADKernel/Core/System.h"
+#include "CADKernel/Geo/GeoPoint.h"
+#include "CADKernel/Geo/Sampling/SurfacicSampling.h"
+#include "CADKernel/UI/Display.h"
+
+using namespace CADKernel;
+
+void FOffsetSurface::InitBoundary() const
+{
+	Boundary = BaseSurface->GetBoundary();
+}
+
+void FOffsetSurface::LinesNotDerivables(const FSurfacicBoundary& Bounds, int32 derivatedOrder, FCoordinateGrid& NotDerivables) const
+{
+	BaseSurface->LinesNotDerivables(Bounds, derivatedOrder, NotDerivables);
+}
+
+TSharedPtr<FEntityGeom> FOffsetSurface::ApplyMatrix(const FMatrixH& InMatrix) const
+{
+	TSharedPtr<FSurface> TransformedBaseSurface;
+
+	TransformedBaseSurface = StaticCastSharedPtr<FSurface>(BaseSurface->ApplyMatrix(InMatrix));
+	if (!TransformedBaseSurface.IsValid())
+	{
+		return TSharedPtr<FEntityGeom>();
+	}
+
+	return FEntity::MakeShared<FOffsetSurface>(Tolerance3D, TransformedBaseSurface.ToSharedRef(), Offset);
+}
+
+#ifdef CADKERNEL_DEV
+FInfoEntity& FOffsetSurface::GetInfo(FInfoEntity& Info) const
+{
+	return FSurface::GetInfo(Info).Add(TEXT("base Surface"), BaseSurface)
+		.Add(TEXT("distance"), Offset);
+}
+#endif
+
+void FOffsetSurface::EvaluatePoint(const FPoint2D& InSurfacicCoordinate, FSurfacicPoint& OutPoint3D, int32 InDerivativeOrder) const
+{
+	if (InDerivativeOrder == 0)
+	{
+		++InDerivativeOrder;
+	}
+	BaseSurface->EvaluatePoint(InSurfacicCoordinate, OutPoint3D, InDerivativeOrder);
+
+	FPoint Normal = OutPoint3D.GradientU ^ OutPoint3D.GradientV;
+	Normal.Normalize();
+	OutPoint3D.Point += Normal*Offset;
+}
+
+void FOffsetSurface::EvaluatePoints(const TArray<FPoint2D>& InSurfacicCoordinates, TArray<FSurfacicPoint>& OutPoint3D, int32 InDerivativeOrder) const
+{
+	if (InDerivativeOrder == 0)
+	{
+		++InDerivativeOrder;
+	}
+
+	BaseSurface->EvaluatePoints(InSurfacicCoordinates, OutPoint3D, InDerivativeOrder);
+
+	for (FSurfacicPoint& Point : OutPoint3D)
+	{
+		FPoint Normal = Point.GradientU ^ Point.GradientV;
+		Normal.Normalize();
+		Point.Point += Normal * Offset;
+	}
+}
+
+void FOffsetSurface::EvaluatePointGrid(const FCoordinateGrid& Coords, FSurfacicSampling& OutPoints, bool bComputeNormals) const
+{
+	OutPoints.bWithNormals = bComputeNormals;
+	BaseSurface->EvaluatePointGrid(Coords, OutPoints, true);
+
+	for (int32 Index = 0; Index < OutPoints.Count(); ++Index)
+	{
+		OutPoints.Points3D[Index] += OutPoints.Normals[Index] * Offset;
+	}
+}
+
+void FOffsetSurface::Presample(const FSurfacicBoundary& InBoundaries, FCoordinateGrid& Coordinates)
+{
+	BaseSurface->Presample(InBoundaries, Coordinates);
+}
\ No newline at end of file
diff --git a/Engine/Source/Runtime/Datasmith/CADKernel/Private/CADKernel/Geo/Surfaces/PlaneSurface.cpp b/Engine/Source/Runtime/Datasmith/CADKernel/Private/CADKernel/Geo/Surfaces/PlaneSurface.cpp
new file mode 100644
--- /dev/null
+++ b/Engine/Source/Runtime/Datasmith/CADKernel/Private/CADKernel/Geo/Surfaces/PlaneSurface.cpp
@@ -0,0 +1,189 @@
+// Copyright Epic Games, Inc. All Rights Reserved.
+
+#include "CADKernel/Geo/Surfaces/PlaneSurface.h"
+
+#include "CADKernel/Core/KernelParameters.h"
+#include "CADKernel/Core/System.h"
+#include "CADKernel/Geo/GeoPoint.h"
+#include "CADKernel/Geo/Sampling/SurfacicSampling.h"
+#include "CADKernel/Math/Plane.h"
+
+using namespace CADKernel;
+
+FPlaneSurface::FPlaneSurface(const double InToleranceGeometric, const FMatrixH& InMatrix)
+	: FSurface(InToleranceGeometric, -HUGE_VALUE, HUGE_VALUE, -HUGE_VALUE, HUGE_VALUE)
+{
+	Matrix = InMatrix;
+
+	ensureCADKernel(FMath::IsNearlyZero(Matrix.Get(3, 0)) && FMath::IsNearlyZero(Matrix.Get(3, 1)) && FMath::IsNearlyZero(Matrix.Get(3, 2)));
+
+	InverseMatrix = Matrix;
+	InverseMatrix.Inverse();
+}
+
+FPlaneSurface::FPlaneSurface(const double InToleranceGeometric, const FPoint& Position, FPoint Normal)
+	: FSurface(InToleranceGeometric, -HUGE_VALUE, HUGE_VALUE, -HUGE_VALUE, HUGE_VALUE)
+{
+	Normal.Normalize();
+	Matrix.FromAxisOrigin(Normal, Position);
+
+	InverseMatrix = Matrix;
+	InverseMatrix.Inverse();
+}
+
+const void FPlaneSurface::ComputeIsoTolerances() const
+{
+	FPoint Origin = Matrix.Multiply(FPoint::ZeroPoint);
+
+	TFunction<double(FPoint)> ComputeTolerance = [&](FPoint Point2D) 
+	{
+		FPoint Point = Matrix.Multiply(Point2D);
+		double Length = Origin.Distance(Point);
+		if (Length > SMALL_NUMBER)
+		{
+			return Tolerance3D / Length;
+		}
+		else
+		{
+			return Tolerance3D;
+		}
+	};
+
+	FSurfacicTolerance& Tolerance = ToleranceIsos;
+
+	FPoint Point2DU { 1 , 0, 0 };
+	Tolerance[EIso::IsoU] = ComputeTolerance(Point2DU);
+
+	FPoint Point2DV{ 0, 1, 0 };
+	Tolerance[EIso::IsoV] = ComputeTolerance(Point2DV);
+}
+
+CADKernel::FPlane FPlaneSurface::GetPlane() const
+{
+	FSurfacicPoint Point3D;
+	EvaluatePoint(FPoint2D::ZeroPoint, Point3D, 0);
+	FPoint Normal = Matrix.Column(2);
+	return CADKernel::FPlane(Point3D.Point, Normal);
+}
+
+void FPlaneSurface::EvaluatePoint(const FPoint2D& InSurfacicCoordinate, FSurfacicPoint& OutPoint3D, int32 InDerivativeOrder) const
+{
+	OutPoint3D.DerivativeOrder = InDerivativeOrder;
+	OutPoint3D.Point = Matrix.Multiply(InSurfacicCoordinate);
+	if(InDerivativeOrder>0) 
+	{
+		OutPoint3D.GradientU = Matrix.Column(0);
+		OutPoint3D.GradientV = Matrix.Column(1);
+	}
+
+	if(InDerivativeOrder>1) 
+	{
+		OutPoint3D.LaplacianU = FPoint::ZeroPoint;
+		OutPoint3D.LaplacianV = FPoint::ZeroPoint;
+		OutPoint3D.LaplacianUV = FPoint::ZeroPoint;
+	}
+}
+
+void FPlaneSurface::EvaluatePoints(const TArray<FPoint2D>& InSurfacicCoordinates, TArray<FSurfacicPoint>& OutPoint3D, int32 InDerivativeOrder) const
+{
+	int32 PointNum = InSurfacicCoordinates.Num();
+	OutPoint3D.SetNum(PointNum);
+
+	for (int32 Index = 0; Index < PointNum; ++Index)
+	{
+		OutPoint3D[Index].DerivativeOrder = InDerivativeOrder;
+	}
+
+	for (int32 Index = 0; Index < PointNum; ++Index)
+	{
+		OutPoint3D[Index].Point = Matrix.Multiply(InSurfacicCoordinates[Index]);
+	}
+
+	if (InDerivativeOrder > 0)
+	{
+		FPoint GradientU(Matrix.Column(0));
+		FPoint GradientV(Matrix.Column(1));
+
+		for (int32 Index = 0; Index < PointNum; ++Index)
+		{
+			OutPoint3D[Index].GradientU = GradientU;
+		}
+
+		for (int32 Index = 0; Index < PointNum; ++Index)
+		{
+			OutPoint3D[Index].GradientV = GradientV;
+		}
+	}
+}
+
+void FPlaneSurface::EvaluatePointGrid(const FCoordinateGrid& Coordinates, FSurfacicSampling& OutPoints, bool bComputeNormals) const
+{
+	OutPoints.bWithNormals = bComputeNormals;
+
+	int32 PointNum = Coordinates.Count();
+	OutPoints.Reserve(PointNum);
+
+	OutPoints.Set2DCoordinates(Coordinates);
+
+	for (FPoint Point : OutPoints.Points2D)
+	{
+		OutPoints.Points3D.Emplace(Matrix.Multiply(Point));
+	}
+
+	if(bComputeNormals)
+	{
+		FPoint Normal(Matrix.Column(2));
+		OutPoints.Normals.Init(Normal, PointNum);
+	}
+}
+
+FPoint FPlaneSurface::ProjectPoint(const FPoint& Point, FPoint* OutProjectedPoint) const
+{
+	FPoint PointCoordinate = InverseMatrix.Multiply(Point);
+	PointCoordinate.Z = 0.0;
+
+	if(OutProjectedPoint)
+	{
+		*OutProjectedPoint = Matrix.Multiply(PointCoordinate);
+	}
+
+	return PointCoordinate;
+}
+
+void FPlaneSurface::ProjectPoints(const TArray<FPoint>& Points, TArray<FPoint>* PointCoordinates, TArray<FPoint>* OutProjectedPointS) const
+{
+	PointCoordinates->Reserve(Points.Num());
+	if(OutProjectedPointS) 
+	{
+		OutProjectedPointS->Reserve(Points.Num());
+	}
+
+	for (const FPoint& Point : Points)
+	{
+		FPoint& PointCoordinate = PointCoordinates->Emplace_GetRef(InverseMatrix.Multiply(Point));
+		PointCoordinate.Z = 0.0;
+	}
+
+	if (OutProjectedPointS)
+	{
+		for (const FPoint& PointCoordinate : *PointCoordinates)
+		{
+			OutProjectedPointS->Emplace(Matrix.Multiply(PointCoordinate));
+		}
+	}
+}
+
+TSharedPtr<FEntityGeom> FPlaneSurface::ApplyMatrix(const FMatrixH& InMatrix) const
+{
+	FMatrixH NewMatrix = InMatrix * Matrix;
+	return FEntity::MakeShared<FPlaneSurface>(Tolerance3D, NewMatrix);
+}
+
+#ifdef CADKERNEL_DEV
+FInfoEntity& FPlaneSurface::GetInfo(FInfoEntity& Info) const
+{
+	return FSurface::GetInfo(Info)
+		.Add(TEXT("Matrix"), Matrix)
+		.Add(TEXT("Inverse"), InverseMatrix);
+}
+#endif
diff --git a/Engine/Source/Runtime/Datasmith/CADKernel/Private/CADKernel/Geo/Surfaces/SphericalSurface.cpp b/Engine/Source/Runtime/Datasmith/CADKernel/Private/CADKernel/Geo/Surfaces/SphericalSurface.cpp
new file mode 100644
--- /dev/null
+++ b/Engine/Source/Runtime/Datasmith/CADKernel/Private/CADKernel/Geo/Surfaces/SphericalSurface.cpp
@@ -0,0 +1,115 @@
+// Copyright Epic Games, Inc. All Rights Reserved.
+
+#include "CADKernel/Geo/Surfaces/SphericalSurface.h"
+
+#include "CADKernel/Geo/GeoPoint.h"
+#include "CADKernel/Geo/Sampling/SurfacicSampling.h"
+
+
+using namespace CADKernel;
+
+void FSphericalSurface::EvaluatePoint(const FPoint2D& InSurfacicCoordinate, FSurfacicPoint& OutPoint3D, int32 InDerivativeOrder) const
+{
+	double CosU = cos(InSurfacicCoordinate.U);
+	double CosV = cos(InSurfacicCoordinate.V);
+
+	double SinU = sin(InSurfacicCoordinate.U);
+	double SinV = sin(InSurfacicCoordinate.V);
+
+	OutPoint3D.DerivativeOrder = InDerivativeOrder;
+	OutPoint3D.Point.Set(Radius * CosV * CosU, Radius * CosV * SinU, Radius * SinV);
+	OutPoint3D.Point = Matrix.Multiply(OutPoint3D.Point);
+
+	if (InDerivativeOrder > 0)
+	{
+		OutPoint3D.GradientU = FPoint(-Radius * CosV * SinU, Radius * CosV * CosU, 0.0);
+		OutPoint3D.GradientV = FPoint(-Radius * SinV * CosU, -Radius * SinV * SinU, Radius * CosV);
+
+		OutPoint3D.GradientU = Matrix.MultiplyVector(OutPoint3D.GradientU);
+		OutPoint3D.GradientV = Matrix.MultiplyVector(OutPoint3D.GradientV);
+	}
+
+	if (InDerivativeOrder > 1)
+	{
+		OutPoint3D.LaplacianU = FPoint(-Radius * CosV * CosU, -Radius * CosV * SinU, 0.0);
+		OutPoint3D.LaplacianV = FPoint(OutPoint3D.LaplacianU.X, OutPoint3D.LaplacianU.Y, -Radius * SinV);
+		OutPoint3D.LaplacianUV = FPoint(Radius * SinV * SinU, -Radius * SinV * CosU, 0.);
+
+		OutPoint3D.LaplacianU = Matrix.MultiplyVector(OutPoint3D.LaplacianU);
+		OutPoint3D.LaplacianV = Matrix.MultiplyVector(OutPoint3D.LaplacianV);
+		OutPoint3D.LaplacianUV = Matrix.MultiplyVector(OutPoint3D.LaplacianUV);
+	}
+}
+
+void FSphericalSurface::EvaluatePointGrid(const FCoordinateGrid& Coordinates, FSurfacicSampling& OutPoints, bool bComputeNormals) const
+{
+	OutPoints.bWithNormals = bComputeNormals;
+
+	int32 PointNum = Coordinates.Count();
+	OutPoints.Reserve(PointNum);
+
+	OutPoints.Set2DCoordinates(Coordinates);
+
+	int32 UCount = Coordinates.IsoCount(EIso::IsoU);
+
+	TArray<double> CosU;
+	TArray<double> SinU; 
+	CosU.Reserve(UCount);
+	SinU.Reserve(UCount);
+
+	for(double Angle : Coordinates[EIso::IsoU])
+	{
+		CosU.Emplace(cos(Angle));
+	}
+
+	for (double Angle : Coordinates[EIso::IsoU])
+	{
+		SinU.Emplace(sin(Angle));
+	}
+
+
+	for (double VAngle : Coordinates[EIso::IsoV])
+	{
+		double CosV = cos(VAngle);
+		double SinV = sin(VAngle);
+
+		for (int32 Undex = 0; Undex < UCount; Undex++)
+		{
+			FPoint& Point = OutPoints.Points3D.Emplace_GetRef(Radius * CosV * CosU[Undex], Radius * CosV * SinU[Undex], Radius * SinV);
+			Point = Matrix.Multiply(Point);
+		}
+	}
+
+	if (bComputeNormals)
+	{
+		FPoint Center = Matrix.Column(3);
+		for (FPoint& Point : OutPoints.Points3D)
+		{
+			OutPoints.Normals.Emplace(Point - Center);
+		}
+
+		OutPoints.NormalizeNormals();
+	}
+}
+
+TSharedPtr<FEntityGeom> FSphericalSurface::ApplyMatrix(const FMatrixH& InMatrix) const
+{
+	FMatrixH NewMatrix = InMatrix * Matrix;
+	return FEntity::MakeShared<FSphericalSurface>(Tolerance3D, NewMatrix, Radius, 
+		Boundary[EIso::IsoU].Min, Boundary[EIso::IsoU].Max,
+		Boundary[EIso::IsoV].Min, Boundary[EIso::IsoV].Max);
+}
+
+#ifdef CADKERNEL_DEV
+FInfoEntity& FSphericalSurface::GetInfo(FInfoEntity& Info) const
+{
+	return FSurface::GetInfo(Info)
+		.Add(TEXT("Matrix"), Matrix)
+		.Add(TEXT("Radius"), Radius)
+		.Add(TEXT("MeridianStartAngle"), Boundary[EIso::IsoU].Min)
+		.Add(TEXT("MeridianEndAngle"),   Boundary[EIso::IsoU].Max)
+		.Add(TEXT("ParallelStartAngle"), Boundary[EIso::IsoV].Min)
+		.Add(TEXT("ParallelEndAngle"),   Boundary[EIso::IsoV].Max);
+}
+#endif
+
