CADKernel Source Code v0.01

diff --git a/Engine/Source/Runtime/Datasmith/CADKernel/Private/CADKernel/Geo/Surfaces/Surface.cpp b/Engine/Source/Runtime/Datasmith/CADKernel/Private/CADKernel/Geo/Surfaces/Surface.cpp
new file mode 100644
--- /dev/null
+++ b/Engine/Source/Runtime/Datasmith/CADKernel/Private/CADKernel/Geo/Surfaces/Surface.cpp
@@ -0,0 +1,282 @@
+// Copyright Epic Games, Inc. All Rights Reserved.
+
+#include "CADKernel/Geo/Surfaces/Surface.h"
+
+#include "CADKernel/Core/KernelParameters.h"
+#include "CADKernel/Core/System.h"
+
+#include "CADKernel/Geo/GeoPoint.h"
+#include "CADKernel/Geo/Sampler/SamplerOnParam.h"
+#include "CADKernel/Geo/Sampling/SurfacicPolyline.h"
+#include "CADKernel/Geo/Sampling/SurfacicSampling.h"
+
+#include "CADKernel/Mesh/Structure/Grid.h"
+
+namespace CADKernel
+{
+#ifdef CADKERNEL_DEV
+	FInfoEntity& FSurface::GetInfo(FInfoEntity& Info) const
+	{
+		return FEntityGeom::GetInfo(Info)
+			.Add(TEXT("Surface type"), SurfacesTypesNames[(uint8)GetSurfaceType()])
+			.Add(TEXT("Boundary"), Boundary);
+	}
+#endif
+
+	void FSurface::Sample(const FSurfacicBoundary& InBoundary, int32 NumberOfSubdivisions[2], FCoordinateGrid& OutCoordinateSampling) const
+	{
+		for (int32 Index = EIso::IsoU; Index <= EIso::IsoV; Index++)
+		{
+			EIso Iso = (EIso)Index;
+			double Step = (InBoundary[Iso].Max - InBoundary[Iso].Min) / ((double)(NumberOfSubdivisions[Iso] - 1));
+
+			OutCoordinateSampling[Iso].Empty(NumberOfSubdivisions[Iso]);
+			double Value = InBoundary[Iso].Min;
+			for (int32 Kndex = 0; Kndex < NumberOfSubdivisions[Iso]; Kndex++)
+			{
+				OutCoordinateSampling[Iso].Add(Value);
+				Value += Step;
+			}
+		}
+	}
+
+	void FSurface::Sample(const FSurfacicBoundary& Bounds, int32 NumberOfSubdivisions[2], FSurfacicSampling& OutPointSampling) const
+	{
+		FCoordinateGrid CoordinateSampling;
+		Sample(Bounds, NumberOfSubdivisions, CoordinateSampling);
+		EvaluatePointGrid(CoordinateSampling, OutPointSampling);
+	}
+
+	void FSurface::Presample(const FSurfacicBoundary& InBoundaries, FCoordinateGrid& OutCoordinates)
+	{
+		FSurfaceSamplerOnParam Sampler(*this, InBoundaries, Tolerance3D * 10, Tolerance3D, OutCoordinates);
+		Sampler.Sample();
+	}
+
+	const void FSurface::ComputeIsoTolerances() const
+	{
+		FSurfacicSampling Grid;
+		int32 NumberOfSubdivisions[2] = { 10, 10 };
+
+		Sample(Boundary, NumberOfSubdivisions, Grid);
+
+		const FSurfacicBoundary& SurfaceBoundaries = GetBoundary();
+
+		FSurfacicTolerance& Tolerances = ToleranceIsos;
+
+		TFunction<double(double, double)> ComputeTolerance = [&](double CurveLength3D, double CurveLength2D)
+		{
+			if (CurveLength3D > SMALL_NUMBER)
+			{
+				double Tolerance2D = Tolerance3D * CurveLength2D / CurveLength3D;
+				return FMath::Max(FMath::Min(Tolerance2D, CurveLength2D / 2.), SMALL_NUMBER_SQUARE);
+			}
+			return CurveLength2D / 2.;
+		};
+
+		// Compute 2D tolerance along U
+		{
+			double MaxLengthIsoV = 0.0;
+
+			// For each isoV curve
+			for (int32 Index1 = 0, IndexPoint = 0; Index1 < NumberOfSubdivisions[EIso::IsoV]; ++Index1, ++IndexPoint)
+			{
+				double Length = 0.0;
+
+				// Compute the length of the curve
+				for (int32 Index2 = 0; Index2 < NumberOfSubdivisions[EIso::IsoU] - 1; ++Index2, IndexPoint ++)
+				{
+					Length += Grid.Points3D[IndexPoint].Distance(Grid.Points3D[IndexPoint + 1]);
+				}
+
+				if (Length > MaxLengthIsoV)
+				{
+					MaxLengthIsoV = Length;
+				}
+			}
+
+			Tolerances[EIso::IsoU] = ComputeTolerance(MaxLengthIsoV, SurfaceBoundaries[EIso::IsoU].Length());
+		};
+
+		// Compute 2D tolerance along V
+		{
+			double MaxLengthIsoU = 0.0;
+
+			int32 IndexPoint = 0;
+			int32 Delta = NumberOfSubdivisions[EIso::IsoU];
+
+			// For each isoU curve
+			for (int32 Index1 = 0; Index1 < NumberOfSubdivisions[EIso::IsoU]; ++Index1, ++IndexPoint)
+			{
+				double Length = 0.0;
+					IndexPoint = Index1;
+
+				// Compute the length of the curve
+				for (int32 Index2 = 0; Index2 < NumberOfSubdivisions[EIso::IsoV] - 1; ++Index2, IndexPoint += Delta)
+				{
+					Length += Grid.Points3D[IndexPoint].Distance(Grid.Points3D[IndexPoint + Delta]);
+				}
+
+				if (Length > MaxLengthIsoU)
+				{
+					MaxLengthIsoU = Length;
+				}
+			}
+
+			Tolerances[EIso::IsoV] = ComputeTolerance(MaxLengthIsoU, SurfaceBoundaries[EIso::IsoV].Length());
+		};
+	}
+
+	void FSurface::EvaluatePoints(const TArray<FPoint2D>& InSurfacicCoordinates, TArray<FSurfacicPoint>& OutPoint3D, int32 InDerivativeOrder) const
+	{
+		OutPoint3D.SetNum(InSurfacicCoordinates.Num());
+		for (int32 Index = 0; Index < InSurfacicCoordinates.Num(); ++Index)
+		{
+			EvaluatePoint(InSurfacicCoordinates[Index], OutPoint3D[Index], InDerivativeOrder);
+		}
+	}
+
+	void FSurface::EvaluatePoints(const TArray<FCurvePoint2D>& InSurfacicCoordinates, TArray<FSurfacicPoint>& OutPoint3D, int32 InDerivativeOrder) const
+	{
+		OutPoint3D.SetNum(InSurfacicCoordinates.Num());
+		for (int32 Index = 0; Index < InSurfacicCoordinates.Num(); Index++)
+		{
+			EvaluatePoint(InSurfacicCoordinates[Index].Point, OutPoint3D[Index], InDerivativeOrder);
+		}
+	}
+
+	void FSurface::EvaluatePoints(const TArray<FCurvePoint2D>& InSurfacicCoordinates, TArray<FCurvePoint>& OutPoints3D, int32 InDerivativeOrder) const
+	{
+		OutPoints3D.SetNum(InSurfacicCoordinates.Num());
+		TArray<FSurfacicPoint> SurfacicPoints3D;
+		EvaluatePoints(InSurfacicCoordinates, SurfacicPoints3D, InDerivativeOrder);
+		for(int32 Index = 0; Index < InSurfacicCoordinates.Num(); ++Index)
+		{
+			OutPoints3D[Index].Combine(InSurfacicCoordinates[Index], SurfacicPoints3D[Index]);
+		}
+	}
+
+	void FSurface::EvaluatePoints(FSurfacicPolyline& Polyline) const
+	{
+		int32 DerivativeOrder = Polyline.bWithNormals ? 1 : 0;
+
+		TArray<FSurfacicPoint> Point3D;
+		EvaluatePoints(Polyline.Points2D, Point3D, DerivativeOrder);
+
+		Polyline.Points3D.Empty(Polyline.Points2D.Num());
+		for (FSurfacicPoint Point : Point3D)
+		{
+			Polyline.Points3D.Emplace(Point.Point);
+		}
+
+		if (Polyline.bWithNormals)
+		{
+			Polyline.Normals.Empty(Polyline.Points2D.Num());
+			for (FSurfacicPoint Point : Point3D)
+			{
+				Polyline.Normals.Emplace(Point.GradientU ^ Point.GradientV);
+			}
+			for (FPoint Normal : Polyline.Normals)
+			{
+				Normal.Normalize();
+			}
+		}
+	}
+
+	void FSurface::EvaluatePoints(const TArray<FCurvePoint2D>& InPoints2D, FSurfacicPolyline& Polyline) const
+	{
+		int32 DerivativeOrder = 1;
+
+		TArray<FSurfacicPoint> Point3D;
+		EvaluatePoints(InPoints2D, Point3D, DerivativeOrder);
+
+		Polyline.Points2D.Empty(InPoints2D.Num());
+		Polyline.Points3D.Empty(InPoints2D.Num());
+		Polyline.Tangents.Empty(InPoints2D.Num());
+		Polyline.Normals.Empty(InPoints2D.Num());
+
+		for (const FCurvePoint2D& Point : InPoints2D)
+		{
+			Polyline.Points2D.Emplace(Point.Point);
+		}
+
+		for (const FSurfacicPoint& Point : Point3D)
+		{
+			Polyline.Points3D.Emplace(Point.Point);
+		}
+
+		for (int32 Index = 0; Index < InPoints2D.Num(); ++Index)
+		{
+			Polyline.Tangents.Emplace(Point3D[Index].GradientU * InPoints2D[Index].Gradient.U + Point3D[Index].GradientV * InPoints2D[Index].Gradient.V);
+		}
+
+		for (const FSurfacicPoint& Point : Point3D)
+		{
+			Polyline.Normals.Emplace(Point.GradientU ^ Point.GradientV);
+		}
+
+		for (FPoint& Normal : Polyline.Normals)
+		{
+			Normal.Normalize();
+		}
+	}
+
+	void FSurface::EvaluateNormals(const TArray<FPoint2D>& InPoints2D, TArray<FPoint>& Normals) const
+	{
+		int32 DerivativeOrder = 1;
+		TArray<FSurfacicPoint> Point3D;
+		EvaluatePoints(InPoints2D, Point3D, DerivativeOrder);
+
+		Normals.Empty(InPoints2D.Num());
+
+		for (const FSurfacicPoint& Point : Point3D)
+		{
+			Normals.Emplace(Point.GradientU ^ Point.GradientV);
+		}
+
+		for (FPoint& Normal : Normals)
+		{
+			Normal.Normalize();
+		}
+	}
+
+	void FSurface::EvaluatePointGrid(const FCoordinateGrid& Coordinates, FSurfacicSampling& OutPoints, bool bComputeNormals) const
+	{
+		OutPoints.Reserve(Coordinates.Count());
+
+		OutPoints.Set2DCoordinates(Coordinates);
+
+		int32 DerivativeOrder = bComputeNormals ? 1 : 0;
+		TArray<FSurfacicPoint> Point3D;
+		EvaluatePoints(OutPoints.Points2D, Point3D, DerivativeOrder);
+
+		OutPoints.bWithNormals = bComputeNormals;
+
+		for (FSurfacicPoint Point : Point3D)
+		{
+			OutPoints.Points3D.Emplace(Point.Point);
+		}
+
+		if (bComputeNormals)
+		{
+			for (FSurfacicPoint Point : Point3D)
+			{
+				OutPoints.Normals.Emplace(Point.GradientU ^ Point.GradientV);
+			}
+			OutPoints.NormalizeNormals();
+		}
+	}
+
+	void FSurface::EvaluateGrid(FGrid& Grid) const
+	{
+		FSurfacicSampling OutPoints;
+
+		const FCoordinateGrid& CoordinateGrid = Grid.GetCuttingCoordinates();
+		EvaluatePointGrid(CoordinateGrid, OutPoints, true);
+
+		Swap(OutPoints.Points3D, Grid.GetInner3DPoints());
+		Swap(OutPoints.Normals, Grid.GetNormals());
+	}
+
+} // namespace CADKernel
+
diff --git a/Engine/Source/Runtime/Datasmith/CADKernel/Private/CADKernel/Geo/Surfaces/SurfaceSerialize.cpp b/Engine/Source/Runtime/Datasmith/CADKernel/Private/CADKernel/Geo/Surfaces/SurfaceSerialize.cpp
new file mode 100644
--- /dev/null
+++ b/Engine/Source/Runtime/Datasmith/CADKernel/Private/CADKernel/Geo/Surfaces/SurfaceSerialize.cpp
@@ -0,0 +1,64 @@
+// Copyright Epic Games, Inc. All Rights Reserved.
+
+#include "CADKernel/Geo/Surfaces/Surface.h"
+
+#include "CADKernel/Core/CADKernelArchive.h"
+
+#include "CADKernel/Geo/Surfaces/PlaneSurface.h"
+#include "CADKernel/Geo/Surfaces/BezierSurface.h"
+#include "CADKernel/Geo/Surfaces/NURBSSurface.h"
+#include "CADKernel/Geo/Surfaces/RuledSurface.h"
+#include "CADKernel/Geo/Surfaces/RevolutionSurface.h"
+#include "CADKernel/Geo/Surfaces/TabulatedCylinderSurface.h"
+#include "CADKernel/Geo/Surfaces/CylinderSurface.h"
+#include "CADKernel/Geo/Surfaces/OffsetSurface.h"
+#include "CADKernel/Geo/Surfaces/CompositeSurface.h"
+#include "CADKernel/Geo/Surfaces/CoonsSurface.h"
+#include "CADKernel/Geo/Surfaces/SphericalSurface.h"
+#include "CADKernel/Geo/Surfaces/TorusSurface.h"
+#include "CADKernel/Geo/Surfaces/ConeSurface.h"
+
+
+namespace CADKernel
+{
+	TSharedPtr<FSurface> FSurface::Deserialize(FCADKernelArchive& Archive)
+	{
+		ESurface SurfaceType = ESurface::None;
+		Archive << SurfaceType;
+
+		switch (SurfaceType)
+		{
+		case ESurface::Bezier:
+			return FEntity::MakeShared<FBezierSurface>(Archive);
+		case ESurface::Cone:
+			return FEntity::MakeShared<FConeSurface>(Archive);
+		case ESurface::Composite:
+			return FEntity::MakeShared<FCompositeSurface>(Archive);
+		case ESurface::Coons:
+			return FEntity::MakeShared<FCoonsSurface>(Archive);
+		case ESurface::Cylinder:
+			return FEntity::MakeShared<FCylinderSurface>(Archive);
+		case ESurface::Nurbs:
+			return FEntity::MakeShared<FNURBSSurface>(Archive);
+		case ESurface::Offset:
+			return FEntity::MakeShared<FOffsetSurface>(Archive);
+		case ESurface::Plane:
+			return FEntity::MakeShared<FPlaneSurface>(Archive);
+		case ESurface::Revolution:
+			return FEntity::MakeShared<FRevolutionSurface>(Archive);
+		case ESurface::Ruled:
+			return FEntity::MakeShared<FRuledSurface>(Archive);
+		case ESurface::Sphere:
+			return FEntity::MakeShared<FSphericalSurface>(Archive);
+		case ESurface::TabulatedCylinder:
+			return FEntity::MakeShared<FTabulatedCylinderSurface>(Archive);
+		case ESurface::Torus:
+			return FEntity::MakeShared<FTorusSurface>(Archive);
+		default:
+			return TSharedPtr<FSurface>();
+		}
+		return TSharedPtr<FSurface>();
+	}
+
+} // namespace CADKernel
+
diff --git a/Engine/Source/Runtime/Datasmith/CADKernel/Private/CADKernel/Geo/Surfaces/TorusSurface.cpp b/Engine/Source/Runtime/Datasmith/CADKernel/Private/CADKernel/Geo/Surfaces/TorusSurface.cpp
new file mode 100644
--- /dev/null
+++ b/Engine/Source/Runtime/Datasmith/CADKernel/Private/CADKernel/Geo/Surfaces/TorusSurface.cpp
@@ -0,0 +1,126 @@
+// Copyright Epic Games, Inc. All Rights Reserved.
+
+#include "CADKernel/Geo/Surfaces/TorusSurface.h"
+
+#include "CADKernel/Geo/GeoPoint.h"
+#include "CADKernel/Geo/Sampling/SurfacicSampling.h"
+
+using namespace CADKernel;
+
+TSharedPtr<FEntityGeom> FTorusSurface::ApplyMatrix(const FMatrixH& InMatrix) const
+{
+	FMatrixH NewMatrix = InMatrix * Matrix;
+	return FEntity::MakeShared<FTorusSurface>(Tolerance3D, NewMatrix, MajorRadius, MinorRadius, Boundary[EIso::IsoU].Min, Boundary[EIso::IsoU].Max, Boundary[EIso::IsoV].Min, Boundary[EIso::IsoV].Max);
+}
+
+#ifdef CADKERNEL_DEV
+FInfoEntity& FTorusSurface::GetInfo(FInfoEntity& Info) const
+{
+	return FSurface::GetInfo(Info)
+		.Add(TEXT("Matrix"), Matrix)
+		.Add(TEXT("MajorRadius"), MajorRadius)
+		.Add(TEXT("MinorRadius"), MinorRadius)
+		.Add(TEXT("MajorStartAngle"), Boundary[EIso::IsoU].Min)
+		.Add(TEXT("MajorEndAngle"), Boundary[EIso::IsoU].Max)
+		.Add(TEXT("MinorStartAngle"), Boundary[EIso::IsoV].Min)
+		.Add(TEXT("MinorEndAngle"), Boundary[EIso::IsoV].Max);
+}
+#endif
+
+void FTorusSurface::EvaluatePoint(const FPoint2D& InSurfacicCoordinate, FSurfacicPoint& OutPoint3D, int32 InDerivativeOrder) const
+{
+	double CosU = cos(InSurfacicCoordinate.U);
+	double CosV = cos(InSurfacicCoordinate.V);
+
+	double SinU = sin(InSurfacicCoordinate.U);
+	double SinV = sin(InSurfacicCoordinate.V);
+
+	OutPoint3D.DerivativeOrder = InDerivativeOrder;
+	OutPoint3D.Point.Set((MajorRadius + MinorRadius * CosV) * CosU, (MajorRadius + MinorRadius * CosV) * SinU, MinorRadius * SinV);
+
+	OutPoint3D.Point = Matrix.Multiply(OutPoint3D.Point);
+
+	if (InDerivativeOrder > 0) 
+	{
+		OutPoint3D.GradientU = FPoint((MajorRadius + MinorRadius * CosV) * -SinU, (MajorRadius + MinorRadius * CosV) * CosU, 0.0);
+		OutPoint3D.GradientV = FPoint((MinorRadius * -SinV) * CosU, (MinorRadius * -SinV) * SinU, MinorRadius * CosV);
+
+		OutPoint3D.GradientU = Matrix.MultiplyVector(OutPoint3D.GradientU);
+		OutPoint3D.GradientV = Matrix.MultiplyVector(OutPoint3D.GradientV);
+
+		if (InDerivativeOrder > 1)
+		{
+			OutPoint3D.LaplacianU = FPoint((MajorRadius + MinorRadius * CosV) * -CosU, (MajorRadius + MinorRadius * CosV) * -SinU, 0.0);
+			OutPoint3D.LaplacianV = FPoint((MinorRadius * -CosV) * CosU, (MinorRadius * -CosV) * SinU, MinorRadius * -SinV);
+			OutPoint3D.LaplacianUV = FPoint((MinorRadius * -SinV) * -SinU, (MinorRadius * -SinV) * CosU, 0.0);
+
+			OutPoint3D.LaplacianU = Matrix.MultiplyVector(OutPoint3D.LaplacianU);
+			OutPoint3D.LaplacianV = Matrix.MultiplyVector(OutPoint3D.LaplacianV);
+			OutPoint3D.LaplacianUV = Matrix.MultiplyVector(OutPoint3D.LaplacianUV);
+		}
+	}
+}
+
+void FTorusSurface::EvaluatePointGrid(const FCoordinateGrid& Coordinates, FSurfacicSampling& OutPoints, bool bComputeNormals) const
+{
+	int32 PointNum = Coordinates.Count();
+
+	OutPoints.bWithNormals = bComputeNormals;
+
+	OutPoints.Reserve(PointNum);
+	OutPoints.Set2DCoordinates(Coordinates);
+
+	int32 UCount = Coordinates.IsoCount(EIso::IsoU);
+
+	TArray<double> CosU;
+	TArray<double> SinU;
+	CosU.Reserve(UCount);
+	SinU.Reserve(UCount);
+
+	for (double Angle : Coordinates[EIso::IsoU])
+	{
+		CosU.Emplace(cos(Angle));
+	}
+
+	for (double Angle : Coordinates[EIso::IsoU])
+	{
+		SinU.Emplace(sin(Angle));
+	}
+
+
+	for (double Angle : Coordinates[EIso::IsoV])
+	{
+		double CosV = cos(Angle);
+		double SinV = sin(Angle);
+
+		for (int32 Undex = 0; Undex < UCount; Undex++)
+		{
+			OutPoints.Points3D.Emplace((MajorRadius + MinorRadius * CosV) * CosU[Undex], (MajorRadius + MinorRadius * CosV) * SinU[Undex], MinorRadius * SinV);
+		}
+
+		if (bComputeNormals)
+		{
+			for (int32 Undex = 0; Undex < UCount; Undex++)
+			{
+				FPoint GradientU = FPoint((MajorRadius + MinorRadius * CosV) * -SinU[Undex], (MajorRadius + MinorRadius * CosV) * CosU[Undex], 0.0);
+				FPoint GradientV = FPoint((MinorRadius * -SinV) * CosU[Undex], (MinorRadius * -SinV) * SinU[Undex], MinorRadius * CosV);
+				OutPoints.Normals.Emplace(GradientU ^ GradientV);
+			}
+		}
+	}
+
+	for (FPoint& Point : OutPoints.Points3D)
+	{
+		Point = Matrix.Multiply(Point);
+	}
+
+	if (bComputeNormals)
+	{
+		FPoint Center = Matrix.Column(3);
+		for (FPoint& Normal : OutPoints.Normals)
+		{
+			 Normal = Matrix.PointRotation(Normal, Center);
+		}
+		OutPoints.NormalizeNormals();
+	}
+}
diff --git a/Engine/Source/Runtime/Datasmith/CADKernel/Private/CADKernel/Math/Geometry.cpp b/Engine/Source/Runtime/Datasmith/CADKernel/Private/CADKernel/Math/Geometry.cpp
new file mode 100644
--- /dev/null
+++ b/Engine/Source/Runtime/Datasmith/CADKernel/Private/CADKernel/Math/Geometry.cpp
@@ -0,0 +1,136 @@
+// Copyright Epic Games, Inc. All Rights Reserved.
+
+#include "CADKernel/Math/Geometry.h"
+
+#include "CADKernel/Core/KernelParameters.h" 
+#include "CADKernel/Core/System.h"
+#include "CADKernel/Math/Geometry.h"
+#include "CADKernel/Math/MatrixH.h"
+#include "CADKernel/Utils/Util.h"
+
+namespace CADKernel
+{
+	double ComputeCurvature(const FPoint& Gradient, const FPoint& Laplacian)
+	{
+		FPoint GradientCopy = Gradient;
+		FPoint LaplacianCopy = Laplacian;
+		GradientCopy.Normalize();
+		LaplacianCopy.Normalize();
+		FPoint Normal = GradientCopy ^ LaplacianCopy;
+		return (Normal.Length() * Laplacian.Length()) / (Gradient.Length() * Gradient.Length());
+	}
+
+	double ComputeCurvature(const FPoint& Normal, const FPoint& Gradient, const FPoint& Laplacian)
+	{
+		FPoint GradientCopy = Gradient;
+		FPoint LaplacianCopy = Laplacian;
+		FPoint NormalCopy = Normal;
+		GradientCopy.Normalize();
+		LaplacianCopy.Normalize();
+		NormalCopy.Normalize();
+		FPoint Coef = (LaplacianCopy ^ GradientCopy) ^ NormalCopy;
+		return (Coef.Length() * Laplacian.Length()) / Gradient.SquareLength();
+	}
+
+	void FindLoopIntersectionsWithIso(const EIso Iso, const double IsoParameter, const TArray<TArray<FPoint2D>>& Loops, TArray<double>& OutIntersections)
+	{
+		OutIntersections.Empty(8);
+
+		int32 UIndex = Iso == EIso::IsoU ? 0 : 1;
+		int32 VIndex = Iso == EIso::IsoU ? 1 : 0;
+
+		TFunction<void(const FPoint2D&, const FPoint2D&)> ComputeIntersection = [&](const FPoint2D& Point1, const FPoint2D& Point2)
+		{
+			if (IsoParameter > Point1[UIndex] && IsoParameter <= Point2[UIndex])
+			{
+				double Intersection = (IsoParameter - Point1[UIndex]) / (Point2[UIndex] - Point1[UIndex]) * (Point2[VIndex] - Point1[VIndex]) + Point1[VIndex];
+				OutIntersections.Add((IsoParameter - Point1[UIndex]) / (Point2[UIndex] - Point1[UIndex]) * (Point2[VIndex] - Point1[VIndex]) + Point1[VIndex]);
+			}
+		};
+
+		for (const TArray<FPoint2D>& Loop : Loops)
+		{
+			const FPoint2D* Point1 = &Loop.Last();
+			for (const FPoint2D& Point2 : Loop)
+			{
+				if (!FMath::IsNearlyEqual((*Point1)[UIndex], Point2[UIndex]))
+				{
+					if ((*Point1)[UIndex] < Point2[UIndex])
+					{
+						ComputeIntersection(*Point1, Point2);
+					}
+					else
+					{
+						ComputeIntersection(Point2, *Point1);
+					}
+				}
+				Point1 = &Point2;
+			}
+		}
+		Algo::Sort(OutIntersections);
+	}
+
+	bool IntersectSegments2D(const TSegment<FPoint2D>& SegmentAB, const TSegment<FPoint2D>& SegmentCD)
+	{
+		constexpr const double Min = -SMALL_NUMBER;
+		constexpr const double Max = 1. + SMALL_NUMBER;
+
+		TFunction<bool(double, double, double, double)> Intersect = [](double A, double B, double C, double D)
+		{
+			return !((D < A) || (B < C));
+		};
+
+		TFunction<bool(double, double, double, double)> TestWhenParallel = [&](double A, double B, double C, double D)
+		{
+			if (A < B)
+			{
+				if (C < D)
+				{
+					return Intersect(A, B, C, D);
+				}
+				else
+				{
+					return Intersect(A, B, D, C);
+				}
+			}
+			else
+			{
+				if (C < D)
+				{
+					return Intersect(B, A, C, D);
+				}
+				else
+				{
+					return Intersect(B, A, D, C);
+				}
+			}
+		};
+
+		FPoint2D AB = SegmentAB[1] - SegmentAB[0];
+		FPoint2D CD = SegmentCD[1] - SegmentCD[0];
+		FPoint2D CA = SegmentAB[0] - SegmentCD[0];
+
+		const double ParallelCoef = CD ^ AB;
+		if (FMath::IsNearlyZero(ParallelCoef))
+		{
+			const double ParallelCoef2 = CA ^ AB;
+			if (!FMath::IsNearlyZero(ParallelCoef2))
+			{
+				return false;
+			}
+
+			if (fabs(AB.U) > fabs(AB.V))
+			{
+				return TestWhenParallel(SegmentAB[0].U, SegmentAB[1].U, SegmentCD[0].U, SegmentCD[1].U);
+			}
+			else
+			{
+				return TestWhenParallel(SegmentAB[0].V, SegmentAB[1].V, SegmentCD[0].V, SegmentCD[1].V);
+			}
+		}
+
+		const double ABIntersectionCoordinate = (CA ^ CD) / ParallelCoef;
+		const double CDIntersectionCoordinate = (CA ^ AB) / ParallelCoef;
+		return (ABIntersectionCoordinate <= Max && ABIntersectionCoordinate >= Min && CDIntersectionCoordinate <= Max && CDIntersectionCoordinate >= Min);
+	}
+}
diff --git a/Engine/Source/Runtime/Datasmith/CADKernel/Private/CADKernel/Math/MatrixH.cpp b/Engine/Source/Runtime/Datasmith/CADKernel/Private/CADKernel/Math/MatrixH.cpp
new file mode 100644
--- /dev/null
+++ b/Engine/Source/Runtime/Datasmith/CADKernel/Private/CADKernel/Math/MatrixH.cpp
@@ -0,0 +1,364 @@
+// Copyright Epic Games, Inc. All Rights Reserved.
+
+#include "CADKernel/Math/MatrixH.h"
+
+#include "CADKernel/Math/MathConst.h"
+#include "CADKernel/UI/Message.h"
+
+namespace CADKernel
+{
+
+	const FMatrixH FMatrixH::Identity;
+
+	FPoint FMatrixH::Multiply(const FPoint& Point) const
+	{
+		return FPoint(
+			Point.X * Get(0, 0) + Point.Y * Get(0, 1) + Point.Z * Get(0, 2) + Get(0, 3),
+			Point.X * Get(1, 0) + Point.Y * Get(1, 1) + Point.Z * Get(1, 2) + Get(1, 3),
+			Point.X * Get(2, 0) + Point.Y * Get(2, 1) + Point.Z * Get(2, 2) + Get(2, 3)
+		);
+	}
+
+	/**
+	 * Apply only rotation matrix to the vector
+	 */
+	FPoint FMatrixH::MultiplyVector(const FPoint& Vector) const
+	{
+		return FPoint(
+			Vector.X * Get(0, 0) + Vector.Y * Get(0, 1) + Vector.Z * Get(0, 2),
+			Vector.X * Get(1, 0) + Vector.Y * Get(1, 1) + Vector.Z * Get(1, 2),
+			Vector.X * Get(2, 0) + Vector.Y * Get(2, 1) + Vector.Z * Get(2, 2)
+		);
+	}
+
+	void FMatrixH::BuildChangeOfCoordinateSystemMatrix(const FPoint& Xaxis, const FPoint& Yaxis, const FPoint& Zaxis, const FPoint& Origin)
+	{
+		Get(0, 0) = Xaxis[0];
+		Get(1, 0) = Xaxis[1];
+		Get(2, 0) = Xaxis[2];
+		Get(3, 0) = 0.0;
+
+		Get(0, 1) = Yaxis[0];
+		Get(1, 1) = Yaxis[1];
+		Get(2, 1) = Yaxis[2];
+		Get(3, 1) = 0.0;
+
+		Get(0, 2) = Zaxis[0];
+		Get(1, 2) = Zaxis[1];
+		Get(2, 2) = Zaxis[2];
+		Get(3, 2) = 0.0;
+
+		Get(0, 3) = Origin[0];
+		Get(1, 3) = Origin[1];
+		Get(2, 3) = Origin[2];
+		Get(3, 3) = 1.0;
+	}
+
+	void FMatrixH::FromAxisOrigin(const FPoint& Axis, const FPoint& Origin)
+	{
+		FPoint Zaxis = FPoint(0, 1, 0);
+		FPoint Xaxis = Zaxis ^ Axis;
+
+		//on cherche le vecteur ox
+		Xaxis.Normalize();
+		if (FMath::Abs(Xaxis.Length()) < SMALL_NUMBER)
+		{
+			Zaxis = FPoint(1, 0, 0);
+			Xaxis = Axis ^ Zaxis;
+			Xaxis.Normalize();
+			if (FMath::Abs(Xaxis.Length()) < SMALL_NUMBER)
+			{
+				Zaxis = FPoint(0, 0, 1);
+				Xaxis = Axis ^ Zaxis;
+
+				Xaxis.Normalize();
+				ensureCADKernel(FMath::Abs(Xaxis.Length()) > SMALL_NUMBER);
+			}
+		}
+
+		FPoint YAxis = Axis ^ Xaxis;
+		Zaxis = Axis;
+
+		YAxis.Normalize();
+		Zaxis.Normalize();
+
+		BuildChangeOfCoordinateSystemMatrix(Xaxis, YAxis, Zaxis, Origin);
+	}
+
+	FMatrixH FMatrixH::MakeRotationMatrix(double Angle, FPoint Axe)
+	{
+		FMatrixH Matrix;
+		Matrix.SetIdentity();
+
+		ensureCADKernel(Axe.Length() > SMALL_NUMBER);
+		Axe.Normalize();
+
+		Matrix.Get(0, 0) = Axe[0] * Axe[0] + (double)cos(Angle) * ((double)1.0 - Axe[0] * Axe[0]);
+		Matrix.Get(0, 1) = (double)(1.0 - cos(Angle)) * Axe[0] * Axe[1] - (double)sin(Angle) * Axe[2];
+		Matrix.Get(0, 2) = (double)(1.0 - cos(Angle)) * Axe[0] * Axe[2] + (double)sin(Angle) * Axe[1];
+
+		Matrix.Get(1, 0) = (double)(1.0 - cos(Angle)) * Axe[1] * Axe[0] + (double)sin(Angle) * Axe[2];
+		Matrix.Get(1, 1) = Axe[1] * Axe[1] + (double)cos(Angle) * ((double)1.0 - Axe[1] * Axe[1]);
+		Matrix.Get(1, 2) = (double)(1.0 - cos(Angle)) * Axe[1] * Axe[2] - (double)sin(Angle) * Axe[0];
+
+		Matrix.Get(2, 0) = (double)(1.0 - cos(Angle)) * Axe[2] * Axe[0] - (double)sin(Angle) * Axe[1];
+		Matrix.Get(2, 1) = (double)(1.0 - cos(Angle)) * Axe[2] * Axe[1] + (double)sin(Angle) * Axe[0];
+		Matrix.Get(2, 2) = Axe[2] * Axe[2] + (double)cos(Angle) * ((double)1.0 - Axe[2] * Axe[2]);
+		return Matrix;
+	}
+
+	FMatrixH FMatrixH::MakeTranslationMatrix(const FPoint& Point)
+	{
+		FMatrixH Matrix;
+		Matrix.SetIdentity();
+		Matrix.Get(0, 3) = Point.X;
+		Matrix.Get(1, 3) = Point.Y;
+		Matrix.Get(2, 3) = Point.Z;
+		return Matrix;
+	}
+
+	FMatrixH FMatrixH::MakeScaleMatrix(double XScale, double YScale, double ZScale)
+	{
+		FMatrixH Matrix;
+		Matrix.SetIdentity();
+		Matrix.Get(0, 0) = XScale;
+		Matrix.Get(1, 1) = YScale;
+		Matrix.Get(2, 2) = ZScale;
+		return Matrix;
+	}
+
+	FMatrixH FMatrixH::operator*(const FMatrixH& InMatrix) const
+	{
+		FMatrixH Result;
+
+		for (int32 Index = 0; Index < 4; Index++)
+		{
+			for (int32 Jndex = 0; Jndex < 4; Jndex++)
+			{
+				Result.Get(Index, Jndex) = 0;
+				for (int32 Kndex = 0; Kndex < 4; Kndex++)
+				{
+					Result.Get(Index, Jndex) = Result.Get(Index, Jndex) + (Get(Index, Kndex) * InMatrix.Get(Kndex, Jndex));
+				}
+			}
+		}
+		return Result;
+	}
+
+	FMatrixH FMatrixH::operator+(const FMatrixH& InMatrix) const
+	{
+		FMatrixH Result;
+		for (int32 i = 0; i < 16; i++)
+		{
+			Result.Matrix[i] = Matrix[i] + InMatrix.Matrix[i];
+		}
+		return Result;
+	}
+
+	FPoint FMatrixH::PointRotation(const FPoint& PointToRotate, const FPoint& Origin) const
+	{
+		FPoint Result = Origin;
+		for (int32 Index = 0; Index < 3; Index++)
+		{
+			for (int32 Jndex = 0; Jndex < 3; Jndex++)
+			{
+				Result[Index] += Get(Index, Jndex) * (PointToRotate[Jndex] - Origin[Jndex]);
+			}
+		}
+		return Result;
+	}
+
+	FPoint2D FMatrixH::PointRotation(const FPoint2D& PointToRotate, const FPoint2D& Origin) const
+	{
+		FPoint2D Result = Origin;
+		for (int32 Index = 0; Index < 2; Index++)
+		{
+			for (int32 Jndex = 0; Jndex < 2; Jndex++)
+			{
+				Result[Index] += Get(Index, Jndex) * (PointToRotate[Jndex] - Origin[Jndex]);
+			}
+		}
+		return Result;
+	}
+
+	FPoint FMatrixH::Column(int32 Index) const
+	{
+		return FPoint(Get(0, Index), Get(1, Index), Get(2, Index));
+	}
+
+	FPoint FMatrixH::Row(int32 Index) const
+	{
+		return FPoint(Get(Index, 0), Get(Index, 1), Get(Index, 2));
+	}
+
+	void FMatrixH::Inverse()
+	{
+		InverseMatrixN(Matrix, 4);
+	}
+
+	void FMatrixH::Transpose()
+	{
+		FMatrixH Tmp = *this;
+		for (int32 Index = 0; Index < 4; Index++)
+		{
+			for (int32 Jndex = 0; Jndex < 4; Jndex++)
+			{
+				Get(Index, Jndex) = Tmp.Get(Jndex, Index);
+			}
+		}
+	}
+
+	void FMatrixH::GetMatrixDouble(double* OutMatrix) const
+	{
+		memcpy(OutMatrix, Matrix, 16 * sizeof(double));
+	}
+
+	void FMatrixH::CrossProduct(const FPoint& vec)
+	{
+		SetIdentity();
+		Get(1, 0) = -vec.Z;
+		Get(2, 0) = vec.Y;
+		Get(0, 1) = vec.Z;
+		Get(2, 1) = -vec.X;
+		Get(0, 2) = -vec.Y;
+		Get(1, 2) = vec.X;
+		Get(3, 3) = 1.0;
+	}
+
+	void FMatrixH::Print(EVerboseLevel level) const
+	{
+		FMessage::Printf(level, TEXT(" - Matrix\n"));
+		for (int32 Row = 0; Row < 4; Row++)
+		{
+			FMessage::Printf(level, TEXT("	- "));
+			for (int32 Column = 0; Column < 4; Column++)
+			{
+				FMessage::Printf(level, TEXT("%f "), Get(Row, Column));
+			}
+			FMessage::Printf(level, TEXT("\n"));
+		}
+	}
+
+	void InverseMatrixN(double* Matrice, int32 Rank)
+	{
+		const double One = 1.0;
+		const double Zero = 0.0;
+
+		TArray<double> TempMatrix;
+		TempMatrix.Append(Matrice, Rank * Rank);
+
+		double Determinant = One;
+
+		TArray<int32> ColumnToRow;
+		ColumnToRow.SetNum(Rank);
+		for (int32 Index = 0; Index < Rank; Index++)
+		{
+			ColumnToRow[Index] = Index;
+		}
+
+		for (int32 Column = 0; Column < Rank; ++Column)
+		{
+			double Pivot = 0;
+
+			int32 Row = Column;
+			while (Row < Rank)
+			{
+				Pivot = TempMatrix[Rank * Row + Column];
+				if (!FMath::IsNearlyZero(Pivot))
+				{
+					break;
+				}
+				Row++;
+			}
+
+			Determinant = Determinant * Pivot;
+			if (Row != Column)
+			{
+				Swap(ColumnToRow[Column], ColumnToRow[Row]);
+				for (int32 Index = 0; Index < Rank; Index++)
+				{
+					Swap(TempMatrix[Row * Rank + Index], TempMatrix[Column * Rank + Index]);
+				}
+				Determinant = -Determinant;
+			}
+
+			double InvPivot = One / Pivot;
+			TempMatrix[Column * Rank + Column] = One;
+
+			for (int32 Index = 0; Index < Rank; Index++)
+			{
+				TempMatrix[Column * Rank + Index] = TempMatrix[Column * Rank + Index] * InvPivot;
+			}
+
+			for (Row = 0; Row < Rank; Row++)
+			{
+				if (Row == Column)
+				{
+					continue;
+				}
+
+				double ValueRC = TempMatrix[Row * Rank + Column];
+
+				TempMatrix[Row * Rank + Column] = Zero;
+				for (int32 Index = 0; Index < Rank; Index++)
+				{
+					TempMatrix[Row * Rank + Index] = TempMatrix[Row * Rank + Index] - ValueRC * TempMatrix[Column * Rank + Index];
+				}
+			}
+		}
+
+		for (int32 Column = 0; Column < Rank; Column++)
+		{
+			int32 Row = Column;
+			while (Column < Rank)
+			{
+				if (ColumnToRow[Row] == Column)
+				{
+					break;
+				}
+				Row++;
+			}
+
+			if (Column == Row)
+			{
+				continue;
+			}
+
+			ColumnToRow[Row] = ColumnToRow[Column];
+			for (int32 Index = 0; Index < Rank; Index++)
+			{
+				Swap(TempMatrix[Index * Rank + Column], TempMatrix[Index * Rank + Row]);
+			}
+		}
+
+		memcpy(Matrice, TempMatrix.GetData(), Rank * Rank * sizeof(double));
+	}
+
+	void MatrixProduct(int32 ARowNum, int32 AColumnNum, int32 ResultRank, const double* MatrixA, const double* MatrixB, double* MatrixResult)
+	{
+		for (int32 RowA = 0; RowA < ARowNum; RowA++)
+		{
+			for (int32 ColumnB = 0; ColumnB < ResultRank; ColumnB++)
+			{
+				int32 ResultIndex = RowA * ResultRank + ColumnB;
+
+				MatrixResult[ResultIndex] = 0.0;
+				for (int32 k = 0; k < AColumnNum; k++)
+				{
+					MatrixResult[ResultIndex] = MatrixResult[ResultIndex] + MatrixA[RowA * AColumnNum + k] * MatrixB[k * ResultRank + ColumnB];
+				}
+			}
+		}
+	}
+
+	void TransposeMatrix(int32 RowNum, int32 ColumnNum, const double* InMatrix, double* OutMatrix)
+	{
+		for (int32 Row = 0; Row < RowNum; ++Row)
+		{
+			for (int32 Column = 0; Column < ColumnNum; ++Column)
+			{
+				OutMatrix[RowNum * Column + Row] = InMatrix[ColumnNum * Row + Column];
+			}
+		}
+	}
+} // namespace CADKernel
diff --git a/Engine/Source/Runtime/Datasmith/CADKernel/Private/CADKernel/Math/Point.cpp b/Engine/Source/Runtime/Datasmith/CADKernel/Private/CADKernel/Math/Point.cpp
new file mode 100644
--- /dev/null
+++ b/Engine/Source/Runtime/Datasmith/CADKernel/Private/CADKernel/Math/Point.cpp
@@ -0,0 +1,70 @@
+// Copyright Epic Games, Inc. All Rights Reserved.
+
+#include "CADKernel/Math/Point.h"
+
+#include "CADKernel/Utils/Util.h"
+#include "CADKernel/Math/MatrixH.h"
+
+using namespace CADKernel;
+
+const FPoint FPoint::ZeroPoint(0.,0.,0.);
+const FFPoint FFPoint::ZeroPoint(0.f, 0.f, 0.f);
+const FPoint2D FPoint2D::ZeroPoint(0., 0.);
+const FPointH FPointH::ZeroPoint(0., 0., 0., 1.);
+
+double FPoint::SignedAngle(const FPoint & Other, const FPoint & Normal) const
+{
+	FPoint Vector1 = *this; 
+	FPoint Vector2 = Other; 
+	FPoint Vector3 = Normal; 
+
+	Vector1.Normalize();
+	Vector2.Normalize();
+	Vector3.Normalize();
+
+	double ScalarProduct = Vector1 * Vector2;
+
+	if (ScalarProduct >= 1 - SMALL_NUMBER)
+	{
+		return 0.;
+	}
+
+	if (ScalarProduct <= -1 + SMALL_NUMBER)
+	{
+		return PI;
+	}
+
+	return MixedTripleProduct(Vector1, Vector2, Vector3) > 0 ? acos(ScalarProduct) : -acos(ScalarProduct);
+}
+
+double FPoint::ComputeCosinus(const FPoint& OtherVector) const
+{
+	FPoint ThisNormalized = *this;
+	FPoint OtherNormalized = OtherVector;
+
+	ThisNormalized.Normalize();
+	OtherNormalized.Normalize();
+
+	double Cosinus = ThisNormalized * OtherNormalized;
+
+	return FMath::Max(-1.0, FMath::Min(Cosinus, 1.0));
+}
+
+double FPoint::ComputeSinus(const FPoint& OtherVector) const
+{
+	FPoint ThisNormalized = *this;
+	FPoint OtherNormalized = OtherVector;
+
+	ThisNormalized.Normalize();
+	OtherNormalized.Normalize();
+
+	FPoint SinusPoint = ThisNormalized ^ OtherNormalized;
+	double Sinus = SinusPoint.Length();
+	return FMath::Min(Sinus, 1.0);
+}
+
+double FPoint::ComputeAngle(const FPoint& OtherVector) const
+{
+	double CosAngle = ComputeCosinus(OtherVector);
+	return acos(CosAngle);
+}
diff --git a/Engine/Source/Runtime/Datasmith/CADKernel/Private/CADKernel/Mesh/Criteria/AngleCriterion.cpp b/Engine/Source/Runtime/Datasmith/CADKernel/Private/CADKernel/Mesh/Criteria/AngleCriterion.cpp
new file mode 100644
--- /dev/null
+++ b/Engine/Source/Runtime/Datasmith/CADKernel/Private/CADKernel/Mesh/Criteria/AngleCriterion.cpp
@@ -0,0 +1,24 @@
+// Copyright Epic Games, Inc. All Rights Reserved.
+#include "CADKernel/Mesh/Criteria/AngleCriterion.h"
+
+using namespace CADKernel;
+
+FAngleCriterion::FAngleCriterion(double DegreeAngle) 
+	: FCriterion(ECriterion::Angle)
+{
+	AngleCriterionValue = FMath::DegreesToRadians(DegreeAngle);
+	SinMaxAngle = sin(AngleCriterionValue*0.5);
+}
+
+FAngleCriterion::FAngleCriterion(FCADKernelArchive& Archive, ECriterion InCriterionType)
+	: FCriterion(InCriterionType)
+{
+	Serialize(Archive);
+}
+
+void FAngleCriterion::Serialize(FCADKernelArchive& Ar)
+{
+	FCriterion::Serialize(Ar);
+	Ar << AngleCriterionValue;
+	Ar << SinMaxAngle;
+}
diff --git a/Engine/Source/Runtime/Datasmith/CADKernel/Private/CADKernel/Mesh/Criteria/CriteriaGrid.cpp b/Engine/Source/Runtime/Datasmith/CADKernel/Private/CADKernel/Mesh/Criteria/CriteriaGrid.cpp
new file mode 100644
--- /dev/null
+++ b/Engine/Source/Runtime/Datasmith/CADKernel/Private/CADKernel/Mesh/Criteria/CriteriaGrid.cpp
@@ -0,0 +1,203 @@
+// Copyright Epic Games, Inc. All Rights Reserved.
+#include "CADKernel/Mesh/Criteria/CriteriaGrid.h"
+
+#include "CADKernel/Mesh/Criteria/Criterion.h"
+
+#include "CADKernel/Topo/TopologicalFace.h"
+#include "CADKernel/Geo/Surfaces/Surface.h"
+#include "CADKernel/UI/Display.h"
+#include "CADKernel/Mesh/Meshers/IsoTriangulator/DefineForDebug.h"
+
+using namespace CADKernel;
+
+
+const FPoint& FCriteriaGrid::GetPoint(int32 UIndex, int32 VIndex, bool bIsInternalU, bool bIsInternalV, FPoint* OutNormal) const
+{
+	int32 Index = GetIndex(UIndex, VIndex, bIsInternalU, bIsInternalV);
+	ensureCADKernel(Grid.Points3D.IsValidIndex(Index));
+
+	if (OutNormal)
+	{
+		*OutNormal = Grid.Normals[Index];
+	}
+	return Grid.Points3D[Index];
+}
+
+void FCriteriaGrid::Init()
+{
+	TFunction<void(const TArray<double>&, TArray<double>&)> ComputeMiddlePointsCoordinates = [](const TArray<double>& Tab, TArray<double>& Tab2)
+	{
+		Tab2.SetNum(Tab.Num() * 2 - 1);
+		Tab2[0] = Tab[0];
+		for (int32 Index = 1; Index < Tab.Num(); Index++)
+		{
+			Tab2[2 * Index - 1] = (Tab[Index - 1] + Tab[Index]) / 2.0;
+			Tab2[2 * Index] = Tab[Index];
+		}
+	};
+
+	FCoordinateGrid CoordinateGrid2;
+	ComputeMiddlePointsCoordinates(Surface->GetCrossingPointCoordinates(EIso::IsoU), CoordinateGrid2[EIso::IsoU]);
+	ComputeMiddlePointsCoordinates(Surface->GetCrossingPointCoordinates(EIso::IsoV), CoordinateGrid2[EIso::IsoV]);
+
+	Surface->GetCarrierSurface()->EvaluatePointGrid(CoordinateGrid2, Grid, true);
+	TrueUcoorindateCount = CoordinateGrid2[EIso::IsoU].Num();
+}
+
+FCriteriaGrid::FCriteriaGrid(TSharedRef<FTopologicalFace> InSurface)
+	: Surface(InSurface)
+	, CoordinateGrid(InSurface->GetCrossingPointCoordinates())
+{
+	Surface->Presample();
+	Surface->InitDeltaUs();
+	Init();
+#ifdef DISPLAY_CRITERIA_GRID
+	Display();
+#endif
+}
+
+void FCriteriaGrid::ApplyCriteria(const TArray<TSharedPtr<FCriterion>>& Criteria) const
+{
+	TArray<double>& DeltaUMaxArray = Surface->GetCrossingPointDeltaMaxs(EIso::IsoU);
+	TArray<double>& DeltaUMiniArray = Surface->GetCrossingPointDeltaMins(EIso::IsoU);
+	TArray<double>& DeltaVMaxArray = Surface->GetCrossingPointDeltaMaxs(EIso::IsoV);
+	TArray<double>& DeltaVMinArray = Surface->GetCrossingPointDeltaMins(EIso::IsoV);
+	FSurfaceCurvature& SurfaceCurvature = Surface->GetCurvatures();
+
+	for (int32 IndexV = 0; IndexV < GetCoordinateCount(EIso::IsoV) - 1; ++IndexV)
+	{
+		for (int32 IndexU = 0; IndexU < GetCoordinateCount(EIso::IsoU) - 1; ++IndexU)
+		{
+			const FPoint& Point = GetPoint(IndexU, IndexV);
+			const FPoint& PointU = GetPoint(IndexU + 1, IndexV);
+			const FPoint& PointV = GetPoint(IndexU, IndexV + 1);
+			const FPoint& PointUV = GetPoint(IndexU + 1, IndexV + 1);
+			const FPoint& PointUMid = GetIntermediateU(IndexU, IndexV);
+			const FPoint& PointVMid = GetIntermediateV(IndexU, IndexV);
+			const FPoint& PointUVMid = GetIntermediateUV(IndexU, IndexV);
+
+			// Evaluate Sag
+			double LengthU;
+			double SagU = FCriterion::EvaluateSag(Point, PointU, PointUMid, LengthU);
+			double LengthV;
+			double SagV = FCriterion::EvaluateSag(Point, PointV, PointVMid, LengthV);
+			double LengthUV;
+			double SagUV = FCriterion::EvaluateSag(Point, PointUV, PointUVMid, LengthUV);
+
+			for (const TSharedPtr< FCriterion> Criterion : Criteria)
+			{
+				Criterion->UpdateDelta((GetCoordinate(EIso::IsoU, IndexU + 1) - GetCoordinate(EIso::IsoU, IndexU)), SagU, SagUV, SagV, LengthU, LengthUV, DeltaUMaxArray[IndexU], DeltaUMiniArray[IndexU], SurfaceCurvature[EIso::IsoU]);
+				Criterion->UpdateDelta((GetCoordinate(EIso::IsoV, IndexV + 1) - GetCoordinate(EIso::IsoV, IndexV)), SagV, SagUV, SagU, LengthV, LengthUV, DeltaVMaxArray[IndexV], DeltaVMinArray[IndexV], SurfaceCurvature[EIso::IsoV]);
+			}
+		}
+	}
+
+	// Delta of the extremities are smooth to avoid big disparity 
+	if (DeltaUMaxArray.Num() > 2)
+	{
+		DeltaUMaxArray[0] = (DeltaUMaxArray[0] + DeltaUMaxArray[1] * 2) * AThird;
+		DeltaUMaxArray.Last() = (DeltaUMaxArray.Last() + DeltaUMaxArray[DeltaUMaxArray.Num() - 2] * 2) * AThird;
+	}
+
+	if (DeltaVMaxArray.Num() > 2)
+	{
+		DeltaVMaxArray[0] = (DeltaVMaxArray[0] + DeltaVMaxArray[1] * 2) * AThird;
+		DeltaVMaxArray.Last() = (DeltaVMaxArray.Last() + DeltaVMaxArray[DeltaVMaxArray.Num() - 2] * 2) * AThird;
+	}
+}
+
+void FCriteriaGrid::Display()
+{
+	Open3DDebugSession(TEXT("Grid"));
+
+	Open3DDebugSession(TEXT("CriteriaGrid Point 3d"));
+	for (int32 VIndex = 0; VIndex < GetCoordinateCount(EIso::IsoV); ++VIndex)
+	{
+		for (int32 UIndex = 0; UIndex < GetCoordinateCount(EIso::IsoU); ++UIndex)
+		{
+			CADKernel::DisplayPoint(GetPoint(UIndex, VIndex), UIndex);
+		}
+	}
+	Close3DDebugSession();
+
+	Open3DDebugSession(TEXT("CriteriaGrid IntermediateU"));
+	for (int32 VIndex = 0; VIndex < GetCoordinateCount(EIso::IsoV); ++VIndex)
+	{
+		for (int32 UIndex = 0; UIndex < GetCoordinateCount(EIso::IsoU) - 1; ++UIndex)
+		{
+			CADKernel::DisplayPoint(GetIntermediateU(UIndex, VIndex), EVisuProperty::ControlPoint);
+		}
+	}
+	Close3DDebugSession();
+
+	Open3DDebugSession(TEXT("CriteriaGrid IntermediateV"));
+	for (int32 VIndex = 0; VIndex < GetCoordinateCount(EIso::IsoV) - 1; ++VIndex)
+	{
+		for (int32 UIndex = 0; UIndex < GetCoordinateCount(EIso::IsoU); ++UIndex)
+		{
+			CADKernel::DisplayPoint(GetIntermediateV(UIndex, VIndex), EVisuProperty::ControlPoint);
+		}
+	}
+	Close3DDebugSession();
+
+	Open3DDebugSession(TEXT("CriteriaGrid IntermediateUV"));
+	for (int32 VIndex = 0; VIndex < GetCoordinateCount(EIso::IsoV) - 1; ++VIndex)
+	{
+		for (int32 UIndex = 0; UIndex < GetCoordinateCount(EIso::IsoU) - 1; ++UIndex)
+		{
+			CADKernel::DisplayPoint(GetIntermediateUV(UIndex, VIndex), EVisuProperty::ControlPoint);
+		}
+	}
+	Close3DDebugSession();
+
+	Open3DDebugSession(TEXT("Loop 3D"));
+	for (const TSharedPtr<FTopologicalLoop> Loop : Surface->GetLoops())
+	{
+		::Display(Loop);
+	}
+	Close3DDebugSession();
+
+	Open3DDebugSession(TEXT("Loop 2D"));
+	for (const TSharedPtr<FTopologicalLoop> Loop : Surface->GetLoops())
+	{
+		::Display2D(Loop);
+	}
+	Close3DDebugSession();
+
+	Open3DDebugSession(TEXT("CriteriaGrid Point 2D"));
+	for (int32 VIndex = 0; VIndex < GetCoordinateCount(EIso::IsoV); ++VIndex)
+	{
+		for (int32 UIndex = 0; UIndex < GetCoordinateCount(EIso::IsoU); ++UIndex)
+		{
+			CADKernel::DisplayPoint(FPoint2D(GetCoordinate(EIso::IsoU, UIndex), GetCoordinate(EIso::IsoV, VIndex)));
+		}
+	}
+	Close3DDebugSession();
+
+	Open3DDebugSession(TEXT("CriteriaGrid Point 2D Intermediate"));
+	for (int32 VIndex = 0; VIndex < GetCoordinateCount(EIso::IsoV); ++VIndex)
+	{
+		for (int32 UIndex = 1; UIndex < GetCoordinateCount(EIso::IsoU); ++UIndex)
+		{
+			CADKernel::DisplayPoint(FPoint2D((GetCoordinate(EIso::IsoU, UIndex) + GetCoordinate(EIso::IsoU, UIndex - 1))*0.5, GetCoordinate(EIso::IsoV, VIndex)), EVisuProperty::ControlPoint);
+		}
+	}
+	for (int32 VIndex = 1; VIndex < GetCoordinateCount(EIso::IsoV); ++VIndex)
+	{
+		for (int32 UIndex = 0; UIndex < GetCoordinateCount(EIso::IsoU); ++UIndex)
+		{
+			CADKernel::DisplayPoint(FPoint2D(GetCoordinate(EIso::IsoU, UIndex), (GetCoordinate(EIso::IsoV, VIndex) + GetCoordinate(EIso::IsoV, VIndex - 1))*0.5), EVisuProperty::ControlPoint);
+		}
+	}
+	for (int32 VIndex = 1; VIndex < GetCoordinateCount(EIso::IsoV); ++VIndex)
+	{
+		for (int32 UIndex = 1; UIndex < GetCoordinateCount(EIso::IsoU); ++UIndex)
+		{
+			CADKernel::DisplayPoint(FPoint2D((GetCoordinate(EIso::IsoU, UIndex) + GetCoordinate(EIso::IsoU, UIndex - 1))*0.5, (GetCoordinate(EIso::IsoV, VIndex) + GetCoordinate(EIso::IsoV, VIndex - 1))*0.5), EVisuProperty::ControlPoint);
+		}
+	}
+	Close3DDebugSession();
+
+	Close3DDebugSession();
+}
+
diff --git a/Engine/Source/Runtime/Datasmith/CADKernel/Private/CADKernel/Mesh/Criteria/CriteriaGrid.h b/Engine/Source/Runtime/Datasmith/CADKernel/Private/CADKernel/Mesh/Criteria/CriteriaGrid.h
new file mode 100644
--- /dev/null
+++ b/Engine/Source/Runtime/Datasmith/CADKernel/Private/CADKernel/Mesh/Criteria/CriteriaGrid.h
@@ -0,0 +1,99 @@
+// Copyright Epic Games, Inc. All Rights Reserved.
+#pragma once
+
+#include "CADKernel/Core/Types.h"
+#include "CADKernel/Geo/GeoPoint.h"
+#include "CADKernel/Geo/Sampling/SurfacicSampling.h"
+#include "CADKernel/Math/Point.h"
+
+namespace CADKernel
+{
+	class FTopologicalFace;
+	class FCriterion;
+
+	class FCriteriaGrid
+	{
+	protected:
+		TSharedRef<FTopologicalFace> Surface;
+
+		const FCoordinateGrid& CoordinateGrid;
+		int32 TrueUcoorindateCount; 
+		FSurfacicSampling Grid;
+
+		int32 GetIndex(int32 UIndex, int32 VIndex, bool bIsInternalU, bool bIsInternalV) const
+		{
+			int32 TrueUIndex = UIndex * 2;
+			int32 TrueVIndex = VIndex * 2;
+			if (bIsInternalU)
+			{
+				++TrueUIndex;
+			}
+			if (bIsInternalV)
+			{
+				++TrueVIndex;
+			}
+			return TrueVIndex * TrueUcoorindateCount + TrueUIndex;
+		}
+
+		const FPoint& GetPoint(int32 UIndex, int32 VIndex, bool bIsInterU, bool bIsInterV, FPoint* normal = nullptr) const;
+		void Init();
+
+	public:
+		FCriteriaGrid(TSharedRef<FTopologicalFace> Surface);
+
+		TSharedRef<FTopologicalFace> GetSurface()
+		{
+			return Surface;
+		}
+
+		void ApplyCriteria(const TArray<TSharedPtr<FCriterion>>& GetCriteria) const;
+
+		double GetCoordinate(EIso Iso, int32 ind) const
+		{
+			return CoordinateGrid[Iso][ind];
+		}
+
+		constexpr const TArray<double>& GetCoordinates(EIso Iso) const
+		{
+			return CoordinateGrid[Iso];
+		}
+
+		int32 GetCoordinateCount(EIso Iso) const
+		{
+			return CoordinateGrid.IsoCount(Iso);
+		}
+
+		const TArray<FPoint>& GetGridPoints() const
+		{
+			return Grid.Points3D;
+		}
+
+		const TArray<FPoint>& GetGridNormals() const
+		{
+			return Grid.Normals;
+		}
+
+		const FPoint& GetPoint(int32 iU, int32 iV, FPoint* normal = nullptr) const
+		{
+			return GetPoint(iU, iV, false, false, normal);
+		}
+
+		const FPoint& GetIntermediateU(int32 iU, int32 iV, FPoint* normal = nullptr) const
+		{
+			return GetPoint(iU, iV, true, false, normal);
+		}
+
+		const FPoint& GetIntermediateV(int32 iU, int32 iV, FPoint* normal = nullptr) const
+		{
+			return GetPoint(iU, iV, false, true, normal);
+		}
+
+		const FPoint& GetIntermediateUV(int32 iU, int32 iV, FPoint* normal = nullptr) const
+		{
+			return GetPoint(iU, iV, true, true, normal);
+		}
+
+		void Display();
+	};
+}
+
diff --git a/Engine/Source/Runtime/Datasmith/CADKernel/Private/CADKernel/Mesh/Criteria/Criterion.cpp b/Engine/Source/Runtime/Datasmith/CADKernel/Private/CADKernel/Mesh/Criteria/Criterion.cpp
new file mode 100644
--- /dev/null
+++ b/Engine/Source/Runtime/Datasmith/CADKernel/Private/CADKernel/Mesh/Criteria/Criterion.cpp
@@ -0,0 +1,161 @@
+// Copyright Epic Games, Inc. All Rights Reserved.
+#include "CADKernel/Mesh/Criteria/Criterion.h"
+
+#include "CADKernel/Mesh/Criteria/AngleCriterion.h"
+#include "CADKernel/Mesh/Criteria/CriteriaGrid.h"
+#include "CADKernel/Mesh/Criteria/CurvatureCriterion.h"
+#include "CADKernel/Mesh/Criteria/SagCriterion.h"
+#include "CADKernel/Mesh/Criteria/SizeCriterion.h"
+
+#include "CADKernel/Core/KernelParameters.h"
+#include "CADKernel/Geo/GeoPoint.h"
+#include "CADKernel/Topo/TopologicalEdge.h"
+
+namespace CADKernel
+{
+
+	void FCriterion::ApplyOnEdgeParameters(TSharedRef<FTopologicalEdge> Edge, const TArray<double>& Coordinates, const TArray<FCurvePoint>& Points) const
+	{
+		TArray<double>& DeltaUMaxs = Edge->GetDeltaUMaxs();
+
+		double NumericPrecision = Edge->GetTolerance3D();
+		NumericPrecision /= 10;
+
+		if (Edge->Length() <= NumericPrecision)
+		{
+			return;
+		}
+
+		int32 PreviousIndex = 0;
+		for (int32 Index = 1; Index < Coordinates.Num(); Index++)
+		{
+			int32 MiddleIndex = PreviousIndex + Index;
+			double ChordLength;
+			double Sag = EvaluateSag(Points[2 * PreviousIndex].Point, Points[2 * Index].Point, Points[MiddleIndex].Point, ChordLength);
+
+			double AiAm = Points[2 * PreviousIndex].Point.Distance(Points[MiddleIndex].Point);
+			double AmAi2 = Points[2 * Index].Point.Distance(Points[MiddleIndex].Point);
+			double AiAi2 = Points[2 * PreviousIndex].Point.Distance(Points[2 * Index].Point);
+
+			double DeltaU = Coordinates[Index] - Coordinates[PreviousIndex];
+
+			if (ChordLength > NumericPrecision && Sag > SMALL_NUMBER)
+			{
+				double DeltaUMax = ComputeDeltaU(ChordLength, DeltaU, Sag);
+				if (DeltaUMax < DeltaUMaxs[PreviousIndex]) DeltaUMaxs[PreviousIndex] = DeltaUMax;
+			}
+
+			Sag = EvaluateSag(Points[2 * Index].Point, Points[2 * PreviousIndex].Point, Points[MiddleIndex].Point, ChordLength);
+			if (ChordLength > NumericPrecision && Sag > SMALL_NUMBER)
+			{
+				double DeltaUMax = ComputeDeltaU(ChordLength, DeltaU, Sag);
+				if (DeltaUMax < DeltaUMaxs[PreviousIndex]) 
+				{
+					DeltaUMaxs[PreviousIndex] = DeltaUMax;
+				}
+			}
+			PreviousIndex = Index;
+		}
+	}
+
+	TSharedPtr<FCriterion> FCriterion::Deserialize(FCADKernelArchive& Archive)
+	{
+		ECriterion CriterionType = ECriterion::None;
+		Archive << CriterionType;
+
+		switch (CriterionType)
+		{
+		case ECriterion::Size:
+		case ECriterion::MinSize:
+		case ECriterion::MaxSize:
+			return FEntity::MakeShared<FSizeCriterion>(Archive, CriterionType);
+		case ECriterion::CADCurvature:
+			return FEntity::MakeShared<FCurvatureCriterion>(Archive, CriterionType);
+		case ECriterion::Sag:
+			return FEntity::MakeShared<FSagCriterion>(Archive, CriterionType);
+		case ECriterion::Angle:
+			return FEntity::MakeShared<FAngleCriterion>(Archive, CriterionType);
+		default:
+			return nullptr;
+		}
+	}
+
+	TSharedPtr<FCriterion> FCriterion::CreateCriterion(ECriterion Type, double Value)
+	{
+		switch (Type)
+		{
+		case ECriterion::MinSize:
+		case ECriterion::MaxSize:
+			return FEntity::MakeShared<FSizeCriterion>(Value, Type);
+		case ECriterion::Sag:
+			return FEntity::MakeShared<FSagCriterion>(Value);
+		case ECriterion::Angle:
+			return FEntity::MakeShared<FAngleCriterion>(Value);
+		case ECriterion::CADCurvature:
+			return FEntity::MakeShared<FCurvatureCriterion>();
+		default:
+			ERROR_FUNCTION_CALL_NOT_EXPECTED;
+			return TSharedPtr<FCriterion>();
+		}
+	}
+
+	double FCriterion::DefaultValue(ECriterion Type)
+	{
+		switch (Type)
+		{
+		case ECriterion::MinSize:
+			return FSizeCriterion::DefaultValue(Type);
+		case ECriterion::MaxSize:
+			return FSizeCriterion::DefaultValue(Type);
+		case ECriterion::Sag:
+			return FSagCriterion::DefaultValue();
+		case ECriterion::Angle:
+			return FAngleCriterion::DefaultValue();
+		default:
+			ERROR_FUNCTION_CALL_NOT_EXPECTED;
+			return 0.0;
+		}
+	}
+
+#ifdef CADKERNEL_DEV
+	FInfoEntity& FCriterion::GetInfo(FInfoEntity& Info) const
+	{
+		FEntity::GetInfo(Info)
+			.Add(TEXT("CriterionType"), CriterionTypeNames[(uint8)GetCriterionType()])
+			.Add(TEXT("Value"), Value());
+		return Info;
+	}
+#endif
+
+	const TCHAR* CriterionTypeNames[] = {
+		TEXT("Size"),
+		TEXT("MaxSize"),
+		TEXT("MinSize"),
+		TEXT("Angle"),
+		TEXT("Chord Error"),
+		TEXT("CAD Curvature"),
+		TEXT("")
+	};
+
+	// Defined for Python purpose
+	const char* CriterionTypeConstNames[] = {
+		"CRITERION_SIZE",
+		"CRITERION_MAX_SIZE",
+		"CRITERION_MIN_SIZE",
+		"CRITERION_ANGLE",
+		"CRITERION_SAG",
+		"CRITERION_CAD_CURVATURE",
+		nullptr
+	};
+
+	// Defined for Python purpose
+	const char* CriterionTypeConstDescHelp[] = {
+		": Size of elements (only for check)",
+		": maximum size of elements",
+		": minimum size of elements",
+		": angle between the normal of elements",
+		": distance between elements and CAD surfaces/curves",
+		": CAD curvature (only for check)",
+		nullptr
+	};
+}
\ No newline at end of file
diff --git a/Engine/Source/Runtime/Datasmith/CADKernel/Private/CADKernel/Mesh/Criteria/CurvatureCriterion.h b/Engine/Source/Runtime/Datasmith/CADKernel/Private/CADKernel/Mesh/Criteria/CurvatureCriterion.h
new file mode 100644
--- /dev/null
+++ b/Engine/Source/Runtime/Datasmith/CADKernel/Private/CADKernel/Mesh/Criteria/CurvatureCriterion.h
@@ -0,0 +1,84 @@
+// Copyright Epic Games, Inc. All Rights Reserved.
+#pragma once
+
+#include "CADKernel/Math/Curvature.h"
+#include "CADKernel/Mesh/Criteria/Criterion.h"
+
+namespace CADKernel
+{
+	struct FCurvePoint;
+
+	class FCurvatureCriterion : public FCriterion
+	{
+		friend class FEntity;
+		
+	protected:
+
+		FCurvatureCriterion()
+			: FCriterion(ECriterion::CADCurvature)
+		{
+		}
+
+		FCurvatureCriterion(FCADKernelArchive& Archive, ECriterion InCriterionType)
+			: FCriterion(InCriterionType)
+		{
+			Serialize(Archive);
+		}
+
+	public:
+
+		void Serialize(FCADKernelArchive& Ar)
+		{
+			FCriterion::Serialize(Ar);
+		}
+
+		/**
+		 * Sag & Angle criterion.pdf
+		 * https://docs.google.com/presentation/d/1bUnrRFWCW3sDn9ngb9ftfQS-2JxNJaUZlh783hZMMEw/edit?usp=sharing
+		 */
+		virtual void UpdateDelta(double InDeltaU, double InUSag, double InDiagonalSag, double InVSag, double ChordLength, double DiagonalLength, double& OutSagDeltaUMax, double& OutSagDeltaUMin, FIsoCurvature& IsoCurvature) const override
+		{
+			if (DiagonalLength > SMALL_NUMBER)
+			{
+				double Curvature = 8 * InUSag / FMath::Square(DiagonalLength);
+				if (IsoCurvature.Max < Curvature)
+				{
+					IsoCurvature.Max = Curvature;
+				}
+				if (IsoCurvature.Min > Curvature)
+				{
+					IsoCurvature.Min = Curvature;
+				}
+			}
+		}
+
+		virtual double Value() const override
+		{
+			return 0.;
+		}
+
+		virtual void ApplyOnEdgeParameters(TSharedRef<FTopologicalEdge> Edge, const TArray<double>& Coordinates, const TArray<FCurvePoint>& Points) const override
+		{
+			// Do nothing
+		}
+
+	protected:
+
+		/**
+		 * Sag & Angle criterion.pdf
+		 * https://docs.google.com/presentation/d/1bUnrRFWCW3sDn9ngb9ftfQS-2JxNJaUZlh783hZMMEw/edit?usp=sharing
+		 */
+		virtual bool IsAppliedBetweenBreaks() const override
+		{
+			return true;
+		}
+
+		virtual double ComputeDeltaU(double ChordLength, double DeltaU, double Sag) const override
+		{
+			return 0.0;
+		};
+
+	};
+
+} // namespace CADKernel
+
diff --git a/Engine/Source/Runtime/Datasmith/CADKernel/Private/CADKernel/Mesh/Criteria/SagCriterion.h b/Engine/Source/Runtime/Datasmith/CADKernel/Private/CADKernel/Mesh/Criteria/SagCriterion.h
new file mode 100644
--- /dev/null
+++ b/Engine/Source/Runtime/Datasmith/CADKernel/Private/CADKernel/Mesh/Criteria/SagCriterion.h
@@ -0,0 +1,64 @@
+// Copyright Epic Games, Inc. All Rights Reserved.
+#pragma once
+
+#include "CADKernel/Mesh/Criteria/Criterion.h"
+
+namespace CADKernel
+{
+
+	class FSagCriterion : public FCriterion
+	{
+		friend class FEntity;
+		
+	protected:
+		double MaxSag;
+
+		FSagCriterion(double InSag)
+			: FCriterion(ECriterion::Sag)
+			, MaxSag(InSag)
+		{
+		}
+		
+		FSagCriterion(FCADKernelArchive& Archive, ECriterion InCriterionType)
+			: FCriterion(InCriterionType)
+		{
+			Serialize(Archive);
+		}
+
+	public:
+
+		void Serialize(FCADKernelArchive& Ar)
+		{
+			FCriterion::Serialize(Ar);
+			Ar << MaxSag;
+		}
+
+		virtual double Value() const override
+		{
+			return MaxSag;
+		}
+
+		static double DefaultValue()
+		{
+			return 0.15;
+		}
+
+	protected:
+
+		/**
+		 * Sag & Angle criterion.pdf
+		 * https://docs.google.com/presentation/d/1bUnrRFWCW3sDn9ngb9ftfQS-2JxNJaUZlh783hZMMEw/edit?usp=sharing
+		*/
+		virtual double ComputeDeltaU(double ChordLength, double DeltaU, double Sag) const override
+		{
+			return sqrt(MaxSag / Sag) * DeltaU;
+		}
+
+		virtual bool IsAppliedBetweenBreaks() const override
+		{
+			return true;
+		}
+	};
+
+} // namespace CADKernel
+
diff --git a/Engine/Source/Runtime/Datasmith/CADKernel/Private/CADKernel/Mesh/Criteria/SizeCriterion.cpp b/Engine/Source/Runtime/Datasmith/CADKernel/Private/CADKernel/Mesh/Criteria/SizeCriterion.cpp
new file mode 100644
--- /dev/null
+++ b/Engine/Source/Runtime/Datasmith/CADKernel/Private/CADKernel/Mesh/Criteria/SizeCriterion.cpp
@@ -0,0 +1,106 @@
+// Copyright Epic Games, Inc. All Rights Reserved.
+#include "CADKernel/Mesh/Criteria/SizeCriterion.h"
+
+#include "CADKernel/Core/KernelParameters.h"
+#include "CADKernel/Core/System.h"
+#include "CADKernel/Geo/GeoPoint.h"
+#include "CADKernel/Topo/TopologicalEdge.h"
+
+using namespace CADKernel;
+
+FSizeCriterion::FSizeCriterion(double InSize, ECriterion InType)
+	: FCriterion(InType)
+	, Size(InSize)
+{
+}
+
+FSizeCriterion::FSizeCriterion(FCADKernelArchive& Archive, ECriterion InType)
+	: FCriterion(InType)
+{
+	Serialize(Archive);
+}
+
+
+void FSizeCriterion::ApplyOnEdgeParameters(TSharedRef<FTopologicalEdge> Edge, const TArray<double>& Coordinates, const TArray<FCurvePoint>& Points) const
+{
+	double NumericPrecision = Edge->GetTolerance3D();
+	if (Edge->Length() <= NumericPrecision)
+	{
+		return;
+	}
+
+	switch (CriterionType)
+	{
+	case ECriterion::MaxSize:
+		ApplyOnEdgeParameters(Edge, Coordinates, Points, Edge->GetDeltaUMaxs(), [](double NewValue, double& AbacusValue)
+			{
+				if (NewValue < AbacusValue)
+				{
+					AbacusValue = NewValue;
+				}
+			});
+		break;
+
+	case ECriterion::MinSize:
+		ApplyOnEdgeParameters(Edge, Coordinates, Points, Edge->GetDeltaUMins(), [](double NewValue, double& AbacusValue)
+			{
+				if (NewValue > AbacusValue)
+				{
+					AbacusValue = NewValue;
+				}
+			});
+		break;
+
+	default:
+		break;
+	}
+}
+
+void FSizeCriterion::ApplyOnEdgeParameters(TSharedRef<FTopologicalEdge> Edge, const TArray<double>& Coordinates, const TArray<FCurvePoint>& Points, TArray<double>& DeltaUArray, TFunction<void(double, double&)> Compare) const
+{
+	double DeltaUMax = Coordinates[Coordinates.Num() - 1] - Coordinates[0];
+
+	for (int32 PIndex = 1; PIndex < Coordinates.Num(); PIndex++)
+	{
+		double DeltaU = Coordinates[PIndex] - Coordinates[PIndex - 1];
+		double Length = Points[2 * (PIndex - 1)].Point.Distance(Points[2 * PIndex].Point);
+
+		DeltaU = (Length > 0) ? DeltaU * Size / Length : DeltaUMax;
+		Compare(DeltaU, DeltaUArray[PIndex - 1]);
+	}
+}
+
+void FSizeCriterion::UpdateDelta(double InDeltaU, double InUSag, double InDiagonalSag, double InVSag, double ChordLength, double DiagonalLength, double& OutSagDeltaUMax, double& OutSagDeltaUMin, FIsoCurvature& SurfaceCurvature) const
+{
+	if (ChordLength < KINDA_SMALL_NUMBER)
+	{
+		return;
+	}
+
+	switch (CriterionType)
+	{
+	case ECriterion::MaxSize:
+		{
+			double DeltaU = InDeltaU * Size / ChordLength;
+			if (DeltaU < OutSagDeltaUMax)
+			{
+				OutSagDeltaUMax = DeltaU;
+			}
+		}
+		break;
+
+	case ECriterion::MinSize:
+		{
+			double DeltaU = InDeltaU * Size / ChordLength;
+			if (DeltaU > OutSagDeltaUMin)
+			{
+				OutSagDeltaUMin = DeltaU;
+			}
+		}
+		break;
+
+	default:
+		break;
+	}
+}
+
diff --git a/Engine/Source/Runtime/Datasmith/CADKernel/Private/CADKernel/Mesh/Criteria/SizeCriterion.h b/Engine/Source/Runtime/Datasmith/CADKernel/Private/CADKernel/Mesh/Criteria/SizeCriterion.h
new file mode 100644
--- /dev/null
+++ b/Engine/Source/Runtime/Datasmith/CADKernel/Private/CADKernel/Mesh/Criteria/SizeCriterion.h
@@ -0,0 +1,58 @@
+// Copyright Epic Games, Inc. All Rights Reserved.
+#pragma once
+
+#include "CADKernel/Mesh/Criteria/Criterion.h"
+
+/**
+ * Sag & Angle criterion.pdf
+ * https://docs.google.com/presentation/d/1bUnrRFWCW3sDn9ngb9ftfQS-2JxNJaUZlh783hZMMEw/edit?usp=sharing
+*/
+
+namespace CADKernel
+{
+	struct FCurvePoint;
+
+	class FSizeCriterion : public FCriterion
+	{
+		friend class FEntity;
+	
+	protected:
+		double Size;
+
+		FSizeCriterion(double Size, ECriterion Type);
+		FSizeCriterion(FCADKernelArchive& Archive, ECriterion InCriterionType);
+
+	public:
+
+		void Serialize(FCADKernelArchive& Ar)
+		{
+			FCriterion::Serialize(Ar);
+			Ar << Size;
+		}
+
+		void ApplyOnEdgeParameters(TSharedRef<FTopologicalEdge> Edge, const TArray<double>& Coordinates, const TArray<FCurvePoint>& Points) const override;
+
+		double Value() const override
+		{
+			return Size;
+		}
+
+		static double DefaultValue(ECriterion Type)
+		{
+			switch (Type)
+			{
+			case ECriterion::MinSize:
+				return 0.1;
+			case ECriterion::MaxSize:
+				return 30;
+			}
+			return 0;
+		}
+
+		void UpdateDelta(double InDeltaU, double InUSag, double InDiagonalSag, double InVSag, double ChordLength, double DiagonalLength, double& OutSagDeltaUMax, double& OutSagDeltaUMin, FIsoCurvature& SurfaceCurvature) const override;
+
+		void ApplyOnEdgeParameters(TSharedRef<FTopologicalEdge> Edge, const TArray<double>& TabU, const TArray<FCurvePoint>& tabPt, TArray<double>& tabDeltaU, TFunction<void(double, double&)> Compare) const;
+	};
+
+} // namespace CADKernel
+
