CADKernel Source Code v0.01

diff --git a/Engine/Source/Runtime/Datasmith/CADKernel/Private/CADKernel/Mesh/Meshers/IsoTriangulator.cpp b/Engine/Source/Runtime/Datasmith/CADKernel/Private/CADKernel/Mesh/Meshers/IsoTriangulator.cpp
new file mode 100644
--- /dev/null
+++ b/Engine/Source/Runtime/Datasmith/CADKernel/Private/CADKernel/Mesh/Meshers/IsoTriangulator.cpp
@@ -0,0 +1,3059 @@
+// Copyright Epic Games, Inc. All Rights Reserved.
+
+#include "CADKernel/Mesh/Meshers/IsoTriangulator.h"
+
+#include "CADKernel/Core/Chrono.h"
+#include "CADKernel/Math/MathConst.h"
+#include "CADKernel/Mesh/Meshers/IsoTriangulator/IsoSegment.h"
+#include "CADKernel/Mesh/Structure/Grid.h"
+#include "CADKernel/Mesh/Structure/EdgeMesh.h"
+#include "CADKernel/Mesh/Structure/FaceMesh.h"
+#include "CADKernel/Topo/TopologicalEdge.h"
+#include "CADKernel/Topo/TopologicalFace.h"
+#include "CADKernel/Topo/TopologicalFace.h"
+#include "CADKernel/UI/Display.h"
+
+using namespace CADKernel;
+
+//#define ADD_TRIANGLE_2D
+#define DEBUG_ISOTRIANGULATOR
+
+FIsoTriangulator::FIsoTriangulator(const FGrid& InGrid, TSharedRef<FFaceMesh> EntityMesh)
+	: Grid(InGrid)
+	, Mesh(EntityMesh)
+	, LoopSegmentsIntersectionTool(InGrid)
+	, InnerSegmentsIntersectionTool(InGrid)
+	, InnerToLoopSegmentsIntersectionTool(InGrid)
+{
+	FinalInnerSegments.Reserve(3 * Grid.InnerNodesCount());
+	IndexOfLowerLeftInnerNodeSurroundingALoop.Reserve(Grid.GetLoopCount());
+
+#ifdef DEBUG_ISOTRIANGULATOR
+	if (Grid.GetFace()->GetId() == SURFACE_TO_DEBUG)
+	{
+		bDisplay = true;
+	}
+#endif
+}
+
+bool FIsoTriangulator::Triangulate()
+{
+	EGridSpace DisplaySpace = EGridSpace::Default2D;
+#ifdef DEBUG_ISOTRIANGULATOR
+	if (bDisplay)
+	{
+		Open3DDebugSession(FString::Printf(TEXT("Triangulate %d"), Grid.GetFace()->GetId()));
+		Wait(Grid.GetFace()->GetId() == SurfaceToDebug);
+	}
+#endif
+
+	FTimePoint StartTime = FChrono::Now();
+
+	// =============================================================================================================
+	// Build the first elements (IsoNodes (i.e. Inner nodes), Loops nodes, and knows segments 
+	// =============================================================================================================
+
+	BuildNodes();
+#ifdef CADKERNEL_DEV
+	DisplayIsoNodes(DisplaySpace);
+#endif
+
+	FillMeshNodes();
+
+	if (!BuildLoopSegments())
+	{
+		FMessage::Printf(Log, TEXT("A loop of the surface %d is in self intersecting. The mesh of this sector is canceled."), Grid.GetFace()->GetId());
+		return false;
+	}
+
+#ifdef CADKERNEL_DEV
+	Display(DisplaySpace, TEXT("FIsoTrianguler::LoopSegments"), LoopSegments, false, false, EVisuProperty::GreenCurve);
+#endif
+
+	BuildThinZoneSegments();
+#ifdef CADKERNEL_DEV
+	Display(DisplaySpace, TEXT("FIsoTrianguler::ThinZoneSegments"), ThinZoneSegments, false);
+	LoopSegmentsIntersectionTool.Display(TEXT("FIsoTrianguler::IntersectionTool Loop"));
+#endif
+
+	BuildInnerSegments();
+#ifdef CADKERNEL_DEV
+	Display(DisplaySpace, TEXT("FIsoTrianguler::FinalInnerSegments"), FinalInnerSegments, false, false, EVisuProperty::BlueCurve);
+	Chronos.TriangulateDuration1 = FChrono::Elapse(StartTime);
+#endif
+
+	// =============================================================================================================
+	// =============================================================================================================
+
+	BuildInnerSegmentsIntersectionTool();
+#ifdef CADKERNEL_DEV
+	InnerSegmentsIntersectionTool.Display(TEXT("FIsoTrianguler::IntersectionTool Inner"));
+#endif
+
+	SortLoopNodes();
+#ifdef CADKERNEL_DEV
+	Chronos.TriangulateDuration2 = FChrono::Elapse(StartTime);
+#endif
+
+	// =============================================================================================================
+	// 	   Find first segments of the final mesh
+	// =============================================================================================================
+
+	// use of SortedLoopNodes but specific order
+	FindIsoSegmentToLinkLoopToLoop();
+#ifdef CADKERNEL_DEV
+	Display(DisplaySpace, TEXT("FIsoTrianguler::Final Iso ToLink Loop To Loops"), FinalToLoops, false);
+	//Wait(bDisplay);
+#endif
+
+#ifdef NEED_TO_CHECK_USEFULNESS
+	CompleteIsoSegmentLoopToLoop();
+	Display(DisplaySpace, TEXT("FIsoTrianguler::Final Iso ToLink Loop To Loops with Complete"), FinalToLoops, false);
+#endif
+
+	FindIsoSegmentToLinkInnerToLoop();
+#ifdef CADKERNEL_DEV
+	Display(DisplaySpace, TEXT("FIsoTrianguler::Final Iso ToLink Inner To Loops"), FinalToLoops, false);
+#endif
+
+	// =============================================================================================================
+	// 	   Find candidate segments for the  mesh and select the best set
+	// =============================================================================================================
+
+	// Delaunay
+	if (Grid.GetLoops3D().Num())
+	{
+		Delaunay();
+#ifdef CADKERNEL_DEV
+		Display(DisplaySpace, TEXT("FIsoTrianguler::Candidate Segments by Delaunay"), CandidateLoopToLoopSegments, false);
+#endif
+	}
+	//Wait(bDisplay);
+	
+#ifdef CADKERNEL_DEV
+	Chronos.TriangulateDuration3 = FChrono::Elapse(StartTime);
+#endif
+
+	// use of SortedloopNodes
+	FindCandidateSegmentsToLinkInnerToLoop();
+#ifdef CADKERNEL_DEV
+	Display(DisplaySpace, TEXT("FIsoTrianguler::Candidate Segments ToLink Inner To Loops"), CandidateInnerToLoopSegments, false);
+#endif
+
+	SelectSegmentInCandidateSegments();
+#ifdef CADKERNEL_DEV
+	Display(DisplaySpace, TEXT("FIsoTrianguler::Final Segment To Loop"), FinalToLoops, false);
+	Display(DisplaySpace, TEXT("FIsoTrianguler::Inner Segments Limit"), FinalInnerSegments, false);
+	//Wait(bDisplay);
+#endif
+
+	// Complete the final segments with segment connecting unconnected Inner point or inner segment 
+	ConnectUnconnectedInnerSegments();
+#ifdef CADKERNEL_DEV
+	Display(DisplaySpace, TEXT("FIsoTrianguler::Final Segment To Loop"), FinalToLoops, false);
+#endif
+
+#ifdef CADKERNEL_DEV
+	Chronos.TriangulateDuration4 = FChrono::Elapse(StartTime);
+#endif
+
+	// =============================================================================================================
+	// Make the final tessellation 
+	// =============================================================================================================
+
+	// Triangulate between inner grid boundary and loops
+	TriangulateOverCycle(EGridSpace::Scaled);
+
+	// Finalise the mesh by the tessellation of the inner grid
+	TriangulateInnerNodes();
+
+#ifdef CADKERNEL_DEV
+	if (bDisplay)
+	{
+		Open3DDebugSession(TEXT("Mesh 3D"));
+		DisplayMesh(Mesh);
+		Close3DDebugSession();
+		Close3DDebugSession();
+	}
+	Wait(bDisplay);
+#endif
+
+#ifdef CADKERNEL_DEV
+	Chronos.PrintTimeElapse();
+#endif
+
+	return true;
+}
+
+void FIsoTriangulator::BuildNodes()
+{
+	FTimePoint StartTime = FChrono::Now();
+
+	int32 NbNodes = 0;
+	for (const TArray<FPoint2D>& Loop : Grid.GetLoops2D(EGridSpace::Default2D))
+	{
+		NbNodes += (int32)Loop.Num();
+	}
+	LoopStartIndex.Reserve(Grid.GetLoops2D(EGridSpace::Default2D).Num());
+	LoopNodes.Reserve((int32)(NbNodes * 1.2 + 5)); // reserve more in case it need to create complementary nodes
+
+	// Loop nodes
+	int32 FaceIndex = 0;
+	int32 LoopIndex = 0;
+	for (const TArray<FPoint2D>& Loop : Grid.GetLoops2D(EGridSpace::Default2D))
+	{
+		LoopStartIndex.Add(LoopNodes.Num());
+		const TArray<int32>& LoopIds = Grid.GetNodeIdsOfFaceLoops()[LoopIndex];
+		FLoopNode* NextNode = nullptr;
+		FLoopNode* FirstNode = &LoopNodes.Emplace_GetRef(LoopIndex, 0, FaceIndex++, LoopIds[0]);
+		FLoopNode* PreviousNode = FirstNode;
+		for (int32 Index = 1; Index < Loop.Num(); ++Index)
+		{
+			NextNode = &LoopNodes.Emplace_GetRef(LoopIndex, Index, FaceIndex++, LoopIds[Index]);
+			PreviousNode->SetNextConnectedNode(NextNode);
+			NextNode->SetPreviousConnectedNode(PreviousNode);
+			PreviousNode = NextNode;
+		}
+		PreviousNode->SetNextConnectedNode(FirstNode);
+		FirstNode->SetPreviousConnectedNode(PreviousNode);
+		LoopIndex++;
+	}
+
+	// Inner node
+	InnerNodes.Reserve(Grid.InnerNodesCount());
+	GlobalIndexToIsoInnerNodes.Init(nullptr, Grid.GetTotalCuttingCount());
+
+	InnerNodeCount = 0;
+	for (int32 Index = 0; Index < (int32)Grid.GetTotalCuttingCount(); ++Index)
+	{
+		if (Grid.IsNodeInsideFace(Index))
+		{
+			FIsoInnerNode& Node = InnerNodes.Emplace_GetRef(Index, FaceIndex++, InnerNodeCount++);
+			GlobalIndexToIsoInnerNodes[Index] = &Node;
+		}
+	}
+
+#ifdef CADKERNEL_DEV
+	Chronos.BuildIsoNodesDuration += FChrono::Elapse(StartTime);
+#endif
+}
+
+void FIsoTriangulator::FillMeshNodes()
+{
+	int32 TriangleNum = 50 + (int32)((2 * InnerNodes.Num() + LoopNodes.Num()) * 1.1);
+	Mesh->Init(TriangleNum, InnerNodes.Num() + LoopNodes.Num());
+
+	TArray<FPoint>& InnerNodeCoordinates = Mesh->GetNodeCoordinates();
+	InnerNodeCoordinates.Reserve(InnerNodeCount);
+	for (int32 Index = 0; Index < (int32)Grid.GetInner3DPoints().Num(); ++Index)
+	{
+		if (Grid.IsNodeInsideFace(Index))
+		{
+			InnerNodeCoordinates.Emplace(Grid.GetInner3DPoints()[Index]);
+		}
+	}
+
+	int32 StartId = Mesh->RegisterCoordinates();
+	for (FIsoInnerNode& Node : InnerNodes)
+	{
+		Node.OffsetId(StartId);
+	}
+
+	Mesh->VerticesGlobalIndex.SetNum(InnerNodes.Num() + LoopNodes.Num());
+	int32 Index = 0;
+	for (FLoopNode& Node : LoopNodes)
+	{
+		Mesh->VerticesGlobalIndex[Index++] = Node.GetId();
+	}
+
+	for (FIsoInnerNode& Node : InnerNodes)
+	{
+		Mesh->VerticesGlobalIndex[Index++] = Node.GetId();
+	}
+
+	for (FLoopNode& Node : LoopNodes)
+	{
+		Mesh->Normals.Emplace(Node.GetNormal(Grid));
+	}
+
+	for (FIsoInnerNode& Node : InnerNodes)
+	{
+		Mesh->Normals.Emplace(Node.GetNormal(Grid));
+	}
+
+	for (FLoopNode& Node : LoopNodes)
+	{
+		Mesh->UVMap.Emplace(Node.Get2DPoint(EGridSpace::Scaled, Grid));
+	}
+
+	for (FIsoInnerNode& Node : InnerNodes)
+	{
+		Mesh->UVMap.Emplace(Node.Get2DPoint(EGridSpace::Scaled, Grid));
+	}
+}
+
+namespace
+{
+	const double MaxSlopeToBeIso = 0.125;
+	constexpr double LimitValueMin(double Slope)
+	{
+		return Slope - MaxSlopeToBeIso;
+	}
+	constexpr double LimitValueMax(double Slope)
+	{
+		return Slope + MaxSlopeToBeIso;
+	}
+}
+
+bool FIsoTriangulator::BuildLoopSegments()
+{
+	FTimePoint StartTime = FChrono::Now();
+
+	int32 LoopNodeCount = LoopNodes.Num();
+	LoopSegments.Reserve(LoopNodeCount);
+
+	int32 LoopIndex = 0;
+	TArray<FIsoSegment*> Segments;
+
+	TFunction<bool()> CheckSelfIntersection = [&]()
+	{
+		FTimePoint StartCheckIntersectionTime = FChrono::Now();
+		FIntersectionSegmentTool IntersectionTool(Grid);
+		IntersectionTool.AddSegment(Segments);
+		IntersectionTool.Sort();
+		for (const FIsoSegment* Segment : Segments)
+		{
+			if (IntersectionTool.DoesIntersect(*Segment))
+			{
+				return true;
+			}
+		}
+		Segments.Empty(LoopNodeCount);
+#ifdef CADKERNEL_DEV
+		Chronos.BuildLoopSegmentsCheckIntersectionDuration += FChrono::Elapse(StartCheckIntersectionTime);
+#endif
+		return false;
+	};
+
+	Segments.Reserve(LoopNodeCount);
+	for (FLoopNode& Node : LoopNodes)
+	{
+		// Check if loops are self intersecting. 
+		if (LoopIndex != Node.GetLoopIndex())
+		{
+			if (CheckSelfIntersection())
+			{
+				return false;
+			}
+			Segments.Empty(LoopNodeCount);
+		}
+
+		FIsoSegment& Segment = IsoSegmentFactory.New();
+		Segment.Init(Node, Node.GetNextNode(), ESegmentType::Loop);
+		Segment.ConnectToNode();
+		LoopSegments.Add(&Segment);
+		Segments.Add(&Segment);
+	}
+
+	// Check last loop
+	if (CheckSelfIntersection())
+	{
+		return false;
+	}
+
+	for (FIsoSegment* Segment : LoopSegments)
+	{
+		double SegmentSlop = ComputeSlope(Segment->GetFirstNode().Get2DPoint(EGridSpace::UniformScaled, Grid), Segment->GetSecondNode().Get2DPoint(EGridSpace::UniformScaled, Grid));
+		if (SegmentSlop < MaxSlopeToBeIso)
+		{
+			Segment->SetAsIsoU();
+		}
+		if (SegmentSlop < LimitValueMax(2) && SegmentSlop > LimitValueMin(2))
+		{
+			Segment->SetAsIsoV();
+		}
+		if (SegmentSlop < LimitValueMax(4) && SegmentSlop > LimitValueMin(4))
+		{
+			Segment->SetAsIsoU();
+		}
+		if (SegmentSlop < LimitValueMax(6) && SegmentSlop > LimitValueMin(6))
+		{
+			Segment->SetAsIsoV();
+		}
+		if (SegmentSlop > LimitValueMin(8))
+		{
+			Segment->SetAsIsoU();
+		}
+	}
+
+	for (FLoopNode& Node : LoopNodes)
+	{
+		if (Node.GetConnectedSegments()[0]->IsIsoU() && Node.GetConnectedSegments()[1]->IsIsoU())
+		{
+			Node.SetAsIsoU();
+		}
+		else if (Node.GetConnectedSegments()[0]->IsIsoV() && Node.GetConnectedSegments()[1]->IsIsoV())
+		{
+			Node.SetAsIsoV();
+		}
+	}
+
+	LoopSegmentsIntersectionTool.AddSegment(LoopSegments);
+	LoopSegmentsIntersectionTool.Sort();
+
+#ifdef CADKERNEL_DEV
+	Chronos.BuildLoopSegmentsDuration += FChrono::Elapse(StartTime);
+#endif
+	return true;
+}
+
+void FIsoTriangulator::BuildThinZoneSegments()
+{
+	ThinZoneSegments.Reserve((int32)(0.6 * (double)LoopNodes.Num()));
+
+	TMap<int32, FLoopNode*> IndexToNode;
+	for (FLoopNode& Node : LoopNodes)
+	{
+		IndexToNode.Add(Node.GetFaceIndex(), &Node);
+	}
+
+	TFunction<void(FIsoNode*, FIsoNode*)> AddSegment = [this](FIsoNode* NodeA, FIsoNode* NodeB)
+	{
+		if (!NodeA)
+		{
+			return;
+		}
+
+		if (!NodeB)
+		{
+			return;
+		}
+
+		if (NodeA->GetSegmentConnectedTo(NodeB))
+		{
+			return;
+		}
+
+		FIsoSegment& Segment = IsoSegmentFactory.New();
+		Segment.Init(*NodeA, *NodeB, ESegmentType::ThinZone);
+		Segment.ConnectToNode();
+		ThinZoneSegments.Add(&Segment);
+	};
+
+	for (const TSharedPtr<FTopologicalLoop> Loop : Grid.GetFace()->GetLoops())
+	{
+		for (const FOrientedEdge& OrientedEdge : Loop->GetEdges())
+		{
+			const TSharedPtr<FTopologicalEdge>& Edge = OrientedEdge.Entity;
+			if (!Edge->IsThinZone())
+			{
+				continue;
+			}
+
+			const TArray<FCuttingPoint>& CuttingPoints = Edge->GetLinkActiveEdge()->GetCuttingPoints();
+			const TArray<int32>& NodeIds = Edge->GetMesh()->EdgeVerticesIndex;
+			for (int32 Index = 0; Index < NodeIds.Num(); ++Index)
+			{
+				if (CuttingPoints[Index].OppositNodeIndex > 0)
+				{
+					AddSegment(IndexToNode[NodeIds[Index]], IndexToNode[CuttingPoints[Index].OppositNodeIndex]);
+				}
+				if (CuttingPoints[Index].OppositNodeIndex2 > 0)
+				{
+					AddSegment(IndexToNode[NodeIds[Index]], IndexToNode[CuttingPoints[Index].OppositNodeIndex2]);
+				}
+			}
+		}
+	}
+
+	LoopSegmentsIntersectionTool.AddSegment(ThinZoneSegments);
+	LoopSegmentsIntersectionTool.Sort();
+}
+
+void FIsoTriangulator::BuildInnerSegments()
+{
+	// TODO this step should not be usefull, only the boundary of the grid is need
+
+	FTimePoint StartTime = FChrono::Now();
+
+	// Build segments according to the Grid following u then following v
+	// Build segment must not be in intersection with the loop
+	int32 NumU = Grid.GetCuttingCount(EIso::IsoU);
+	int32 NumV = Grid.GetCuttingCount(EIso::IsoV);
+
+	LoopSegmentsIntersectionTool.Reserve(InnerSegmentsIntersectionTool.Num());
+
+	TFunction<void(const int32, const int32, const ESegmentType)> BuildSegmentIfValid = [this](const int32 IndexNode1, const int32 IndexNode2, const ESegmentType InType)
+	{
+		if (!Grid.IsNodeInsideFace(IndexNode1))
+		{
+			return;
+		}
+
+		if (!Grid.IsNodeInsideFace(IndexNode2))
+		{
+			return;
+		}
+
+		if (Grid.IsNodeCloseToLoop(IndexNode1) && Grid.IsNodeCloseToLoop(IndexNode2))
+		{
+			if (LoopSegmentsIntersectionTool.DoesIntersect(Grid.GetInner2DPoint(EGridSpace::Default2D, IndexNode1), Grid.GetInner2DPoint(EGridSpace::Default2D, IndexNode2)))
+			{
+				return;
+			}
+		}
+
+		FIsoInnerNode& Node1 = *GlobalIndexToIsoInnerNodes[IndexNode1];
+		FIsoInnerNode& Node2 = *GlobalIndexToIsoInnerNodes[IndexNode2];
+		FIsoSegment& Segment = IsoSegmentFactory.New();
+		Segment.Init(Node1, Node2, InType);
+		Segment.ConnectToNode();
+		FinalInnerSegments.Add(&Segment);
+	};
+
+	for (int32 uIndex = 0; uIndex < NumU; uIndex++)
+	{
+		for (int32 vIndex = 0; vIndex < NumV - 1; vIndex++)
+		{
+			BuildSegmentIfValid(Grid.GobalIndex(uIndex, vIndex), Grid.GobalIndex(uIndex, vIndex + 1), ESegmentType::IsoV);
+		}
+	}
+
+	for (int32 vIndex = 0; vIndex < NumV; vIndex++)
+	{
+		for (int32 uIndex = 0; uIndex < NumU - 1; uIndex++)
+		{
+			BuildSegmentIfValid(Grid.GobalIndex(uIndex, vIndex), Grid.GobalIndex(uIndex + 1, vIndex), ESegmentType::IsoU);
+		}
+	}
+
+#ifdef CADKERNEL_DEV
+	Chronos.BuildInnerSegmentsDuration += FChrono::Elapse(StartTime);
+#endif
+}
+
+void FIsoTriangulator::BuildInnerSegmentsIntersectionTool()
+{
+	FTimePoint StartTime = FChrono::Now();
+
+	// Find Boundary Segments Of Inner Triangulation
+	// 
+	// A pixel grid is build. 
+	// A pixel is the quadrangle of the inner grid
+	// The grid pixel are initialized to False
+	//
+	// A pixel is True if one of its boundary segment does not exist 
+	// The inner of the grid is all pixel False
+	// The boundary of the inner triangulation is defined by all segments adjacent to different cells 
+	// 
+	//    T      T	     T                                                  
+	//       0 ----- 0 												     0 ----- 0 
+	//    T  |   F   |   T       T       T      T         			     |       |    
+	//       0 ----- 0               0 ----- 0 						     0       0               0 ----- 0 
+	//    T  |   F   |   T       T   |   F   |  T					     |       |               |       |  
+	//       0 ----- 0 ----- 0 ----- 0 ----- 0 						     0       0 ----- 0 ----- 0       0 	
+	//    T  |   F   |   F   |   F   |   F   |	T					     |                               |	
+	//       0 ----- 0 ----- 0 ----- 0 ----- 0 						     0                               0 	
+	//    T  |   F   |   F   |   F   |   F   |	T					     |                               |	
+	//       0 ----- 0 ----- 0 ----- 0 ----- 0 						     0 ----- 0 ----- 0 ----- 0 ----- 0 	
+	//    T      T 		 T		 T		 T		T					                
+	// 
+	// https://docs.google.com/presentation/d/1qUVOH-2kU_QXBVKyRUcdDy1Y6WGkcaJCiaS8wGjSZ6M/edit?usp=sharing
+	// Slide "Boundary Segments Of Inner Triangulation"
+
+	int32 NumU = Grid.GetCuttingCount(EIso::IsoU);
+	int32 NumV = Grid.GetCuttingCount(EIso::IsoV);
+
+	TArray<uint8> Pixel;
+	Pixel.Init(0, Grid.GetTotalCuttingCount());
+
+	// A pixel is True if one of its boundary segment does not exist 
+	for (int32 IndexV = 0, Index = 0; IndexV < NumV; ++IndexV)
+	{
+		for (int32 IndexU = 0; IndexU < NumU; ++IndexU, ++Index)
+		{
+			if (!Grid.IsNodeInsideFace(Index))
+			{
+				continue;
+			}
+
+			if (!GlobalIndexToIsoInnerNodes[Index]->IsLinkedToNextU())
+			{
+				Pixel[Index] = true;
+				Pixel[Index - NumU] = true;
+			}
+
+			if (!GlobalIndexToIsoInnerNodes[Index]->IsLinkedToPreviousU())
+			{
+				Pixel[Index - 1] = true;
+				Pixel[Index - 1 - NumU] = true;
+			}
+
+			if (!GlobalIndexToIsoInnerNodes[Index]->IsLinkedToNextV())
+			{
+				Pixel[Index] = true;
+				Pixel[Index - 1] = true;
+			}
+
+			if (!GlobalIndexToIsoInnerNodes[Index]->IsLinkedToPreviousV())
+			{
+				Pixel[Index - NumU] = true;
+				Pixel[Index - NumU - 1] = true;
+			}
+		}
+	}
+
+	// The boundary of the inner triangulation is defined by all segments adjacent to a "True" cell 
+	// These segments are added to InnerSegmentsIntersectionTool
+	InnerSegmentsIntersectionTool.Reserve((int32)FinalInnerSegments.Num());
+
+	for (FIsoSegment* Segment : FinalInnerSegments)
+	{
+		int32 IndexFirstNode = Segment->GetFirstNode().GetIndex();
+		int32 IndexSecondNode = 0;
+		switch (Segment->GetType())
+		{
+		case ESegmentType::IsoU:
+			IndexSecondNode = IndexFirstNode - NumU;
+			break;
+		case ESegmentType::IsoV:
+			IndexSecondNode = IndexFirstNode - 1;
+			break;
+		default:
+			ensureCADKernel(false);
+		}
+		if (Pixel[IndexFirstNode] || Pixel[IndexSecondNode])
+		{
+			InnerSegmentsIntersectionTool.AddSegment(*Segment);
+		}
+	}
+
+	FindInnerGridCellSurroundingSmallLoop();
+
+	// initialize the intersection tool
+	InnerSegmentsIntersectionTool.Sort();
+
+#ifdef CADKERNEL_DEV
+	Chronos.FindLoopSegmentOfInnerTriangulationDuration += FChrono::Elapse(StartTime);
+#endif
+
+#ifdef DEBUG_FINDBOUNDARYSEGMENTS 
+	DisplayPixels(Pixel);
+#endif
+}
+
+void FIsoTriangulator::FindIsoSegmentToLinkLoopToLoop()
+{
+	FTimePoint StartTime = FChrono::Now();
+
+	// This coefficient is to defined the tolerance on coordinates according the iso strip...
+	// With some surface, the parameterization speed can vary enormously depending on the point of the surface.
+	// A good information is the width of iso strip around a point. Indeed, strips have the optimal width to respect meshing criteria.
+	// So a fraction of the strip's width defined a good tolerance around a given point. 
+	constexpr double ToleranceCoefficent = 1. / 12.; // Why 12 ? ;o)
+
+	const TArray<double>& IsoUCoordinates = Grid.GetCuttingCoordinatesAlongIso(EIso::IsoU);
+	const TArray<double>& IsoVCoordinates = Grid.GetCuttingCoordinatesAlongIso(EIso::IsoV);
+
+
+	// Warning, Min delta is compute in EGridSpace::Default2D
+	TFunction<double(const TArray<double>&)> GetMinDelta = [](const TArray<double>& IsoCoordinates)
+	{
+		double MinDelta = HUGE_VAL;
+		for (int32 Index = 0; Index < IsoCoordinates.Num() - 1; ++Index)
+		{
+			double Delta = IsoCoordinates[Index + 1] - IsoCoordinates[Index];
+			if (Delta < MinDelta)
+			{
+				MinDelta = Delta;
+			}
+		}
+		return MinDelta;
+	};
+
+	// Before creating a segment a set of check is done to verify that the segment is valid.
+	TFunction<void(FLoopNode*, const FPoint2D&, FLoopNode*, const FPoint2D&)> CreateSegment = [&](FLoopNode* Node1, const FPoint2D& Coordinate1, FLoopNode* Node2, const FPoint2D& Coordinate2)
+	{
+		if (&Node1->GetPreviousNode() == Node2 || &Node1->GetNextNode() == Node2)
+		{
+			return;
+		}
+
+		if (Node1->GetSegmentConnectedTo(Node2))
+		{
+			return;
+		}
+
+		if (InnerSegmentsIntersectionTool.DoesIntersect(Coordinate1, Coordinate2))
+		{
+			return;
+		}
+
+		if (LoopSegmentsIntersectionTool.DoesIntersect(*Node1, *Node2))
+		{
+			return;
+		}
+
+		// Is Outside and not too flat at Node1
+		const double FlatAngle = 0.1;
+		if (!IsPointPIsInsideSectorABC(Node1->GetPreviousNode().Get2DPoint(EGridSpace::Default2D, Grid), Coordinate1, Node1->GetNextNode().Get2DPoint(EGridSpace::Default2D, Grid), Coordinate2, FlatAngle))
+		{
+			return;
+		}
+
+		// Is Outside and not too flat at Node2
+		if (!IsPointPIsInsideSectorABC(Node2->GetPreviousNode().Get2DPoint(EGridSpace::Default2D, Grid), Coordinate2, Node2->GetNextNode().Get2DPoint(EGridSpace::Default2D, Grid), Coordinate1, FlatAngle))
+		{
+			return;
+		}
+
+		FIsoSegment& Segment = IsoSegmentFactory.New();
+		Segment.Init(*Node1, *Node2, ESegmentType::LoopToLoop);
+		Segment.ConnectToNode();
+		FinalToLoops.Add(&Segment);
+		InnerToLoopSegmentsIntersectionTool.AddSegment(Segment);
+#ifdef DEBUG_FIND_ISOSEGMENT_TO_LINK_LOOP_TO_LOOP
+		Display(EGridSpace::UniformScaled, *Node1, *Node2);
+#endif
+	};
+
+	// Find the index of the closed Iso strip. An iso strip is a the strip [Iso[Index], Iso[Index+1]]
+	// As the process is iterative with sorted points, Index can only be equal or bigger than with the previous node
+	TFunction<void(const TArray<double>&, int32&, const double&)> FindStripIndex = [&](const TArray<double>& Iso, int32& Index, const double& PointCoord)
+	{
+		if (Index > 0)
+		{
+			--Index;
+		}
+
+		// The last strip is not tested as it must be good if the previous are not even if PointCoord >= Iso.Last
+		for (; Index < Iso.Num() - 2; ++Index)
+		{
+			if (PointCoord < Iso[Index + 1])
+			{
+				break;
+			}
+		}
+	};
+
+	TArray<FLoopNode*> SortedLoopNodesAlong = SortedLoopNodes;
+
+	// Find pair of points iso aligned along Axe2
+	// For all loop nodes sorted along Axe1, check if the pair (Node[i], Node[i+1]) is aligned along Axe2. 
+	// The segment (Node[i], Node[i+1]) is valid if its length is smaller or nearly equal to a crossing strip. 
+	// Axe1 == 0 => IsoU, coordinate U is ~constant
+	// Axe1 == 1 => IsoV
+
+	TFunction<void(int32, const TArray<double>&, const TArray<double>&)> FindIsoSegmentAlong = [&](int32 InAxe, const TArray<double>& IsoU, const TArray<double>& IsoV)
+	{
+		EIso ComplementaryAxe = (InAxe + 1) % 2 == 0 ? EIso::IsoU : EIso::IsoV;
+
+		int32 IndexU = 0;
+		for (int32 Index = 0; Index < SortedLoopNodesAlong.Num() - 1; ++Index)
+		{
+			FLoopNode* LoopNode = SortedLoopNodesAlong[Index];
+			if (!LoopNode->IsIso(ComplementaryAxe))
+			{
+				continue;
+			}
+			FLoopNode* NextNode = SortedLoopNodesAlong[Index + 1];
+			if (!NextNode->IsIso(ComplementaryAxe))
+			{
+				continue;
+			}
+
+			const FPoint2D& LoopPoint = LoopNode->Get2DPoint(EGridSpace::Default2D, Grid);
+
+#ifdef DEBUG_FIND_ISOSEGMENT_TO_LINK_LOOP_TO_LOOP
+			F3DDebugSession G(FString::Printf(TEXT("Iso %d Index %d"), InAxe, Index));
+			Display(EGridSpace::UniformScaled, *LoopNode);
+#endif
+
+			FindStripIndex(IsoU, IndexU, LoopPoint[InAxe]);
+
+			double ToleranceU = (IsoU[IndexU + 1] - IsoU[IndexU]) * ToleranceCoefficent;
+
+#ifdef DEBUG_FIND_ISOSEGMENT_TO_LINK_LOOP_TO_LOOP
+			Display(EGridSpace::UniformScaled, *NextNode, 0, EVisuProperty::RedPoint);
+#endif
+			const FPoint2D& NextPoint = NextNode->Get2DPoint(EGridSpace::Default2D, Grid);
+			if (FMath::IsNearlyEqual(NextPoint[InAxe], LoopPoint[InAxe], ToleranceU))
+			{
+				// the nodes are nearly iso aligned, are they nearly in the same V Stip
+				double MinV = LoopPoint[ComplementaryAxe];
+				double MaxV = NextPoint[ComplementaryAxe];
+				Sort(MinV, MaxV);
+
+				int32 IndexV = 0;
+				FindStripIndex(IsoV, IndexV, MinV);
+
+				if (IndexV >= IsoV.Num() - 1)
+				{
+					continue;
+				}
+
+				// We check that segment length is not greater the the crossing strip width
+				bool bIsSmallerThanStrip = false;
+				if (MaxV <= IsoV[IndexV + 1])
+				{
+					// both point are in the same strip
+					bIsSmallerThanStrip = true;
+				}
+				else
+				{
+					// either MinV is nearly equal to IsoV[IndexV + 1]- 
+					double FirstStripCrossingLength = (IsoV[IndexV + 1] - MinV) / (IsoV[IndexV + 1] - IsoV[IndexV]);
+					if (IndexV < IsoV.Num() - 2 && MaxV < IsoV[IndexV + 1])
+					{
+						double SecondStripCrossingLength = (MaxV - IsoV[IndexV + 1]) / (IsoV[IndexV + 2] - IsoV[IndexV + 1]);
+						if ((FirstStripCrossingLength + SecondStripCrossingLength) < 1 + ToleranceCoefficent)
+						{
+							bIsSmallerThanStrip = true;
+						}
+					}
+					// either MaxV is nearly equal to IsoV[IndexV + 1]+
+					else if (IndexV < IsoV.Num() - 3 && MaxV < IsoV[IndexV + 2])
+					{
+						double ThirdStripCrossingLength = (MaxV - IsoV[IndexV + 2]) / (IsoV[IndexV + 3] - IsoV[IndexV + 2]);
+						if ((FirstStripCrossingLength + ThirdStripCrossingLength) < ToleranceCoefficent)
+						{
+							bIsSmallerThanStrip = true;
+						}
+					}
+				}
+				if (bIsSmallerThanStrip)
+				{
+					CreateSegment(LoopNode, LoopPoint, NextNode, NextPoint);
+				}
+			}
+		}
+	};
+
+	int32 InitNum = InnerSegmentsIntersectionTool.Num() + LoopSegmentsIntersectionTool.Num();
+	FinalToLoops.Reserve(InitNum);
+	InnerToLoopSegmentsIntersectionTool.Reserve(InitNum);
+
+	// Nodes are sorted according to a value function of their coordinates.
+	// To sort along U, the value is U + DeltaFactor*(V - VMin)
+	// DeltaFactor is a value that for all values Ui of U, Ui + DeltaFactor.(VMax - VMin) < U(i+1)
+	// With this, Node[i+1] is either the next node of the same side of the loop, either the closed U aligned node of the opposite loop.  
+	{
+		constexpr int32 IsoU = 0; // coordinate U is ~constant
+		constexpr int32 IsoV = 1; // coordinate V is ~constant
+
+		double DeltaFactor = FMath::Min((GetMinDelta(IsoUCoordinates) / 1000.), (GetMinDelta(IsoVCoordinates) / 1000.));
+
+		// Bounds and GetMinDelta are defined in EGridSpace::Default2D,
+		const FSurfacicBoundary& Bounds = Grid.GetFace()->GetBoundary();
+		double UMin = Bounds.UVBoundaries[EIso::IsoU].Min;
+		double VMin = Bounds.UVBoundaries[EIso::IsoV].Min;
+		Algo::Sort(SortedLoopNodesAlong, [&](const FLoopNode* LoopNode1, const FLoopNode* LoopNode2)
+			{
+				const FPoint2D& Node1Coordinates = LoopNode1->Get2DPoint(EGridSpace::Default2D, Grid);
+				const FPoint2D& Node2Coordinates = LoopNode2->Get2DPoint(EGridSpace::Default2D, Grid);
+				return (Node1Coordinates.U + (Node1Coordinates.V - VMin) * DeltaFactor) < (Node2Coordinates.U + (Node2Coordinates.V - VMin) * DeltaFactor);
+			});
+		FindIsoSegmentAlong(IsoU, IsoUCoordinates, IsoVCoordinates);
+
+		Algo::Sort(SortedLoopNodesAlong, [&](const FLoopNode* LoopNode1, const FLoopNode* LoopNode2)
+			{
+				const FPoint2D& Node1Coordinates = LoopNode1->Get2DPoint(EGridSpace::Default2D, Grid);
+				const FPoint2D& Node2Coordinates = LoopNode2->Get2DPoint(EGridSpace::Default2D, Grid);
+				return (Node1Coordinates.V + (Node1Coordinates.U - UMin) * DeltaFactor) < (Node2Coordinates.V + (Node2Coordinates.U - UMin) * DeltaFactor);
+			});
+		FindIsoSegmentAlong(IsoV, IsoVCoordinates, IsoUCoordinates);
+	}
+#ifdef CADKERNEL_DEV
+	Chronos.FindInnerSegmentToLinkLoopToLoopDuration += FChrono::Elapse(StartTime);
+#endif
+}
+
+void FIsoTriangulator::FindIsoSegmentToLinkInnerToLoop()
+{
+	FTimePoint StartTime = FChrono::Now();
+
+	int32 IndexU = 0;
+	int32 IndexV = 0;
+
+	const TArray<double>& IsoUCoordinates = Grid.GetCuttingCoordinatesAlongIso(EIso::IsoU);
+	const TArray<double>& IsoVCoordinates = Grid.GetCuttingCoordinatesAlongIso(EIso::IsoV);
+
+	int32 InitNum = InnerSegmentsIntersectionTool.Num() + LoopSegmentsIntersectionTool.Num();
+	FinalToLoops.Reserve(InitNum);
+	InnerToLoopSegmentsIntersectionTool.Reserve(InitNum);
+
+	TArray<double> IsoUTolerances;
+	TArray<double> IsoVTolerances;
+
+	// As the gap between two IsoPoints is not constant, the tolerance of closeness is define as 1/10 of the coordinate variation between two consecutive iso coordiantes 
+	TFunction<void(const TArray<double>&, TArray<double>&)> SetIsoTolerance = [](const TArray<double>& Isos, TArray<double>& IsoTolerances)
+	{
+		IsoTolerances.SetNum(Isos.Num() - 1);
+		for (int32 Index = 0; Index < Isos.Num() - 1; ++Index)
+		{
+			IsoTolerances[Index] = (Isos[Index + 1] - Isos[Index]) / 10.;
+		}
+	};
+	SetIsoTolerance(IsoUCoordinates, IsoUTolerances);
+	SetIsoTolerance(IsoVCoordinates, IsoVTolerances);
+
+	// We check if a loop point is aligned with the inner grid i.e. U or V points coordinate is closed to one Iso coordinates. The "closed" criteria is define by the local tolerance compute by SetIsoTolerance.
+	// OutIndex Value defined the segment in which the point is localized i.e. PointCoordinate is in [Iso[Index], Iso[Index+1]
+	// is OutPointIndex != 0, OutPointIndex = Index if PointCoordinate is closed to Iso[OutIndex] or OutPointIndex = OutIndex + 1 if PointCoordinate is closed to Iso[OutIndex + 1]
+	TFunction<void(const TArray<double>&, const TArray<double>&, const double&, int32&, int32&)> GetIsoIndexOfPoint = [](const TArray<double>& Iso, const TArray<double>& IsoTolerance, const double& PointCoord, int32& OutIndex, int32& OutPointIndex)
+	{
+		while (PointCoord < Iso[OutIndex])
+		{
+			OutIndex--;
+		}
+
+		for (; OutIndex < Iso.Num() - 2; ++OutIndex)
+		{
+			if (Iso[OutIndex] <= PointCoord && PointCoord <= Iso[OutIndex + 1])
+			{
+				break;
+			}
+		}
+
+		if (PointCoord < Iso[OutIndex] + IsoTolerance[OutIndex])
+		{
+			OutPointIndex = OutIndex;
+		}
+		else if (PointCoord > Iso[OutIndex + 1] - IsoTolerance[OutIndex])
+		{
+			OutPointIndex = OutIndex + 1;
+		}
+	};
+
+	TFunction<const int32(FLoopNode&, const FPoint2D&, const int32, const int32, ESegmentType)> TryToBuildIsoSegment = [&](FLoopNode& IsoLoopNode, const FPoint2D& Point, int32 GlobalIndex1, int32 GlobalIndex2, ESegmentType Type)
+	{
+		int32 IsoGlobalIndex = -1;
+		if (GlobalIndexToIsoInnerNodes[GlobalIndex1] != nullptr && GlobalIndexToIsoInnerNodes[GlobalIndex2] == nullptr)
+		{
+			IsoGlobalIndex = GlobalIndex1;
+		}
+		else if (GlobalIndexToIsoInnerNodes[GlobalIndex1] == nullptr && GlobalIndexToIsoInnerNodes[GlobalIndex2] != nullptr)
+		{
+			IsoGlobalIndex = GlobalIndex2;
+		}
+		if (IsoGlobalIndex < 0)
+		{
+			return IsoGlobalIndex;
+		}
+
+		if (LoopSegmentsIntersectionTool.DoesIntersect(IsoLoopNode, Grid.GetInner2DPoint(EGridSpace::Default2D, IsoGlobalIndex)))
+		{
+			return IsoGlobalIndex;
+		}
+
+		FIsoInnerNode* Node = GlobalIndexToIsoInnerNodes[IsoGlobalIndex];
+		FIsoSegment& Segment = IsoSegmentFactory.New();
+		Segment.Init(*Node, IsoLoopNode, Type);
+		Segment.ConnectToNode();
+		FinalToLoops.Add(&Segment);
+		InnerToLoopSegmentsIntersectionTool.AddSegment(Segment);
+		return IsoGlobalIndex;
+	};
+
+	const TArray<TArray<FPoint2D>>& Loops = Grid.GetLoops2D(EGridSpace::Default2D);
+
+	// Segment smaller than the size of the grid is build
+	for (FLoopNode& LoopNode : LoopNodes)
+	{
+		const FPoint2D& Point = Loops[LoopNode.GetLoopIndex()][LoopNode.GetIndex()];
+		int32 PointIndexU = -1;
+		int32 PointIndexV = -1;
+
+		GetIsoIndexOfPoint(IsoUCoordinates, IsoUTolerances, Point.U, IndexU, PointIndexU);
+		GetIsoIndexOfPoint(IsoVCoordinates, IsoVTolerances, Point.V, IndexV, PointIndexV);
+
+		// Point is in the square [IndexU, IndexV], ... [IndexU+1, IndexV+1]]
+		// If PointIndexU (resp PointIndexV) != -1, Point is closed to IsoU [PointIndexU]
+		// If GlobalIndexToIsoInnerNodes[IndexU+1, IndexV] == null => it is outside the domaine
+		// So if GlobalIndexToIsoInnerNodes[IndexU, IndexV] != null, and PointIndexV = IndexV, the segment [InnerNodes[IndexU, IndexV], LoopNode] could be build
+		if (PointIndexU == IndexU || PointIndexU == IndexU + 1)
+		{
+			const int32 GlobalIndex1 = Grid.GobalIndex(PointIndexU, IndexV);
+			const int32 GlobalIndex2 = GlobalIndex1 + Grid.GetCuttingCount(EIso::IsoU);
+			const int32 IsoGlobalIndex = TryToBuildIsoSegment(LoopNode, Point, GlobalIndex1, GlobalIndex2, ESegmentType::InnerToLoopV);
+			if (IsoGlobalIndex == GlobalIndex1)
+			{
+				GlobalIndexToIsoInnerNodes[IsoGlobalIndex]->SetLinkedToIso(EIsoLink::IsoUNext);
+			}
+			else if (IsoGlobalIndex == GlobalIndex2)
+			{
+				GlobalIndexToIsoInnerNodes[IsoGlobalIndex]->SetLinkedToIso(EIsoLink::IsoUPrevious);
+			}
+		}
+
+		if (PointIndexV == IndexV || PointIndexV == IndexV + 1)
+		{
+			const int32 GlobalIndex1 = Grid.GobalIndex(IndexU, PointIndexV);
+			const int32 GlobalIndex2 = GlobalIndex1 + 1;
+			const int32 IsoGlobalIndex = TryToBuildIsoSegment(LoopNode, Point, GlobalIndex1, GlobalIndex2, ESegmentType::InnerToLoopU);
+			if (IsoGlobalIndex == GlobalIndex1)
+			{
+				GlobalIndexToIsoInnerNodes[IsoGlobalIndex]->SetLinkedToIso(EIsoLink::IsoVNext);
+			}
+			else if (IsoGlobalIndex == GlobalIndex2)
+			{
+				GlobalIndexToIsoInnerNodes[IsoGlobalIndex]->SetLinkedToIso(EIsoLink::IsoVPrevious);
+			}
+		}
+	}
+
+	InnerToLoopSegmentsIntersectionTool.Sort();
+#ifdef CADKERNEL_DEV
+	Chronos.FindIsoSegmentToLinkInnerToLoopDuration += FChrono::Elapse(StartTime);
+#endif
+}
+
+void FIsoTriangulator::FindCandidateSegmentsToLinkInnerToLoop()
+{
+	FTimePoint StartTime = FChrono::Now();
+
+	TArray<FIsoInnerNode*> SortedInnerNode;
+	SortedInnerNode.Reserve(InnerNodes.Num());
+	for (FIsoInnerNode& IsoNode : InnerNodes)
+	{
+		if (IsoNode.IsComplete())
+		{
+			continue;
+		}
+		SortedInnerNode.Add(&IsoNode);
+	}
+
+	Algo::Sort(SortedInnerNode, [this](const FIsoInnerNode* IsoNode1, const FIsoInnerNode* IsoNode2)
+		{
+			const FPoint2D& Node1Coordinates = IsoNode1->Get2DPoint(EGridSpace::Default2D, Grid);
+			const FPoint2D& Node2Coordinates = IsoNode2->Get2DPoint(EGridSpace::Default2D, Grid);
+			return (Node1Coordinates.U + Node1Coordinates.V) < (Node2Coordinates.U + Node2Coordinates.V);
+		});
+
+#ifdef CADKERNEL_DEV
+	Chronos.FindSegmentToLinkInnerToLoopDuration = FChrono::Elapse(StartTime);
+#endif
+
+	TFunction<void(FIsoInnerNode*, FLoopNode*)> CreateInnerToContoursCandidateSegment = [this](FIsoInnerNode* InnerNode, FLoopNode* LoopNode)
+	{
+		if (FIsoSegment::IsItAlreadyDefined(InnerNode, LoopNode))
+		{
+			return;
+		}
+
+		// check intersection with
+		if (LoopSegmentsIntersectionTool.DoesIntersect(*LoopNode, InnerNode->Get2DPoint(EGridSpace::Default2D, Grid)))
+		{
+			return;
+		}
+
+		if (InnerSegmentsIntersectionTool.DoesIntersect(*InnerNode, LoopNode->Get2DPoint(EGridSpace::Default2D, Grid)))
+		{
+			return;
+		}
+
+		if (InnerToLoopSegmentsIntersectionTool.DoesIntersect(*InnerNode, *LoopNode))
+		{
+			return;
+		}
+
+		FIsoSegment& SegCandidate = IsoSegmentFactory.New();
+		SegCandidate.Init(*InnerNode, *LoopNode, ESegmentType::InnerToLoop);
+		CandidateInnerToLoopSegments.Add(&SegCandidate);
+	};
+
+	TFunction<double(const FPoint&, const FPoint&)> ComputeDiagonal = [](const FPoint& PointA, const FPoint& PointB)
+	{
+		FPoint2D AB = PointB - PointA;
+		return AB.U + AB.V;
+	};
+
+	double DiagonalMax = sqrt(FMath::Square(Grid.GetMaxDeltaU(EIso::IsoU)) + FMath::Square(Grid.GetMaxDeltaU(EIso::IsoV)));
+	DiagonalMax *= 1.5;
+
+	TFunction<double(const TArray<double>&, int32&, const double&)> GetDeltaCoordinate = [](const TArray<double>& Iso, int32& Index, const double& PointCoord)
+	{
+		while (PointCoord + SMALL_NUMBER < Iso[Index])
+		{
+			if (Index == 0)
+			{
+				break;
+			}
+			Index--;
+		}
+		for (; Index < Iso.Num(); ++Index)
+		{
+			if (Index == Iso.Num() - 2)
+			{
+				break;
+			}
+			if (PointCoord < Iso[Index + 1u])
+			{
+				break;
+			}
+		}
+
+		if (Index == 1 && Iso.Num() > 3)
+		{
+			return Iso[1] - Iso[0];
+		}
+		else if (Index == 1 && Iso.Num() == 3)
+		{
+			return FMath::Max(Iso[1] - Iso[0], Iso[2] - Iso[1]);
+		}
+		return Iso[Index + 1u] - Iso[Index];
+	};
+
+	if (SortedInnerNode.Num())
+	{
+		CandidateInnerToLoopSegments.Reserve(SortedInnerNode.Num());
+		int32 IndexU = 0;
+		int32 IndexV = 0;
+
+		int32 StartInnerNodeIndex = 0;
+		for (FLoopNode* LoopNode : SortedLoopNodes)
+		{
+			const FPoint2D& PointLoop = LoopNode->Get2DPoint(EGridSpace::Default2D, Grid);
+			{
+				const FPoint2D& PointInner = Grid.GetInner2DPoint(EGridSpace::Default2D, SortedInnerNode[StartInnerNodeIndex]->GetIndex());
+				double Diagonal = ComputeDiagonal(PointLoop, PointInner);
+				if (Diagonal > DiagonalMax)
+				{
+					continue;
+				}
+			}
+
+			for (; StartInnerNodeIndex < SortedInnerNode.Num(); ++StartInnerNodeIndex)
+			{
+				const FPoint2D& PointInner = Grid.GetInner2DPoint(EGridSpace::Default2D, SortedInnerNode[StartInnerNodeIndex]->GetIndex());
+
+				double Diagonal = ComputeDiagonal(PointLoop, PointInner);
+				if (Diagonal >= -DiagonalMax)
+				{
+					break;
+				}
+			}
+			if (StartInnerNodeIndex == SortedInnerNode.Num())
+			{
+				break;
+			}
+
+			for (int32 InnerNodeIndex = StartInnerNodeIndex; InnerNodeIndex < SortedInnerNode.Num(); ++InnerNodeIndex)
+			{
+				const FPoint2D& PointInner = Grid.GetInner2DPoint(EGridSpace::Default2D, SortedInnerNode[InnerNodeIndex]->GetIndex());
+
+				FPoint2D DiagonalV = PointLoop - PointInner;
+				double Diagonal = FMath::Abs(DiagonalV.U + DiagonalV.V);
+
+				if (Diagonal > DiagonalMax)
+				{
+					break;
+				}
+
+				double DeltaUMax = GetDeltaCoordinate(Grid.GetCuttingCoordinatesAlongIso(EIso::IsoU), IndexU, PointInner.U);
+				DeltaUMax *= 1.5;
+				double DeltaVMax = GetDeltaCoordinate(Grid.GetCuttingCoordinatesAlongIso(EIso::IsoV), IndexV, PointInner.V);
+				DeltaVMax *= 1.5;
+
+				if (FMath::Abs(DiagonalV.U) > DeltaUMax || FMath::Abs(DiagonalV.V) > DeltaVMax)
+				{
+					continue;
+				}
+
+				CreateInnerToContoursCandidateSegment(SortedInnerNode[InnerNodeIndex], LoopNode);
+			}
+		}
+	}
+
+#ifdef CADKERNEL_DEV
+	Chronos.FindSegmentToLinkInnerToLoopDuration = FChrono::Elapse(StartTime);
+#endif
+}
+
+void FIsoTriangulator::SelectSegmentInCandidateSegments()
+{
+	constexpr double FlatAngle = 0.25; // PI/12
+	constexpr double ColinearityCriteriaAngle = 0.25; // PI/12
+
+	FTimePoint StartTime = FChrono::Now();
+
+	// A candidate segment must not be too collinear with existing segment. So if this is the case, it's refused. 
+	// final step of the mesh will manage these cases
+	TFunction<bool(const FIsoNode&, const FIsoSegment*, const double)> CheckColinerarityWithFinalSegments = [&](const FIsoNode& Node, const FIsoSegment* CandidateSegment, const double CandidateSegmentSlop)
+	{
+		bool bIsTooColineareWithFinalSegment = false;
+		for (const FIsoSegment* NodeSegment : Node.GetConnectedSegments())
+		{
+			if (NodeSegment == CandidateSegment)
+			{
+				continue;
+			}
+
+			double AngleBetweenCandidateAndFinalSegment = &NodeSegment->GetFirstNode() == &Node ?
+				ComputeOrientedSlope(NodeSegment->GetFirstNode().Get2DPoint(EGridSpace::Default2D, Grid), NodeSegment->GetSecondNode().Get2DPoint(EGridSpace::Default2D, Grid), CandidateSegmentSlop) :
+				ComputeOrientedSlope(NodeSegment->GetSecondNode().Get2DPoint(EGridSpace::Default2D, Grid), NodeSegment->GetFirstNode().Get2DPoint(EGridSpace::Default2D, Grid), CandidateSegmentSlop);
+
+			if (FMath::Abs(AngleBetweenCandidateAndFinalSegment) < ColinearityCriteriaAngle)
+			{
+				bIsTooColineareWithFinalSegment = true;
+				return true;
+			}
+		}
+		return false;
+	};
+
+	// The main idea is to select only one segment by Quarter ([0 Pi/2], [Pi/2, Pi], [Pi 3Pi/2], [3Pi/2 2Pi])
+	// in additional of Iso Segment (EIso::IsoU or EIso::IsoV). If the vertex has not an iso segment, closed iso segment are accepted 
+	TFunction<bool(const FIsoNode&, const FIsoNode&, const double, int32)> CheckQuarterCompletude = [&](const FIsoNode& StartNode, const FIsoNode& EndNode, const double CandidateSegmentSlop, int32 Quarter)
+	{
+		if (CandidateSegmentSlop < FlatAngle)
+		{
+			if (StartNode.IsLinkedToLoopInNearlyIso(Quarter))
+			{
+				return false;
+			}
+			if (EndNode.IsLinkedToLoopInNearlyIso(Quarter + 2))
+			{
+				return false;
+			}
+		}
+		else if (CandidateSegmentSlop < 2 - FlatAngle)
+		{
+			if (StartNode.IsLinkedToLoopInQuarter(Quarter))
+			{
+				return false;
+			}
+			if (EndNode.IsLinkedToLoopInQuarter(Quarter > 2 ? Quarter - 2 : Quarter + 2))
+			{
+				return false;
+			}
+		}
+		else
+		{
+			if (StartNode.IsLinkedToLoopInNearlyIso(Quarter + 1))
+			{
+				return false;
+			}
+			if (EndNode.IsLinkedToLoopInNearlyIso(Quarter + 3))
+			{
+				return false;
+			}
+		}
+		return true;
+	};
+
+	TFunction<bool(FIsoNode&, FIsoNode&, const double, int32)>  SetNodeFlag = [&](FIsoNode& StartNode, FIsoNode& EndNode, const double CandidateSegmentSlop, int32 Quarter)
+	{
+		if (CandidateSegmentSlop < FlatAngle)
+		{
+			StartNode.SetLinkedToIso(Quarter);
+			EndNode.SetLinkedToIso(Quarter + 2);
+		}
+		else if (CandidateSegmentSlop < 2 - FlatAngle)
+		{
+			StartNode.SetLinkedToLoopInQuarter(Quarter);
+			EndNode.SetLinkedToLoopInQuarter(Quarter > 2 ? Quarter - 2 : Quarter + 2);
+		}
+		else
+		{
+			StartNode.SetLinkedToIso(Quarter + 1);
+			EndNode.SetLinkedToIso(Quarter + 3);
+		}
+		return true;
+	};
+
+	CandidateSegments.Reserve(CandidateLoopToLoopSegments.Num() + CandidateInnerToLoopSegments.Num());
+	CandidateSegments.Insert(CandidateLoopToLoopSegments, CandidateSegments.Num());
+	CandidateSegments.Insert(CandidateInnerToLoopSegments, CandidateSegments.Num());
+
+	for (FIsoSegment* Segment : CandidateSegments)
+	{
+		Segment->SetCandidate();
+	}
+
+	Algo::Sort(CandidateSegments, [this](const FIsoSegment* Segment1, const FIsoSegment* Segment2)
+		{
+			return (Segment1->Get3DLengthSquare(Grid)) < (Segment2->Get3DLengthSquare(Grid));
+		});
+
+	InnerSegmentsIntersectionTool.SetNum(InnerSegmentsIntersectionTool.Num() + (int32)CandidateSegments.Num());
+
+	FinalToLoops.Reserve(CandidateSegments.Num());
+
+	FIntersectionSegmentTool LocalSegmentsIntersectionTool(Grid);
+	for (FIsoSegment* Segment : CandidateSegments)
+	{
+		FIsoNode& StartNode = Segment->GetFirstNode();
+		FIsoNode& EndNode = Segment->GetSecondNode();
+
+		double CandidateSegmentSlop = ComputePositiveSlope(StartNode.Get2DPoint(EGridSpace::Default2D, Grid), EndNode.Get2DPoint(EGridSpace::Default2D, Grid), 0);
+
+		// Define the quarter ([0 Pi/2], [Pi/2, Pi], [Pi 3Pi/2], [3Pi/2 2Pi])
+		int32 Quarter = 3;
+		if (CandidateSegmentSlop < 2)
+		{
+			Quarter = 0;
+		}
+		else if (CandidateSegmentSlop < 4)
+		{
+			Quarter = 1;
+		}
+		else if (CandidateSegmentSlop < 6)
+		{
+			Quarter = 2;
+		}
+
+		if (!CheckQuarterCompletude(StartNode, EndNode, CandidateSegmentSlop - Quarter * 2, Quarter))
+		{
+			continue;
+		}
+
+		if (CheckColinerarityWithFinalSegments(StartNode, Segment, CandidateSegmentSlop))
+		{
+			continue;
+		}
+
+		if (CheckColinerarityWithFinalSegments(EndNode, Segment, CandidateSegmentSlop > 4 ? CandidateSegmentSlop - 4 : CandidateSegmentSlop + 4))
+		{
+			continue;
+		}
+
+		if (LocalSegmentsIntersectionTool.DoesIntersect(*Segment))
+		{
+			continue;
+		}
+
+		Segment->SetSelected();
+		Segment->ConnectToNode();
+		SetNodeFlag(StartNode, EndNode, CandidateSegmentSlop - Quarter * 2, Quarter);
+
+		FinalToLoops.Add(Segment);
+
+		LocalSegmentsIntersectionTool.AddSegment(*Segment);
+		LocalSegmentsIntersectionTool.Sort();
+	}
+
+	for (FIsoSegment* Segment : CandidateSegments)
+	{
+		if (Segment->IsACandidate())
+		{
+			IsoSegmentFactory.DeleteEntity(Segment);
+		}
+	}
+
+	InnerSegmentsIntersectionTool.Sort();
+#ifdef CADKERNEL_DEV
+	Chronos.SelectSegmentToLinkInnerToLoopsDuration = FChrono::Elapse(StartTime);
+#endif
+}
+
+void FIsoTriangulator::ConnectUnconnectedInnerSegments()
+{
+	TArray<FIsoNode*> UnconnectedNode;
+	UnconnectedNode.Reserve(FinalInnerSegments.Num());
+	for (FIsoSegment* Segment : FinalInnerSegments)
+	{
+		if (Segment->GetFirstNode().GetConnectedSegments().Num() == 1)
+		{
+			UnconnectedNode.Add(&Segment->GetFirstNode());
+		}
+	}
+
+	for (FIsoNode* Node : UnconnectedNode)
+	{
+		double MinDistance = HUGE_VALUE;
+		FLoopNode* Candidate = nullptr;
+		for (FLoopNode& LoopNode : LoopNodes)
+		{
+			double Distance = LoopNode.Get2DPoint(EGridSpace::Scaled, Grid).SquareDistance(Node->Get2DPoint(EGridSpace::Scaled, Grid));
+			if (Distance < MinDistance)
+			{
+				if (!InnerToLoopSegmentsIntersectionTool.DoesIntersect(*Node, LoopNode))
+				{
+					MinDistance = Distance;
+					Candidate = &LoopNode;
+				}
+			}
+		}
+
+		if (Candidate)
+		{
+			FIsoSegment& Segment = IsoSegmentFactory.New();
+			Segment.Init(*Node, *Candidate, ESegmentType::InnerToLoop);
+			Segment.ConnectToNode();
+			FinalToLoops.Add(&Segment);
+			InnerToLoopSegmentsIntersectionTool.AddSegment(Segment);
+			InnerToLoopSegmentsIntersectionTool.Sort();
+		}
+	}
+}
+
+/**
+ * Criteria to find the optimal "Delaunay" triangle starting from the segment AB to a set of point P
+ * A "Delaunay" triangle is an equilateral triangle
+ * The optimal value is the smallest value.
+ */
+double CotangentCriteria(const FPoint& APoint, const FPoint& BPoint, const FPoint& PPoint, FPoint& OutNormal)
+{
+	const double BigValue = HUGE_VALUE;
+
+	FPoint PA = APoint - PPoint;
+	FPoint PB = BPoint - PPoint;
+
+	// the ratio between the scalar product PA.PB (=|PA| |PB| cos (A,P,B) )
+	// with the norm of the cross product |PA^PB| (=|PA| |PB| |sin(A,P,B)|)
+	// is compute. 
+	double ScalareProduct = PA * PB;
+	OutNormal = PA ^ PB;
+	double NormOFScalarProduct = sqrt(OutNormal * OutNormal);
+
+	// PPoint is aligned with (A,B)
+	if (NormOFScalarProduct < SMALL_NUMBER)
+	{
+		return BigValue;
+	}
+
+	// return Cotangent value 
+	return ScalareProduct / NormOFScalarProduct;
+}
+
+double CotangentCriteria(const FPoint2D& APoint, const FPoint2D& BPoint, const FPoint2D& PPoint)
+{
+	const double BigValue = HUGE_VALUE;
+
+	FPoint2D PA = APoint - PPoint;
+	FPoint2D PB = BPoint - PPoint;
+
+	// the ratio between the scalar product PA.PB (=|PA| |PB| cos (A,P,B) )
+	// with the norm of the cross product |PA^PB| (=|PA| |PB| |sin(A,P,B)|)
+	// is compute. 
+	double ScalareProduct = PA * PB;
+	double OutNormal = PA ^ PB;
+	double NormOFPointProduct = FMath::Abs(OutNormal);
+
+	if (NormOFPointProduct < SMALL_NUMBER)
+	{
+		// PPoint is aligned with (A,B)
+		return BigValue;
+	}
+
+	// return Cotangent value 
+	return ScalareProduct / NormOFPointProduct;
+}
+
+#define D3_COTANGENT_CRITERIA
+void FIsoTriangulator::MeshCycle(const EGridSpace Space, const TArray<FIsoSegment*>& Cycle, const TArray<bool>& CycleOrientation)
+{
+	int32 NodeCycleNum = (int32)Cycle.Num();
+
+	if (NodeCycleNum == 4)
+	{
+		FIsoNode* Nodes[4];
+		if (CycleOrientation[0])
+		{
+			Nodes[0] = &Cycle[0]->GetFirstNode();
+			Nodes[1] = &Cycle[0]->GetSecondNode();
+		}
+		else
+		{
+			Nodes[0] = &Cycle[0]->GetSecondNode();
+			Nodes[1] = &Cycle[0]->GetFirstNode();
+		}
+
+		if (CycleOrientation[2])
+		{
+			Nodes[2] = &Cycle[2]->GetFirstNode();
+			Nodes[3] = &Cycle[2]->GetSecondNode();
+		}
+		else
+		{
+			Nodes[2] = &Cycle[2]->GetSecondNode();
+			Nodes[3] = &Cycle[2]->GetFirstNode();
+		}
+
+		const FPoint2D* NodeCoordinates[4];
+		for (int32 Index = 0; Index < 4; ++Index)
+		{
+			NodeCoordinates[Index] = &Nodes[Index]->Get2DPoint(Space, Grid);
+		}
+
+		double SegmentSlopes[4];
+		SegmentSlopes[0] = ComputeSlope(*NodeCoordinates[0], *NodeCoordinates[1]);
+		SegmentSlopes[1] = ComputeSlope(*NodeCoordinates[1], *NodeCoordinates[2]);
+		SegmentSlopes[2] = ComputeSlope(*NodeCoordinates[2], *NodeCoordinates[3]);
+		SegmentSlopes[3] = ComputeSlope(*NodeCoordinates[3], *NodeCoordinates[0]);
+
+		double RelativeSlopes[4];
+		RelativeSlopes[0] = TransformIntoOrientedSlope(SegmentSlopes[1] - SegmentSlopes[0]);
+		RelativeSlopes[1] = TransformIntoOrientedSlope(SegmentSlopes[2] - SegmentSlopes[1]);
+		RelativeSlopes[2] = TransformIntoOrientedSlope(SegmentSlopes[3] - SegmentSlopes[2]);
+		RelativeSlopes[3] = TransformIntoOrientedSlope(SegmentSlopes[0] - SegmentSlopes[3]);
+
+		int32 FlattenNodeIndex = 0;
+		for (int32 IndexAngle = 0; IndexAngle < 4; ++IndexAngle)
+		{
+			if (RelativeSlopes[IndexAngle] < RelativeSlopes[FlattenNodeIndex])
+			{
+				FlattenNodeIndex = IndexAngle;
+			}
+		}
+
+		int32 NodeIndices[4];
+		NodeIndices[0] = FlattenNodeIndex;
+		for (int32 IndexN = 1; IndexN < 4; ++IndexN)
+		{
+			NodeIndices[IndexN] = NodeIndices[IndexN - 1] == 3 ? 0 : NodeIndices[IndexN - 1] + 1;
+		}
+
+#ifdef ADD_TRIANGLE_2D
+		if (bDisplay)
+		{
+			F3DDebugSession G(TEXT("Mesh cycle"));
+			DisplayTriangle(EGridSpace::UniformScaled, *Nodes[NodeIndices[1]], *Nodes[NodeIndices[3]], *Nodes[NodeIndices[0]]);
+			DisplayTriangle(EGridSpace::UniformScaled, *Nodes[NodeIndices[1]], *Nodes[NodeIndices[2]], *Nodes[NodeIndices[3]]);
+		}
+#endif 
+		Mesh->AddTriangle(Nodes[NodeIndices[1]]->GetFaceIndex(), Nodes[NodeIndices[3]]->GetFaceIndex(), Nodes[NodeIndices[0]]->GetFaceIndex());
+		Mesh->AddTriangle(Nodes[NodeIndices[1]]->GetFaceIndex(), Nodes[NodeIndices[2]]->GetFaceIndex(), Nodes[NodeIndices[3]]->GetFaceIndex());
+
+		return;
+	}
+	else if (NodeCycleNum == 3)
+	{
+		if (CycleOrientation[0])
+		{
+#ifdef ADD_TRIANGLE_2D
+			if (bDisplay)
+			{
+				F3DDebugSession G(TEXT("Mesh cycle"));
+				DisplayTriangle(EGridSpace::UniformScaled, Cycle[0]->GetFirstNode(), Cycle[0]->GetSecondNode(), CycleOrientation[1] ? Cycle[1]->GetSecondNode() : Cycle[1]->GetFirstNode());
+			}
+#endif 
+			Mesh->AddTriangle(Cycle[0]->GetFirstNode().GetFaceIndex(), Cycle[0]->GetSecondNode().GetFaceIndex(), CycleOrientation[1] ? Cycle[1]->GetSecondNode().GetFaceIndex() : Cycle[1]->GetFirstNode().GetFaceIndex());
+		}
+		else
+		{
+#ifdef ADD_TRIANGLE_2D
+			if (bDisplay)
+			{
+				F3DDebugSession G(TEXT("Mesh cycle"));
+				DisplayTriangle(EGridSpace::UniformScaled, Cycle[0]->GetSecondNode(), Cycle[0]->GetFirstNode(), CycleOrientation[1] ? Cycle[1]->GetSecondNode() : Cycle[1]->GetFirstNode());
+			}
+#endif 
+			Mesh->AddTriangle(Cycle[0]->GetSecondNode().GetFaceIndex(), Cycle[0]->GetFirstNode().GetFaceIndex(), CycleOrientation[1] ? Cycle[1]->GetSecondNode().GetFaceIndex() : Cycle[1]->GetFirstNode().GetFaceIndex());
+		}
+		return;
+	}
+
+	FIntersectionSegmentTool CycleIntersectionTool(Grid);
+	CycleIntersectionTool.Reserve((int32)FMath::Square(NodeCycleNum));
+	CycleIntersectionTool.AddSegment(Cycle);
+	CycleIntersectionTool.Sort();
+
+	// check if the cycle is self intersecting. 
+	for (const FIsoSegment* Segment : Cycle)
+	{
+		if (CycleIntersectionTool.DoesIntersect(*Segment))
+		{
+			FMessage::Printf(Log, TEXT("A cycle of the surface %d is in self intersecting. The mesh of this sector is canceled."), Grid.GetFace()->GetId());
+			return;
+		}
+	}
+
+#ifdef ADD_TRIANGLE_2D
+	if (bDisplay)
+	{
+		Open3DDebugSession(TEXT("Mesh cycle"));
+	}
+#endif
+
+	TArray<FIsoNode*> CycleNodes;
+	CycleNodes.Reserve(NodeCycleNum);
+
+	TArray<FIsoSegment*> SegmentStack;
+	SegmentStack.Reserve(5 * NodeCycleNum);
+
+	{
+		// Get cycle's nodes and set segments as they have a triangle outside the cycle (to don't try to mesh outside the cycle)
+		auto SegmentOrientation = CycleOrientation.begin();
+		for (auto Segment = Cycle.begin(); Segment != Cycle.end(); ++Segment, ++SegmentOrientation)
+		{
+			if (*SegmentOrientation)
+			{
+				CycleNodes.Add(&(*Segment)->GetFirstNode());
+				ensureCADKernel(!(*Segment)->HasTriangleOnRight());
+				(*Segment)->SetHasTriangleOnRight();
+			}
+			else
+			{
+				CycleNodes.Add(&(*Segment)->GetSecondNode());
+				ensureCADKernel(!(*Segment)->HasTriangleOnLeft());
+				(*Segment)->SetHasTriangleOnLeft();
+			}
+		}
+
+		// If the Segment has 2 adjacent triangles, the segment is a inner cycle segment
+		// It will have triangle in both side
+		//
+		//    X---------------X----------------X        X---------------X----------------X
+		//    |                                |   		|                                |  
+		//    |         X--------------------X |   		|                                |  
+		//    |         |                    | |   		|                                |  
+		//    X---------X  <- inner segment  | |	or  X---------X  <- inner segment    |
+		//    |         |                    | |   		|                                |  
+		//    |         X--------------------X |   		|                                |  
+		//    |                                |		|                                |
+		//    X---------------X----------------X		X---------------X----------------X
+		//												
+		for (FIsoSegment* Segment : Cycle)
+		{
+			if (Segment->HasTriangleOnRightAndLeft())
+			{
+				Segment->ResetHasTriangle();
+			}
+
+			if (Segment->GetFirstNode().GetConnectedSegments().Num() == 1 ||
+				Segment->GetSecondNode().GetConnectedSegments().Num() == 1)
+			{
+				Segment->ResetHasTriangle();
+			}
+		}
+
+		NodeCycleNum = (int32)Cycle.Num();
+
+		TArray<int32> NodeIndex;
+		NodeIndex.Reserve(NodeCycleNum);
+		for (int32 Index = 0; Index < NodeCycleNum; ++Index)
+		{
+			NodeIndex.Add(Index);
+		}
+
+		TArray<double> SegmentLengths;
+		SegmentLengths.Reserve(NodeCycleNum);
+		for (int32 Index = 0, NextIndex = 1; Index < NodeCycleNum; ++Index, ++NextIndex)
+		{
+			if (NextIndex == NodeCycleNum)
+			{
+				NextIndex = 0;
+			}
+			double Length = CycleNodes[Index]->Get3DPoint(Grid).SquareDistance(CycleNodes[NextIndex]->Get3DPoint(Grid));
+			if (Length < SMALL_NUMBER)
+			{
+				Cycle[Index]->SetAsDegenerated();
+			}
+			SegmentLengths.Add(Length);
+		}
+
+		NodeIndex.Sort([SegmentLengths](const int32& Index1, const int32& Index2) { return SegmentLengths[Index1] > SegmentLengths[Index2]; });
+
+		for (int32 Index = 0; Index < NodeCycleNum; ++Index)
+		{
+			SegmentStack.Add(Cycle[NodeIndex[Index]]);
+		}
+	}
+
+	// Function used in FindBestTriangle
+	TFunction<void(FIsoNode*, FIsoNode*, FIsoSegment*)> BuildSegmentIfNeeded = [&](FIsoNode* NodeA, FIsoNode* NodeB, FIsoSegment* ABSegment)
+	{
+		if (ABSegment)
+		{
+			if (&ABSegment->GetFirstNode() == NodeA)
+			{
+#ifdef ADD_TRIANGLE_2D
+				if (ABSegment->HasTriangleOnLeft())
+				{
+					Wait();
+				}
+#endif
+				ensureCADKernel(!ABSegment->HasTriangleOnLeft());
+				ABSegment->SetHasTriangleOnLeft();
+			}
+			else
+			{
+#ifdef ADD_TRIANGLE_2D
+				if (ABSegment->HasTriangleOnRight())
+				{
+					Wait();
+				}
+#endif
+				ensureCADKernel(!ABSegment->HasTriangleOnRight());
+				ABSegment->SetHasTriangleOnRight();
+			}
+		}
+		else
+		{
+			FIsoSegment& NewSegment = IsoSegmentFactory.New();
+			NewSegment.Init(*NodeA, *NodeB, ESegmentType::Unknown);
+			NewSegment.ConnectToNode();
+			CycleIntersectionTool.AddSegment(NewSegment);
+			NewSegment.SetHasTriangleOnLeft();
+			SegmentStack.Add(&NewSegment);
+		}
+	};
+
+	TFunction<void(FIsoSegment*, bool)> FindBestTriangle = [&](FIsoSegment* Segment, bool bOrientation)
+	{
+		FGetSlop GetSlopAtStartNode = ClockwiseSlop;
+		FGetSlop GetSlopAtEndNode = CounterClockwiseSlop;
+
+		// StartNode = A
+		FIsoNode& StartNode = bOrientation ? Segment->GetFirstNode() : Segment->GetSecondNode();
+		// EndNode = B
+		FIsoNode& EndNode = bOrientation ? Segment->GetSecondNode() : Segment->GetFirstNode();
+
+
+		//
+		// For each extremity (A, B) of a segment, in the existing connected segments, the segment with the smallest relative slop is identified ([A, X0] and [B, Xn]).
+		// These segments define the sector in which the best triangle could be.
+		// The triangle to build is the best triangle (according to the Cotangent Criteria) connecting the Segment to one of the allowed nodes (X) between X0 and Xn.
+		// Allowed nodes (X) are in the sector, Disallowed nodes (Z) are outside the sector	
+		//
+		//                                       ------Z------X0-------X------X-----X-------Xn----Z-----Z---
+		//                                                     \                           /
+		//                                                      \    Allowed triangles    /   
+		//                                Not allowed triangles  \                       /   Not allowed triangles
+		//                                                        \                     /
+		//                                             ----Z-------A------Segment------B------Z---
+		//
+		//                                                         Not allowed triangles
+		//
+		// These computations are done in the UniformScaled space to avoid numerical error due to lenght distortion between U or V space and U or V Length.
+		// i.e. if:
+		// (UMax - UMin) / (VMax - VMin) is big 
+		// and 
+		// "medium length along U" / "medium length along V" is small 
+		// The computed angles or slot is not representative of the 3D space.
+		//
+		// The computation is not done is Scale space to don't have problem with degenerated segments
+		//
+		// To avoid flat triangle, a candidate point must defined a minimal slop with [A, X0] or [B, Xn] to not be aligned with one of them. 
+		//
+
+		// PreviousSegment = [A, X0]
+		FIsoSegment* PreviousSegment = FindNextSegment(EGridSpace::UniformScaled, Segment, &StartNode, GetSlopAtStartNode);
+		// NextSegment = [B, Xn]
+		FIsoSegment* NextSegment = FindNextSegment(EGridSpace::UniformScaled, Segment, &EndNode, GetSlopAtEndNode);
+
+		// PreviousNode = X0
+		FIsoNode& PreviousNode = PreviousSegment->GetOtherNode(&StartNode);
+		// NextNode = Xn
+		FIsoNode& NextNode = NextSegment->GetOtherNode(&EndNode);
+
+#ifdef DEBUG_FIND_BEST_TRIANGLE
+		if (bDisplay)
+		{
+			{
+				TriangleIndex++;
+				F3DDebugSession _(FString::Printf(TEXT("Start Segment %d %d"), TriangleIndex, bOrientation));
+				Display(EGridSpace::UniformScaled, *Segment);
+				Display(EGridSpace::UniformScaled, StartNode, 0, EVisuProperty::RedPoint);
+				Display(EGridSpace::UniformScaled, EndNode);
+			}
+
+			{
+				F3DDebugSession _(TEXT("Next Segments"));
+				Display(EGridSpace::UniformScaled, *PreviousSegment);
+				Display(EGridSpace::UniformScaled, *NextSegment, 0, EVisuProperty::EdgeMesh);
+				Display(EGridSpace::UniformScaled, EndNode, 0, EVisuProperty::RedPoint);
+				Display(EGridSpace::UniformScaled, NextNode);
+				Wait(Grid.GetFace()->GetId() == SurfaceToDebug);
+			}
+		}
+#endif
+
+		FIsoNode* CandidatNode = nullptr;
+		FIsoSegment* StartToCandiatSegment = nullptr;
+		FIsoSegment* EndToCandiatSegment = nullptr;
+
+		if (!NextSegment->IsDegenerated() && !PreviousSegment->IsDegenerated())
+		{
+			const FPoint2D& StartPoint2D = StartNode.Get2DPoint(EGridSpace::UniformScaled, Grid);
+			const FPoint2D& EndPoint2D = EndNode.Get2DPoint(EGridSpace::UniformScaled, Grid);
+
+			const FPoint& StartPoint3D = StartNode.Get3DPoint(Grid);
+			const FPoint& EndPoint3D = EndNode.Get3DPoint(Grid);
+
+			// StartMaxSlope and EndMaxSlope are at most equal to 4, because if the slop with candidate node is biggest to 4, the nez triangle will be inverted
+			double StartReferenceSlope = ComputePositiveSlope(StartPoint2D, EndPoint2D, 0);
+			double StartMaxSlope = GetSlopAtStartNode(StartPoint2D, PreviousNode.Get2DPoint(EGridSpace::UniformScaled, Grid), StartReferenceSlope);
+			if (&EndNode != &PreviousNode)
+			{
+				// Case of probable auto-intersection cycle at PreviousNode, cancel is prefered 
+				if (FMath::IsNearlyEqual(StartMaxSlope, 8., (double)KINDA_SMALL_NUMBER))
+				{
+					return;
+				}
+			}
+			StartMaxSlope = FMath::Min(StartMaxSlope, 4.);
+
+			double EndReferenceSlope = StartReferenceSlope < 4 ? StartReferenceSlope + 4 : StartReferenceSlope - 4;
+			double EndMaxSlope = GetSlopAtEndNode(EndPoint2D, NextNode.Get2DPoint(EGridSpace::UniformScaled, Grid), EndReferenceSlope);
+			if (&StartNode != &NextNode)
+			{
+				// Case of probable auto-intersection cycle at PreviousNode, cancel is prefered 
+				if (FMath::IsNearlyEqual(EndMaxSlope, 8., (double)KINDA_SMALL_NUMBER))
+				{
+					return;
+				}
+			}
+			EndMaxSlope = FMath::Min(EndMaxSlope, 4.);
+
+			double MinCriteria = HUGE_VALUE;
+			const double MinSlopToNotBeAligned = 0.0001;
+			double CandidateSlopeAtStartNode = 8.;
+			double CandidateSlopeAtEndNode = 8.;
+
+			for (FIsoNode* Node : CycleNodes)
+			{
+				if (Node == &StartNode)
+				{
+					continue;
+				}
+				if (Node == &EndNode)
+				{
+					continue;
+				}
+
+
+				// Check if the node is inside the sector (X) or outside (Z)
+				const FPoint2D& NodePoint2D = Node->Get2DPoint(EGridSpace::UniformScaled, Grid);
+				double SlopeAtStartNode = GetSlopAtStartNode(StartPoint2D, NodePoint2D, StartReferenceSlope);
+				double SlopeAtEndNode = GetSlopAtEndNode(EndPoint2D, NodePoint2D, EndReferenceSlope);
+
+				if (Node != &PreviousNode)
+				{
+					if (SlopeAtStartNode <= MinSlopToNotBeAligned || SlopeAtStartNode >= StartMaxSlope - MinSlopToNotBeAligned)
+					{
+						continue;
+					}
+				}
+
+				if (Node != &NextNode)
+				{
+					if (SlopeAtEndNode <= MinSlopToNotBeAligned || SlopeAtEndNode >= EndMaxSlope - MinSlopToNotBeAligned)
+					{
+						continue;
+					}
+				}
+
+				if (FMath::IsNearlyEqual(SlopeAtStartNode, CandidateSlopeAtStartNode, MinSlopToNotBeAligned) && SlopeAtEndNode > CandidateSlopeAtEndNode)
+				{
+					continue;
+				}
+
+				if (FMath::IsNearlyEqual(SlopeAtEndNode, CandidateSlopeAtEndNode, MinSlopToNotBeAligned) && SlopeAtStartNode > CandidateSlopeAtStartNode)
+				{
+					continue;
+				}
+
+#ifdef D3_COTANGENT_CRITERIA
+				const FPoint& NodePoint3D = Node->Get3DPoint(Grid);
+				FPoint NodeNormal;
+				double PointCriteria = FMath::Abs(CotangentCriteria(StartPoint3D, EndPoint3D, NodePoint3D, NodeNormal));
+				double CosAngle = FMath::Abs(NodeNormal.ComputeCosinus(Node->GetNormal(Grid)));
+
+				// the criteria is weighted according to the cosinus of the angle between the normal of the candidate triangle and the normal at the tested point
+				if (CosAngle > SMALL_NUMBER)
+				{
+					PointCriteria /= CosAngle;
+				}
+				else
+				{
+					PointCriteria = HUGE_VALUE;
+				}
+#else // 2D CotangentCrietria
+				const FPoint2D& NodePoint2DSpace = Node->Get2DPoint(Space, Grid);
+				double PointCriteria = FMath::Abs(CotangentCriteria(StartPoint2DSpace, EndPoint2DSpace, NodePoint2DSpace));
+#endif
+				if (
+					// the candidate triangle is inside the current candidate triangle
+					((SlopeAtStartNode < (CandidateSlopeAtStartNode + MinSlopToNotBeAligned)) && (SlopeAtEndNode < (CandidateSlopeAtEndNode + MinSlopToNotBeAligned)))
+					||
+					// the candidate triangle is better the current candidate triangle and doesn't contain the current candidate triangle
+					((PointCriteria < MinCriteria) && ((SlopeAtStartNode > CandidateSlopeAtStartNode) ^ (SlopeAtEndNode > CandidateSlopeAtEndNode))))
+				{
+					// check if the candidate segment is not in intersection with existing segments
+					// if the segment exist, it has already been tested
+					FIsoSegment* StartSegment = StartNode.GetSegmentConnectedTo(Node);
+					FIsoSegment* EndSegment = EndNode.GetSegmentConnectedTo(Node);
+
+					if (!StartSegment && CycleIntersectionTool.DoesIntersect(StartNode, *Node))
+					{
+						continue;
+					}
+
+					if (!EndSegment && CycleIntersectionTool.DoesIntersect(EndNode, *Node))
+					{
+						continue;
+					}
+
+					MinCriteria = PointCriteria;
+					CandidatNode = Node;
+					StartToCandiatSegment = StartSegment;
+					EndToCandiatSegment = EndSegment;
+					CandidateSlopeAtStartNode = SlopeAtStartNode;
+					CandidateSlopeAtEndNode = SlopeAtEndNode;
+				}
+			}
+		}
+
+		if (CandidatNode)
+		{
+			if (bOrientation)
+			{
+				Segment->SetHasTriangleOnRight();
+			}
+			else
+			{
+				Segment->SetHasTriangleOnLeft();
+			}
+
+			BuildSegmentIfNeeded(&StartNode, CandidatNode, StartToCandiatSegment);
+			BuildSegmentIfNeeded(CandidatNode, &EndNode, EndToCandiatSegment);
+			Mesh->AddTriangle(EndNode.GetFaceIndex(), StartNode.GetFaceIndex(), CandidatNode->GetFaceIndex());
+
+#ifdef ADD_TRIANGLE_2D
+			if (bDisplay)
+			{
+				F3DDebugSession _(FString::Printf(TEXT("Triangle %d"), TriangleIndex));
+				DisplayTriangle(EGridSpace::UniformScaled, EndNode, StartNode, *CandidatNode);
+				Wait(false);
+			}
+#endif 
+
+			if (StartToCandiatSegment == nullptr || EndToCandiatSegment == nullptr)
+			{
+				CycleIntersectionTool.Sort();
+			}
+		}
+	};
+
+	for (int32 Index = 0; Index < SegmentStack.Num(); ++Index)
+	{
+		FIsoSegment* Segment = SegmentStack[Index];
+		if (Segment->IsDegenerated())
+		{
+			continue;
+		}
+		if (!Segment->HasTriangleOnLeft())
+		{
+			FindBestTriangle(Segment, false);
+		}
+		if (!Segment->HasTriangleOnRight())
+		{
+			FindBestTriangle(Segment, true);
+		}
+	}
+
+	// Reset the flags "has triangle" of cycle's segments to avoid to block the meshing of next cycles
+	for (FIsoSegment* Segment : Cycle)
+	{
+		Segment->ResetHasTriangle();
+	}
+
+#ifdef ADD_TRIANGLE_2D
+	if (Grid.GetFace()->GetId() == SurfaceToDebug)
+	{
+		Close3DDebugSession();
+	}
+#endif 
+}
+
+/**
+ * The purpose is to add surrounding segments to the small loop to intersection tool to prevent traversing inner segments
+ * A loop is inside inner segments
+ *									|			 |
+ *								   -----------------
+ *									|	 XXX	 |
+ *									|	XXXXX	 |
+ *									|	 XXX	 |
+ *								   -----------------
+ *									|			 |
+ *
+ */
+void FIsoTriangulator::FindInnerGridCellSurroundingSmallLoop()
+{
+	FTimePoint StartTime = FChrono::Now();
+
+	if (GlobalIndexToIsoInnerNodes.Num() == 0)
+	{
+		// No inner node
+		return;
+	}
+
+	// when an internal loop is inside inner UV cell
+	int32 NumU = Grid.GetCuttingCount(EIso::IsoU);
+	int32 NumV = Grid.GetCuttingCount(EIso::IsoV);
+	const TArray<double>& UCoordinates = Grid.GetCuttingCoordinatesAlongIso(EIso::IsoU);
+	const TArray<double>& VCoordinates = Grid.GetCuttingCoordinatesAlongIso(EIso::IsoV);
+
+	const TArray<TArray<FPoint2D>>& Loops = Grid.GetLoops2D(EGridSpace::Default2D);
+	for (int32 LoopIndex = 1; LoopIndex < Loops.Num(); ++LoopIndex)
+	{
+		FPoint2D FirstPoint = Loops[LoopIndex][0];
+
+		int32 IndexU = 0;
+		for (; IndexU < NumU - 1; ++IndexU)
+		{
+			if ((FirstPoint.U > UCoordinates[IndexU]) && (FirstPoint.U < UCoordinates[IndexU + 1] + SMALL_NUMBER))
+			{
+				break;
+			}
+		}
+
+		int32 IndexV = 0;
+		for (; IndexV < NumV - 1; ++IndexV)
+		{
+			if ((FirstPoint.V > VCoordinates[IndexV]) && (FirstPoint.V < VCoordinates[IndexV + 1] + SMALL_NUMBER))
+			{
+				break;
+			}
+		}
+
+		double UMin = UCoordinates[IndexU];
+		double UMax = UCoordinates[IndexU + 1] + SMALL_NUMBER;
+		double VMin = VCoordinates[IndexV];
+		double VMax = VCoordinates[IndexV + 1] + SMALL_NUMBER;
+
+		bool bBoudardyIsSurrounded = true;
+		for (const FPoint2D& LoopPoint : Loops[LoopIndex])
+		{
+			if (LoopPoint.U < UMin || LoopPoint.U > UMax || LoopPoint.V < VMin || LoopPoint.V > VMax)
+			{
+				bBoudardyIsSurrounded = false;
+				break;
+			}
+		}
+
+
+		if (bBoudardyIsSurrounded)
+		{
+			int32 Index = IndexV * NumU + IndexU;
+			IndexOfLowerLeftInnerNodeSurroundingALoop.Add((int32)Index);
+
+			FIsoInnerNode* Node = GlobalIndexToIsoInnerNodes[Index];
+			if (Node == nullptr)
+			{
+				Node = GlobalIndexToIsoInnerNodes[Index + 1];
+			}
+			if (Node != nullptr)
+			{
+				for (FIsoSegment* Segment : Node->GetConnectedSegments())
+				{
+					if (Segment->GetType() == ESegmentType::IsoU)
+					{
+						if (Segment->GetSecondNode().GetIndex() == Index + 1)
+						{
+							InnerSegmentsIntersectionTool.AddSegment(*Segment);
+						}
+					}
+					else if (Segment->GetSecondNode().GetIndex() == Index + NumU)
+					{
+						InnerSegmentsIntersectionTool.AddSegment(*Segment);
+					}
+				}
+			}
+
+			Index = (IndexV + 1) * NumU + IndexU + 1;
+			Node = GlobalIndexToIsoInnerNodes[Index];
+			if (Node == nullptr)
+			{
+				Node = GlobalIndexToIsoInnerNodes[Index - 1];
+			}
+			if (Node != nullptr)
+			{
+				for (FIsoSegment* Segment : Node->GetConnectedSegments())
+				{
+					if (Segment->GetType() == ESegmentType::IsoU)
+					{
+						if (Segment->GetFirstNode().GetIndex() == Index - 1)
+						{
+							InnerSegmentsIntersectionTool.AddSegment(*Segment);
+						}
+					}
+					else if (Segment->GetFirstNode().GetIndex() == Index - NumU)
+					{
+						InnerSegmentsIntersectionTool.AddSegment(*Segment);
+					}
+				}
+			}
+		}
+	}
+
+#ifdef CADKERNEL_DEV
+	Chronos.FindSegmentIsoUVSurroundingSmallLoopDuration += FChrono::Elapse(StartTime);
+#endif
+}
+
+void FIsoTriangulator::TriangulateOverCycle(const EGridSpace Space)
+{
+	FTimePoint StartTime = FChrono::Now();
+
+	TArray<FIsoSegment*> Cycle;
+	Cycle.Reserve(100);
+	TArray<bool> CycleOrientation;
+	CycleOrientation.Reserve(100);
+
+	int32 CycleIndex = 0;
+#ifdef FIND_CYCLE
+	if (Grid.GetFace()->GetId() == SurfaceToDebug)
+	{
+		Open3DDebugSession(TEXT("Triangulate Over Cycle"));
+	}
+#endif
+
+	for (FIsoSegment* Segment : LoopSegments)
+	{
+		if (!Segment->HasCycleOnLeft())
+		{
+			Cycle.Empty();
+			CycleOrientation.Empty();
+			bool bLeftSide = true;
+			FindCycle(Segment, bLeftSide, Cycle, CycleOrientation);
+#ifdef FIND_CYCLE
+			if (Grid.GetFace()->GetId() == SurfaceToDebug)
+			{
+				F3DDebugSession G(TEXT("Find & mesh cycles"));
+				FString Message = FString::Printf(TEXT("MeshCycle - cycle %d"), CycleIndex++);
+				Display(EGridSpace::UniformScaled, *Message, Cycle, false);
+			}
+#endif
+			MeshCycle(Space, Cycle, CycleOrientation);
+		}
+	}
+
+	for (FIsoSegment* Segment : FinalToLoops)
+	{
+		if (!Segment->HasCycleOnLeft())
+		{
+			Cycle.Empty();
+			CycleOrientation.Empty();
+			bool bLeftSide = true;
+			FindCycle(Segment, bLeftSide, Cycle, CycleOrientation);
+#ifdef FIND_CYCLE
+			if (Grid.GetFace()->GetId() == SurfaceToDebug)
+			{
+				Open3DDebugSession(TEXT("Find & mesh cycles"));
+				FString Message = FString::Printf(TEXT("MeshCycle - cycle %d"), CycleIndex++);
+				Display(EGridSpace::UniformScaled, *Message, Cycle, false);
+				Close3DDebugSession();
+			}
+#endif
+			MeshCycle(Space, Cycle, CycleOrientation);
+		}
+
+		if (!Segment->HasCycleOnRight())
+		{
+			Cycle.Empty();
+			CycleOrientation.Empty();
+			bool bLeftSide = false;
+			FindCycle(Segment, bLeftSide, Cycle, CycleOrientation);
+#ifdef FIND_CYCLE
+			if (Grid.GetFace()->GetId() == SurfaceToDebug)
+			{
+				Open3DDebugSession(TEXT("Find & mesh cycles"));
+				FString Message = FString::Printf(TEXT("MeshCycle - cycle %d"), CycleIndex++);
+				Display(EGridSpace::UniformScaled, *Message, Cycle, false);
+				Close3DDebugSession();
+			}
+#endif
+			MeshCycle(Space, Cycle, CycleOrientation);
+		}
+	}
+#ifdef FIND_CYCLE
+	Close3DDebugSession();
+#endif
+
+#ifdef CADKERNEL_DEV
+	Chronos.TriangulateOverCycleDuration = FChrono::Elapse(StartTime);
+#endif
+}
+
+#ifdef FIND_CYCLE
+static int32 CycleId = -1;
+static int32 CycleIndex = 0;
+#endif
+
+void FIsoTriangulator::FindCycle(FIsoSegment* StartSegment, bool LeftSide, TArray<FIsoSegment*>& Cycle, TArray<bool>& CycleOrientation)
+{
+
+#ifdef FIND_CYCLE
+	CycleIndex++;
+	//CycleId = CycleIndex;
+	if (Grid.GetFace()->GetId() == SurfaceToDebug)
+	{
+		CycleIndex = CycleId;
+		bDisplay = true;
+	}
+#endif
+
+	Cycle.Empty();
+	CycleOrientation.Empty();
+
+#ifdef DEBUG_FIND_CYCLE
+	if (CycleId == CycleIndex)
+	{
+		Open3DDebugSession(TEXT("Cycle"));
+	}
+#endif
+
+	FIsoSegment* Segment = StartSegment;
+	FIsoNode* Node;
+
+	if (LeftSide)
+	{
+		Segment->SetHaveCycleOnLeft();
+		Node = &StartSegment->GetSecondNode();
+#ifdef DEBUG_FIND_CYCLE
+		if (CycleId == CycleIndex)
+		{
+			FIsoNode* EndNode = &StartSegment->GetFirstNode();
+			F3DDebugSession _(TEXT("FirstSegment left"));
+			Display(EGridSpace::UniformScaled, *EndNode);
+			Display(EGridSpace::UniformScaled, *StartSegment);
+			//Wait();
+		}
+#endif
+	}
+	else
+	{
+		Segment->SetHaveCycleOnRight();
+		Node = &StartSegment->GetFirstNode();
+#ifdef DEBUG_FIND_CYCLE
+		if (CycleId == CycleIndex)
+		{
+			FIsoNode* EndNode = &StartSegment->GetSecondNode();
+			F3DDebugSession _(TEXT("FirstSegment right"));
+			Display(EGridSpace::UniformScaled, *EndNode);
+			Display(EGridSpace::UniformScaled, *StartSegment);
+			//Wait();
+		}
+#endif
+	}
+
+	Cycle.Add(StartSegment);
+	CycleOrientation.Add(LeftSide);
+	Segment = StartSegment;
+
+	for (;;)
+	{
+		Segment = FindNextSegment(EGridSpace::Default2D, Segment, Node, ClockwiseSlop);
+		if (Segment == nullptr)
+		{
+			Cycle.Empty();
+			break;
+		}
+
+		if (Segment == StartSegment)
+		{
+			break;
+		}
+
+		Cycle.Add(Segment);
+
+		if (&Segment->GetFirstNode() == Node)
+		{
+#ifdef CADKERNEL_DEV
+			if (Segment->HasCycleOnLeft())
+			{
+				F3DDebugSession _(TEXT("Segment HasCycleOnLeft"));
+				Display(EGridSpace::UniformScaled, *Segment);
+				Wait();
+			}
+#endif
+			ensureCADKernel(!Segment->HasCycleOnLeft());
+			Segment->SetHaveCycleOnLeft();
+			Node = &Segment->GetSecondNode();
+			CycleOrientation.Add(true);
+#ifdef DEBUG_FIND_CYCLE
+			if (CycleId == CycleIndex)
+			{
+				F3DDebugSession _(TEXT("Next"));
+				Display(EGridSpace::UniformScaled, *Node);
+				Display(EGridSpace::UniformScaled, *Segment);
+				//Wait();
+			}
+#endif
+		}
+		else
+		{
+#ifdef CADKERNEL_DEV
+			if (Segment->HasCycleOnRight())
+			{
+				F3DDebugSession _(TEXT("Segment HasCycleOnRight"));
+				Display(EGridSpace::UniformScaled, *Segment);
+				Wait();
+			}
+#endif
+			ensureCADKernel(!Segment->HasCycleOnRight());
+			Segment->SetHaveCycleOnRight();
+			Node = &Segment->GetFirstNode();
+			CycleOrientation.Add(false);
+#ifdef DEBUG_FIND_CYCLE
+			if (CycleId == CycleIndex)
+			{
+				F3DDebugSession _(TEXT("Next"));
+				Display(EGridSpace::UniformScaled, *Node);
+				Display(EGridSpace::UniformScaled, *Segment);
+				//Wait();
+			}
+#endif
+		}
+		if (false)
+		{
+			Wait();
+		}
+	}
+#ifdef DEBUG_FIND_CYCLE
+	if (CycleId == CycleIndex)
+	{
+		Close3DDebugSession();
+	}
+#endif
+}
+
+#ifdef DEBUG_FIND_NEXTSEGMENT
+static bool bDisplayStar = false;
+#endif
+
+FIsoSegment* FIsoTriangulator::FindNextSegment(EGridSpace Space, const FIsoSegment* StartSegment, const FIsoNode* StartNode, FGetSlop GetSlop) const
+{
+	const FPoint2D& StartPoint = StartNode->Get2DPoint(Space, Grid);
+	const FPoint2D& EndPoint = (StartNode == &StartSegment->GetFirstNode()) ? StartSegment->GetSecondNode().Get2DPoint(Space, Grid) : StartSegment->GetFirstNode().Get2DPoint(Space, Grid);
+
+#ifdef DEBUG_FIND_NEXTSEGMENT
+	if (bDisplayStar)
+	{
+		Open3DDebugSession(TEXT("FindNextSegment"));
+		F3DDebugSession _(TEXT("Start Segment"));
+		Display(EGridSpace::Default2D, *StartNode);
+		Display(EGridSpace::Default2D, *StartSegment);
+		Display(Space, *StartNode);
+		Display(Space, *StartSegment);
+	}
+#endif
+
+	double ReferenceSlope = ComputePositiveSlope(StartPoint, EndPoint, 0);
+
+	double MaxSlope = 8.1;
+	FIsoSegment* NextSegment = nullptr;
+
+	for (FIsoSegment* Segment : StartNode->GetConnectedSegments())
+	{
+		const FPoint2D& OtherPoint = (StartNode == &Segment->GetFirstNode()) ? Segment->GetSecondNode().Get2DPoint(Space, Grid) : Segment->GetFirstNode().Get2DPoint(Space, Grid);
+
+		double Slope = GetSlop(StartPoint, OtherPoint, ReferenceSlope);
+		if (Slope < SMALL_NUMBER_SQUARE) Slope = 8;
+
+#ifdef DEBUG_FIND_NEXTSEGMENT
+		if (bDisplayStar)
+		{
+			F3DDebugSession _(FString::Printf(TEXT("Segment slop %f"), Slope));
+			Display(EGridSpace::Default2D, *Segment);
+			Display(Space, *Segment);
+		}
+#endif
+
+		if (Slope < MaxSlope || NextSegment == StartSegment)
+		{
+			NextSegment = Segment;
+			MaxSlope = Slope;
+		}
+	}
+
+#ifdef DEBUG_FIND_NEXTSEGMENT
+	if (bDisplayStar)
+	{
+		Close3DDebugSession();
+	}
+#endif
+	return NextSegment;
+}
+
+void FIsoTriangulator::TriangulateInnerNodes()
+{
+	FTimePoint StartTime = FChrono::Now();
+
+	int32 NumU = Grid.GetCuttingCount(EIso::IsoU);
+	int32 NumV = Grid.GetCuttingCount(EIso::IsoV);
+
+#ifdef ADD_TRIANGLE_2D
+	Open3DDebugSession(TEXT("Inner Mesh 2D"));
+#endif
+	for (int32 vIndex = 0, Index = 0; vIndex < NumV - 1; vIndex++)
+	{
+		for (int32 uIndex = 0; uIndex < NumU - 1; uIndex++, Index++)
+		{
+			// Do the lower nodes of the cell exist
+			if (!GlobalIndexToIsoInnerNodes[Index] || !GlobalIndexToIsoInnerNodes[Index + 1])
+			{
+				continue;
+			}
+
+			// Is the lower left node connected
+			if (!GlobalIndexToIsoInnerNodes[Index]->IsLinkedToNextU() || !GlobalIndexToIsoInnerNodes[Index]->IsLinkedToNextV())
+			{
+				continue;
+			}
+
+			// Do the upper nodes of the cell exist
+			int32 OppositIndex = Index + NumU + 1;
+			if (!GlobalIndexToIsoInnerNodes[OppositIndex] || !GlobalIndexToIsoInnerNodes[OppositIndex - 1])
+			{
+				continue;
+			}
+
+			// Is the top right node connected
+			if (!GlobalIndexToIsoInnerNodes[OppositIndex]->IsLinkedToPreviousU() || !GlobalIndexToIsoInnerNodes[OppositIndex]->IsLinkedToPreviousV())
+			{
+				continue;
+			}
+
+			bool bIsSurroundingALoop = false;
+			for (int32 BorderIndex : IndexOfLowerLeftInnerNodeSurroundingALoop)
+			{
+				if (Index == BorderIndex)
+				{
+					bIsSurroundingALoop = true;
+					break;
+				}
+			}
+			if (bIsSurroundingALoop)
+			{
+				continue;
+			}
+
+#ifdef ADD_TRIANGLE_2D
+			DisplayTriangle(EGridSpace::Default2D, *GlobalIndexToIsoInnerNodes[Index], *GlobalIndexToIsoInnerNodes[Index + 1], *GlobalIndexToIsoInnerNodes[OppositIndex]);
+			DisplayTriangle(EGridSpace::Default2D, *GlobalIndexToIsoInnerNodes[OppositIndex], *GlobalIndexToIsoInnerNodes[OppositIndex - 1], *GlobalIndexToIsoInnerNodes[Index]);
+#endif 
+
+			Mesh->AddTriangle(GlobalIndexToIsoInnerNodes[Index]->GetFaceIndex(), GlobalIndexToIsoInnerNodes[Index + 1]->GetFaceIndex(), GlobalIndexToIsoInnerNodes[OppositIndex]->GetFaceIndex());
+			Mesh->AddTriangle(GlobalIndexToIsoInnerNodes[OppositIndex]->GetFaceIndex(), GlobalIndexToIsoInnerNodes[OppositIndex - 1]->GetFaceIndex(), GlobalIndexToIsoInnerNodes[Index]->GetFaceIndex());
+		}
+		Index++;
+	}
+#ifdef ADD_TRIANGLE_2D
+	Close3DDebugSession();
+#endif
+}
+
+// =========================================================================================================================================================================================================
+// =========================================================================================================================================================================================================
+// =========================================================================================================================================================================================================
+//
+//
+//                                                                            NOT YET REVIEWED
+//
+//
+// =========================================================================================================================================================================================================
+// =========================================================================================================================================================================================================
+// =========================================================================================================================================================================================================
+
+void FIsoTriangulator::Delaunay()
+{
+	FTimePoint StartTime = FChrono::Now();
+	//F3DDebugSession DelaunayDebugSession(TEXT("Delaunay Algo"));
+
+	int32 LoopCount = Grid.GetLoops3D().Num();
+
+	TArray<FPoint2D> LoopBarycenters;
+	LoopBarycenters.Reserve(LoopCount + 4);
+	LoopBarycenters.Emplace(FPoint2D::ZeroPoint); // outer loop is not defined
+
+	TArray<DelaunayTriangle> TriangleSet; // Index of the Vertex
+	TriangleSet.Reserve(3 * LoopCount);
+
+	TArray<int32> TriangleIndices;
+	TriangleIndices.Reserve(LoopCount);
+
+	TArray<int32> EdgeVertexIndices;
+	EdgeVertexIndices.Reserve(4 * LoopCount);
+	TArray<int32> OuterEdgeIndices;
+	OuterEdgeIndices.Reserve(2 * LoopCount);
+
+	TArray<int32> OuterLoopIndices;
+
+#ifdef DEBUG_DELAUNAY
+	TFunction<void()> DisplayEdges = [&]()
+	{
+		F3DDebugSession _(TEXT("Edges"));
+		for (int32 Index = 0; Index < EdgeVertexIndices.Num(); Index += 2)
+		{
+			if (OuterEdgeIndices[Index / 2])
+			{
+				DisplaySegment(LoopBarycenters[EdgeVertexIndices[Index]], LoopBarycenters[EdgeVertexIndices[Index + 1]], 0, EVisuProperty::YellowCurve);
+			}
+			else
+			{
+				DisplaySegment(LoopBarycenters[EdgeVertexIndices[Index]], LoopBarycenters[EdgeVertexIndices[Index + 1]]);
+			}
+		}
+
+		for (int32 Index = 0; Index < OuterLoopIndices.Num(); Index += 2)
+		{
+			DisplayPoint(LoopBarycenters[OuterLoopIndices[Index]], EVisuProperty::GreenPoint);
+		}
+	};
+
+	TFunction<void()> DisplayTriangles = [&]()
+	{
+		F3DDebugSession _(TEXT("Triangles"));
+		for (int32 Index = 0; Index < TriangleSet.Num(); Index++)
+		{
+			DisplaySegment(LoopBarycenters[TriangleSet[Index].VertexIndex[0]], LoopBarycenters[TriangleSet[Index].VertexIndex[1]]);
+			DisplaySegment(LoopBarycenters[TriangleSet[Index].VertexIndex[1]], LoopBarycenters[TriangleSet[Index].VertexIndex[2]]);
+			DisplaySegment(LoopBarycenters[TriangleSet[Index].VertexIndex[2]], LoopBarycenters[TriangleSet[Index].VertexIndex[0]]);
+		}
+		//Wait();
+	};
+#endif
+
+	FAABB2D LoobsBBox;
+	const TArray<TArray<FPoint2D>>& Loops = Grid.GetLoops2D(EGridSpace::UniformScaled);
+	for (int32 LoopIndex = 1; LoopIndex < LoopCount; ++LoopIndex)
+	{
+		const TArray<FPoint2D>& LoopPoints = Loops[LoopIndex];
+		FPoint2D& BaryCenter = LoopBarycenters.Emplace_GetRef(FPoint2D::ZeroPoint);
+		for (const FPoint2D& Point : LoopPoints)
+		{
+			BaryCenter += Point;
+		}
+		BaryCenter /= (double)LoopPoints.Num();
+		LoobsBBox += BaryCenter;
+	}
+
+	LoobsBBox.Offset(LoobsBBox.DiagonalLength());
+	LoopBarycenters.Emplace(LoobsBBox.GetCorner(3));
+	LoopBarycenters.Emplace(LoobsBBox.GetCorner(2));
+	LoopBarycenters.Emplace(LoobsBBox.GetCorner(0));
+	LoopBarycenters.Emplace(LoobsBBox.GetCorner(1));
+
+
+#ifdef DEBUG_DELAUNAY
+	{
+		F3DDebugSession _(TEXT("LoopBarycenters"));
+		for (int32 LoopIndex = 1; LoopIndex < LoopCount; ++LoopIndex)
+		{
+			DisplayPoint(LoopBarycenters[LoopIndex], EVisuProperty::YellowPoint);
+		}
+		DisplayAABB2D(LoobsBBox);
+
+		for (int32 LoopIndex = LoopCount; LoopIndex < LoopCount + 4; ++LoopIndex)
+		{
+			DisplayPoint(LoopBarycenters[LoopIndex], EVisuProperty::GreenPoint);
+		}
+	}
+#endif
+
+	TriangleSet.Emplace(LoopCount, LoopCount + 1, LoopCount + 2, LoopBarycenters);
+	TriangleSet.Emplace(LoopCount + 2, LoopCount + 3, LoopCount, LoopBarycenters);
+
+#ifdef DEBUG_DELAUNAY
+	DisplayTriangles();
+#endif
+
+	for (int32 CenterIndex = 1; CenterIndex < LoopCount; ++CenterIndex)
+	{
+		TriangleIndices.Empty(LoopCount);
+		const FPoint2D& NewVertex = LoopBarycenters[CenterIndex];
+		for (int32 Tndex = 0; Tndex < TriangleSet.Num(); Tndex++)
+		{
+			const FPoint2D Center = TriangleSet[Tndex].Center;
+			double SquareRadius = TriangleSet[Tndex].SquareRadius;
+			double SquareDistanceToCenter = Center.SquareDistance(NewVertex);
+			if (SquareDistanceToCenter < SquareRadius)
+			{
+				TriangleIndices.Add(Tndex);
+			}
+		}
+
+		EdgeVertexIndices.Empty(LoopCount);
+		for (int32 TriangleIndex : TriangleIndices)
+		{
+			int32 EndVertex = TriangleSet[TriangleIndex].VertexIndex[2];
+			for (int32 Index = 0; Index < 3; Index++)
+			{
+				int32 StartVertex = TriangleSet[TriangleIndex].VertexIndex[Index];
+				int32 Endex = 0;
+				for (; Endex < EdgeVertexIndices.Num(); Endex += 2)
+				{
+					if (EdgeVertexIndices[Endex] == EndVertex && EdgeVertexIndices[Endex + 1] == StartVertex)
+					{
+						EdgeVertexIndices[Endex] = -1;
+						EdgeVertexIndices[Endex + 1] = -1;
+						break;
+					}
+				}
+				if (Endex == EdgeVertexIndices.Num())
+				{
+					EdgeVertexIndices.Add(StartVertex);
+					EdgeVertexIndices.Add(EndVertex);
+				}
+				EndVertex = StartVertex;
+			}
+		}
+
+		int32 Endex = 0;
+		for (int32 Tndex : TriangleIndices)
+		{
+			while (EdgeVertexIndices[Endex] < 0)
+			{
+				Endex += 2;
+			}
+			TriangleSet[Tndex].Set(EdgeVertexIndices[Endex + 1], EdgeVertexIndices[Endex], CenterIndex, LoopBarycenters);
+			Endex += 2;
+		}
+		for (; Endex < EdgeVertexIndices.Num(); Endex += 2)
+		{
+			if (EdgeVertexIndices[Endex] < 0)
+			{
+				continue;
+			}
+			TriangleSet.Emplace(EdgeVertexIndices[Endex + 1], EdgeVertexIndices[Endex], CenterIndex, LoopBarycenters);
+		}
+#ifdef DEBUG_DELAUNAY
+		DisplayTriangles();
+#endif
+	}
+
+#ifdef DEBUG_DELAUNAY
+	Wait();
+
+	DisplayTriangles();
+	Wait();
+#endif
+
+	EdgeVertexIndices.Reserve(LoopCount);
+
+	EdgeVertexIndices.Empty(LoopCount);
+	for (int32 TriangleIndex = 0; TriangleIndex < TriangleSet.Num(); TriangleIndex++)
+	{
+		int32 EndVertex = TriangleSet[TriangleIndex].VertexIndex[2];
+		int32 Index = 0;
+		for (; Index < 3; Index++)
+		{
+			if(TriangleSet[TriangleIndex].VertexIndex[Index] >= LoopCount)
+			{
+				break;
+			}
+		}
+
+		if (Index < 3)
+		{
+			continue;
+		}
+			
+		for (Index = 0; Index < 3; Index++)
+		{
+			int32 StartVertex = TriangleSet[TriangleIndex].VertexIndex[Index];
+			int32 Endex = 0;
+			for (; Endex < EdgeVertexIndices.Num(); Endex += 2)
+			{
+				if (EdgeVertexIndices[Endex] == EndVertex && EdgeVertexIndices[Endex + 1] == StartVertex)
+				{
+					OuterEdgeIndices[Endex/2] = 0;
+					break;
+				}
+			}
+			if (Endex == EdgeVertexIndices.Num())
+			{
+				EdgeVertexIndices.Add(StartVertex);
+				EdgeVertexIndices.Add(EndVertex);
+				OuterEdgeIndices.Add(1);
+			}
+			EndVertex = StartVertex;
+		}
+	}
+
+	OuterLoopIndices.Reserve(LoopCount * 2);
+	for (int32 Index = 0; Index < EdgeVertexIndices.Num(); Index += 2)
+	{
+		if (OuterEdgeIndices[Index / 2])
+		{
+			OuterLoopIndices.Add(EdgeVertexIndices[Index]);
+			OuterLoopIndices.Add(EdgeVertexIndices[Index + 1]);
+		}
+	}
+	OuterLoopIndices.Sort();
+
+#ifdef DEBUG_DELAUNAY
+	DisplayEdges();
+	Wait();
+#endif
+
+	TFunction<void(FLoopNode&, const FPoint2D&, FLoopNode&, const FPoint2D&)> TryCreateSegment = [&](FLoopNode& NodeA, const FPoint2D& ACoordinates, FLoopNode& NodeB, const FPoint2D& BCoordinates)
+	{
+		if (NodeA.GetSegmentConnectedTo(&NodeB))
+		{
+			return;
+		}
+
+		if (InnerSegmentsIntersectionTool.DoesIntersect(ACoordinates, BCoordinates))
+		{
+			return;
+		}
+
+		if (InnerToLoopSegmentsIntersectionTool.DoesIntersect(NodeA, NodeB))
+		{
+			return;
+		}
+
+		if (LoopSegmentsIntersectionTool.DoesIntersect(NodeA, NodeB))
+		{
+			return;
+		}
+
+		// Is Outside and not to flat at NodeA
+		const double FlatAngle = 1;
+		if (!IsPointPIsInsideSectorABC(NodeA.GetPreviousNode().Get2DPoint(EGridSpace::Default2D, Grid), ACoordinates, NodeA.GetNextNode().Get2DPoint(EGridSpace::Default2D, Grid), BCoordinates, FlatAngle))
+		{
+			return;
+		}
+
+		// Is Outside and not to flat at NodeB
+		if (!IsPointPIsInsideSectorABC(NodeB.GetPreviousNode().Get2DPoint(EGridSpace::Default2D, Grid), BCoordinates, NodeB.GetNextNode().Get2DPoint(EGridSpace::Default2D, Grid), ACoordinates, FlatAngle))
+		{
+			return;
+		}
+
+		FIsoSegment& Segment = IsoSegmentFactory.New();
+		Segment.Init(NodeA, NodeB, ESegmentType::LoopToLoop);
+		CandidateLoopToLoopSegments.Add(&Segment);
+#ifdef DEBUG_DELAUNAY
+		DisplaySegment(NodeA.Get2DPoint(EGridSpace::UniformScaled, Grid), NodeB.Get2DPoint(EGridSpace::UniformScaled, Grid));
+#endif
+	};
+
+	TFunction<void(int32, int32)> FindAndTryCreateCandidateSegmentToLinkLoops = [&](int32 IndexLoopA, int32 IndexLoopB)
+	{
+		int32 StartIndexLoopA = LoopStartIndex[IndexLoopA];
+		int32 StartIndexLoopB = LoopStartIndex[IndexLoopB];
+
+		double MinDistanceSquare = HUGE_VALUE_SQUARE;
+		int32 MinIndexA = -1;
+		int32 MinIndexB = -1;
+		for (int32 IndexA = StartIndexLoopA; IndexA < LoopNodes.Num() - 1; ++IndexA)
+		{
+			FLoopNode& NodeA = LoopNodes[IndexA];
+			if (NodeA.GetLoopIndex() != IndexLoopA)
+			{
+				break;
+			}
+
+			const FPoint2D& ACoordinates = NodeA.Get2DPoint(EGridSpace::UniformScaled, Grid);
+
+			for (int32 IndexB = StartIndexLoopB; IndexB < LoopNodes.Num(); ++IndexB)
+			{
+				FLoopNode& NodeB = LoopNodes[IndexB];
+				if (NodeB.GetLoopIndex() != IndexLoopB)
+				{
+					break;
+				}
+				const FPoint2D& BCoordinates = NodeB.Get2DPoint(EGridSpace::UniformScaled, Grid);
+
+				double SquareDistance = ACoordinates.SquareDistance(BCoordinates);
+				if (SquareDistance < MinDistanceSquare)
+				{
+					MinDistanceSquare = SquareDistance;
+					MinIndexA = IndexA;
+					MinIndexB = IndexB;
+				}
+			}
+		}
+
+		if (MinIndexA >= 0 && MinIndexB >= 0)
+		{
+			FLoopNode& NodeA = LoopNodes[MinIndexA];
+			const FPoint2D& ACoordinates = NodeA.Get2DPoint(EGridSpace::Default2D, Grid);
+			FLoopNode& NodeB = LoopNodes[MinIndexB];
+			const FPoint2D& BCoordinates = NodeB.Get2DPoint(EGridSpace::Default2D, Grid);
+			TryCreateSegment(NodeA, ACoordinates, NodeB, BCoordinates);
+		}
+
+	};
+
+	// Between inner loops
+	{
+#ifdef DEBUG_DELAUNAY
+		Open3DDebugSession(TEXT("Build Segments"));
+#endif
+		for (int32 Index = 0; Index < EdgeVertexIndices.Num(); Index += 2)
+		{
+			int32 IndexLoopA = EdgeVertexIndices[Index];
+			int32 IndexLoopB = EdgeVertexIndices[Index + 1];
+
+			FindAndTryCreateCandidateSegmentToLinkLoops(IndexLoopA, IndexLoopB);
+		}
+#ifdef DEBUG_DELAUNAY
+		Close3DDebugSession();
+#endif
+	}
+
+	// With Outer loop
+	{
+#ifdef DEBUG_DELAUNAY
+		Open3DDebugSession(TEXT("Build Segments"));
+#endif
+		for (int32 Index = 0; Index < OuterLoopIndices.Num(); Index += 2)
+		{
+			int32 IndexLoopA = OuterLoopIndices[Index];
+			int32 IndexLoopB = 0;
+
+			FindAndTryCreateCandidateSegmentToLinkLoops(IndexLoopA, IndexLoopB);
+		}
+#ifdef DEBUG_DELAUNAY
+		Close3DDebugSession();
+#endif
+	}
+
+#ifdef CADKERNEL_DEV
+	Chronos.FindSegmentToLinkLoopToLoopByDelaunayDuration += FChrono::Elapse(StartTime);
+#endif
+}
+
+// =========================================================================================================================================================================================================
+// =========================================================================================================================================================================================================
+// =========================================================================================================================================================================================================
+//
+//
+//                                                                            NOT YET REVIEWED
+//
+//
+// =========================================================================================================================================================================================================
+// =========================================================================================================================================================================================================
+// =========================================================================================================================================================================================================
+
+
+#ifdef NEED_TO_CHECK_USEFULNESS
+void FIsoTriangulator::CompleteIsoSegmentLoopToLoop()
+{
+	FIntersectionSegmentTool NewSegmentsIntersectionTool(Grid);
+	// similare to FindSegmentToLinkloopToloop::CreateSegment but different 
+	TFunction<void(FLoopNode&, FIsoNode&)> CreateComplementarySegment = [&](FLoopNode& Node1, FIsoNode& Node2)
+	{
+		if (Node1.GetSegmentConnectedTo(&Node2))
+		{
+			return;
+		}
+
+		if (NewSegmentsIntersectionTool.DoesIntersect(Node1, Node2))
+		{
+			return;
+		}
+
+		if (LoopSegmentsIntersectionTool.DoesIntersect(Node1, Node2))
+		{
+			return;
+		}
+
+		FIsoSegment& Segment = IsoSegmentFactory.New();
+		Segment.Init(Node1, Node2, ESegmentType::Unknown);
+		Segment.ConnectToNode();
+		FinalToLoops.Add(&Segment);
+		InnerToLoopSegmentsIntersectionTool.AddSegment(Segment);
+		NewSegmentsIntersectionTool.AddSegment(Segment);
+		NewSegmentsIntersectionTool.Sort();
+	};
+
+	for (FLoopNode& LoopNode : LoopNodes)
+	{
+		if (LoopNode.GetConnectedSegments().Num() == 2)
+		{
+			if (LoopNode.GetPreviousNode().GetConnectedSegments().Num() == 3)
+			{
+				if ((LoopNode.GetPreviousNode().GetConnectedSegments().Last()->GetType() == ESegmentType::IsoLoopToLoop) && (LoopNode.GetNextNode().GetConnectedSegments().Last()->GetType() == ESegmentType::IsoLoopToLoop))
+				{
+					FIsoNode& Node1 = LoopNode.GetPreviousNode().GetConnectedSegments().Last()->GetOtherNode(&LoopNode.GetPreviousNode());
+					DisplaySegment(Node1.Get2DPoint(EGridSpace::Default2D, Grid), LoopNode.Get2DPoint(EGridSpace::Default2D, Grid));
+					CreateComplementarySegment(LoopNode, Node1);
+					FIsoNode& Node2 = LoopNode.GetNextNode().GetConnectedSegments().Last()->GetOtherNode(&LoopNode.GetNextNode());
+					CreateComplementarySegment(LoopNode, Node2);
+					DisplaySegment(Node2.Get2DPoint(EGridSpace::Default2D, Grid), LoopNode.Get2DPoint(EGridSpace::Default2D, Grid));
+				}
+				else if (LoopNode.GetPreviousNode().GetConnectedSegments().Last()->GetType() == ESegmentType::IsoLoopToLoop)
+				{
+					FIsoNode& Node = LoopNode.GetPreviousNode().GetConnectedSegments().Last()->GetOtherNode(&LoopNode.GetPreviousNode());
+					CreateComplementarySegment(LoopNode, Node);
+					DisplaySegment(Node.Get2DPoint(EGridSpace::Default2D, Grid), LoopNode.Get2DPoint(EGridSpace::Default2D, Grid));
+				}
+				else if (LoopNode.GetNextNode().GetConnectedSegments().Last()->GetType() == ESegmentType::IsoLoopToLoop)
+				{
+					FIsoNode& Node = LoopNode.GetNextNode().GetConnectedSegments().Last()->GetOtherNode(&LoopNode.GetNextNode());
+					CreateComplementarySegment(LoopNode, Node);
+					DisplaySegment(Node.Get2DPoint(EGridSpace::Default2D, Grid), LoopNode.Get2DPoint(EGridSpace::Default2D, Grid));
+				}
+			}
+		}
+	}
+}
+#endif
+
+#ifdef UNUSED
+void FIsoTriangulator::FindSegmentToLinkLoopToLoop()
+{
+	FTimePoint StartTime = FChrono::Now();
+
+	double DiagonalMax = Grid.GetMaxDeltaU(EIso::IsoU) + Grid.GetMaxDeltaU(EIso::IsoV);
+	DiagonalMax *= 2;
+
+	TFunction<double(const TArray<double>&, int32&, const double&)> GetDeltaCoordinate = [](const TArray<double>& Iso, int32& Index, const double& PointCoord)
+	{
+		while (PointCoord + SMALL_NUMBER < Iso[Index])
+		{
+			if (Index == 0)
+			{
+				break;
+			}
+			Index--;
+		}
+		for (; Index < Iso.Num(); ++Index)
+		{
+			if (Index == Iso.Num() - 2)
+			{
+				break;
+			}
+			if (PointCoord < Iso[Index + 1])
+			{
+				break;
+			}
+		}
+
+		// If PointCoord is closer Iso[Index + 1]
+		if ((Iso[Index + 1] - PointCoord) * 2 < (Iso[Index + 1] - Iso[Index]))
+		{
+			if (Index + 2 < Iso.Num())
+			{
+				return FMath::Max(Iso[Index + 1] - Iso[Index], Iso[Index + 2] - Iso[Index + 1]);
+			}
+			else
+			{
+				return Iso[Index + 1] - Iso[Index];
+			}
+		}
+		else
+		{
+			if (Index > 1)
+			{
+				return FMath::Max(Iso[Index] - Iso[Index - 1], Iso[Index + 1] - Iso[Index]);
+			}
+			else
+			{
+				return Iso[Index + 1] - Iso[Index];
+			}
+		}
+		return Iso[Index + 1] - Iso[Index];
+	};
+
+	TFunction<void(FLoopNode*, const FPoint2D&, FLoopNode*, const FPoint2D&)> CreateSegment = [this](FLoopNode* NodeA, const FPoint2D& ACoordinates, FLoopNode* NodeB, const FPoint2D& BCoordinates)
+	{
+		if (NodeA->GetSegmentConnectedTo(NodeB))
+		{
+			return;
+		}
+
+		// IsExiste
+		if (NodeA->GetLoopIndex() == NodeB->GetLoopIndex())
+		{
+			if ((NodeA->GetIndex() == NodeB->GetIndex() + 1) || (NodeB->GetIndex() == NodeA->GetIndex() + 1))
+			{
+				return;
+			}
+		}
+
+		if (InnerSegmentsIntersectionTool.DoesIntersect(ACoordinates, BCoordinates))
+		{
+			return;
+		}
+
+		if (InnerToLoopSegmentsIntersectionTool.DoesIntersect(*NodeA, *NodeB))
+		{
+			return;
+		}
+
+		if (LoopSegmentsIntersectionTool.DoesIntersect(*NodeA, *NodeB))
+		{
+			return;
+		}
+
+		// Is Outside and not to flat at NodeA
+		const double FlatAngle = 1;
+		if (!IsPointPIsInsideSectorABC(NodeA->GetPreviousNode().Get2DPoint(EGridSpace::Default2D, Grid), ACoordinates, NodeA->GetNextNode().Get2DPoint(EGridSpace::Default2D, Grid), BCoordinates, FlatAngle))
+		{
+			return;
+		}
+
+		//double SlopWithNextloop = GetZeroTwoPiSlope(ACoordinates, NodeA->GetNextNode().Get2DPoint(EGridSpace::Default2D, Grid), 0);
+		//double loopDeltaSlop = GetZeroTwoPiSlope(ACoordinates, NodeA->GetPreviousNode().Get2DPoint(EGridSpace::Default2D, Grid), SlopWithNextloop);
+		//double SegmentSlop = GetZeroTwoPiSlope(ACoordinates, BCoordinates, SlopWithNextloop);
+		//if (SegmentSlop < FlatAngle || SegmentSlop + FlatAngle > loopDeltaSlop)
+		//{
+		//	return;
+		//}
+
+		// Is Outside and not to flat at NodeB
+		if (!IsPointPIsInsideSectorABC(NodeB->GetPreviousNode().Get2DPoint(EGridSpace::Default2D, Grid), BCoordinates, NodeB->GetNextNode().Get2DPoint(EGridSpace::Default2D, Grid), ACoordinates, FlatAngle))
+		{
+			return;
+		}
+
+		//SlopWithNextloop = GetZeroTwoPiSlope(BCoordinates, NodeB->GetNextNode().Get2DPoint(EGridSpace::Default2D, Grid), 0);
+		//loopDeltaSlop = GetZeroTwoPiSlope(BCoordinates, NodeB->GetPreviousNode().Get2DPoint(EGridSpace::Default2D, Grid), SlopWithNextloop);
+		//SegmentSlop = GetZeroTwoPiSlope(BCoordinates, ACoordinates, SlopWithNextloop);
+
+		//if (SegmentSlop < FlatAngle || SegmentSlop + FlatAngle > loopDeltaSlop)
+		//{
+		//	return;
+		//}
+
+		FIsoSegment& Segment = IsoSegmentFactory.New();
+		Segment.Init(*NodeA, *NodeB, ESegmentType::LoopToLoop);
+		CandidateLoopToLoopSegments.Add(&Segment);
+	};
+
+	//Open3DDebugSession(TEXT("FIsoTrianguler::FindSegmentToLinkloopToloop"));
+
+	// find closed points
+	int32 IndexU = 0;
+	int32 IndexV = 0;
+
+	for (int32 IndexA = 0; IndexA < SortedLoopNodes.Num() - 1; ++IndexA)
+	{
+		FLoopNode* NodeA = SortedLoopNodes[IndexA];
+		if (NodeA->GetConnectedSegments().Num() > 2)
+		{
+			// Enough connected
+			continue;
+		}
+
+		// check intersection with
+		const FPoint2D& ACoordinates = NodeA->Get2DPoint(EGridSpace::Default2D, Grid);
+
+		double DeltaUMax = GetDeltaCoordinate(Grid.GetCuttingCoordinatesAlongIso(EIso::IsoU), IndexU, ACoordinates.U);
+		//double DeltaUMin = DeltaUMax * 0.1;
+		DeltaUMax *= 1.75;
+		double DeltaVMax = GetDeltaCoordinate(Grid.GetCuttingCoordinatesAlongIso(EIso::IsoV), IndexV, ACoordinates.V);
+		//double DeltaVMin = DeltaVMax * 0.1;
+		DeltaVMax *= 1.75;
+
+		for (int32 IndexB = IndexA + 1; IndexB < SortedLoopNodes.Num(); ++IndexB)
+		{
+			FLoopNode* NodeB = SortedLoopNodes[IndexB];
+			const FPoint2D& BCoordinates = NodeB->Get2DPoint(EGridSpace::Default2D, Grid);
+			FPoint2D AB = BCoordinates - ACoordinates;
+			double Diagonal = FMath::Abs(AB.U + AB.V);
+			if (Diagonal > DiagonalMax)
+			{
+				break;
+			}
+
+			if (FMath::Abs(AB.U) < DeltaUMax && FMath::Abs(AB.V) < DeltaVMax)
+			{
+				CreateSegment(NodeA, ACoordinates, NodeB, BCoordinates);
+			}
+		}
+	}
+	//Close3DDebugSession();
+	Chronos.FindSegmentToLinkLoopToLoopDuration += FChrono::Elapse(StartTime);
+}
+#endif
