CADKernel Source Code v0.01

diff --git a/Engine/Source/Runtime/Datasmith/CADKernel/Private/CADKernel/Mesh/Meshers/IsoTriangulator/DefineForDebug.h b/Engine/Source/Runtime/Datasmith/CADKernel/Private/CADKernel/Mesh/Meshers/IsoTriangulator/DefineForDebug.h
new file mode 100644
--- /dev/null
+++ b/Engine/Source/Runtime/Datasmith/CADKernel/Private/CADKernel/Mesh/Meshers/IsoTriangulator/DefineForDebug.h
@@ -0,0 +1,20 @@
+// Copyright Epic Games, Inc. All Rights Reserved.
+// Refactor in progress
+#pragma once
+
+
+static int32 SurfaceToDebug = 4;
+#define SURFACE_TO_DEBUG SurfaceToDebug
+//#define SURFACE_TO_DEBUG 0
+static bool BoolDisplayDebugMeshStep=false;
+
+//#define DisplayDebugMeshSurfaceBySurfaceStep
+
+//#define DISPLAY_CRITERIA_GRID
+
+//#define FIND_CYCLE
+//#define DEBUG_FIND_CYCLE
+
+
+//#define DEBUG_MESH_EDGE
+//#define DEBUG_GetPreferredUVCoordinatesFromNeighbours
\ No newline at end of file
diff --git a/Engine/Source/Runtime/Datasmith/CADKernel/Private/CADKernel/Mesh/Meshers/IsoTriangulator/IntersectionSegmentTool.cpp b/Engine/Source/Runtime/Datasmith/CADKernel/Private/CADKernel/Mesh/Meshers/IsoTriangulator/IntersectionSegmentTool.cpp
new file mode 100644
--- /dev/null
+++ b/Engine/Source/Runtime/Datasmith/CADKernel/Private/CADKernel/Mesh/Meshers/IsoTriangulator/IntersectionSegmentTool.cpp
@@ -0,0 +1,117 @@
+// Copyright Epic Games, Inc. All Rights Reserved.
+
+#include "CADKernel/Mesh/Meshers/IsoTriangulator/IntersectionSegmentTool.h"
+
+#include "CADKernel/Math/Geometry.h"
+#include "CADKernel/Mesh/Meshers/IsoTriangulator/IsoNode.h"
+#include "CADKernel/Mesh/Meshers/IsoTriangulator/IsoSegment.h"
+#include "CADKernel/Mesh/Structure/Grid.h"
+#include "CADKernel/Utils/Util.h"
+
+using namespace CADKernel;
+
+bool FIntersectionSegmentTool::DoesIntersect(const FIsoSegment& Segment) const
+{
+	return DoesIntersect(Segment.GetFirstNode(), Segment.GetSecondNode());
+}
+
+bool FIntersectionSegmentTool::DoesIntersect(const FIsoNode& StartNode, const FIsoNode& EndNode) const
+{
+	const FPoint2D& StartPoint = StartNode.Get2DPoint(EGridSpace::Default2D, Grid);
+	const FPoint2D& EndPoint = EndNode.Get2DPoint(EGridSpace::Default2D, Grid);
+	const FSegment4IntersectionTools StartEndSegment(StartPoint, EndPoint);
+
+	double AxisMax = StartEndSegment.Boundary[EIso::IsoU].Max + StartEndSegment.Boundary[EIso::IsoV].Max;
+
+	for (const FSegment4IntersectionTools& Segment : Segments)
+	{
+		if (AxisMax < Segment.AxisMin)
+		{
+			break;
+		}
+
+		if (&Segment.IsoSegment->GetFirstNode() == &StartNode || &Segment.IsoSegment->GetSecondNode() == &StartNode)
+		{
+			continue;
+		}
+
+		if (&Segment.IsoSegment->GetFirstNode() == &EndNode || &Segment.IsoSegment->GetSecondNode() == &EndNode)
+		{
+			continue;
+		}
+
+		if (IntersectSegments2D(Segment.Segment2D, StartEndSegment.Segment2D))
+		{
+			return true;
+		}
+	}
+
+	return false;
+}
+
+bool FIntersectionSegmentTool::DoesIntersect(const FIsoNode& StartNode, const FPoint2D& EndPoint) const
+{
+	const FPoint2D& StartPoint = StartNode.Get2DPoint(EGridSpace::Default2D, Grid);
+	const FSegment4IntersectionTools StartEndSegment(StartPoint, EndPoint);
+
+	double DMax = StartEndSegment.Boundary[EIso::IsoU].Max + StartEndSegment.Boundary[EIso::IsoV].Max;
+
+	for (const FSegment4IntersectionTools& Segment : Segments)
+	{
+		if (DMax < Segment.AxisMin)
+		{
+			break;
+		}
+
+		if (&Segment.IsoSegment->GetFirstNode() == &StartNode || &Segment.IsoSegment->GetSecondNode() == &StartNode)
+		{
+			continue;
+		}
+
+		if (!StartEndSegment.CouldItIntersect(Segment))
+		{
+			continue;
+		}
+
+		if (IntersectSegments2D(Segment.Segment2D, StartEndSegment.Segment2D))
+		{
+			return true;
+		}
+	}
+
+	return false;
+}
+
+bool FIntersectionSegmentTool::DoesIntersect(const FPoint2D& StartPoint, const FPoint2D& EndPoint) const
+{
+	const FSegment4IntersectionTools StartEndSegment(StartPoint, EndPoint);
+	double DMax = StartEndSegment.Boundary[EIso::IsoU].Max + StartEndSegment.Boundary[EIso::IsoV].Max;
+
+	for (const FSegment4IntersectionTools& Segment : Segments)
+	{
+		if (DMax < Segment.AxisMin)
+		{
+			break;
+		}
+
+		if (!StartEndSegment.CouldItIntersect(Segment))
+		{
+			continue;
+		}
+
+		if (IntersectSegments2D(Segment.Segment2D, StartEndSegment.Segment2D))
+		{
+			return true;
+		}
+	}
+	return false;
+}
+
+FSegment4IntersectionTools::FSegment4IntersectionTools(const FGrid& Grid, const FIsoSegment& InSegment)
+	: FSegment4IntersectionTools(InSegment.GetFirstNode().Get2DPoint(EGridSpace::Default2D, Grid), InSegment.GetSecondNode().Get2DPoint(EGridSpace::Default2D, Grid))
+{
+	IsoSegment = &InSegment;
+	AxisMin = Boundary[EIso::IsoU].Min + Boundary[EIso::IsoV].Min;
+}
+
+
diff --git a/Engine/Source/Runtime/Datasmith/CADKernel/Private/CADKernel/Mesh/Meshers/IsoTriangulator/IntersectionSegmentTool.h b/Engine/Source/Runtime/Datasmith/CADKernel/Private/CADKernel/Mesh/Meshers/IsoTriangulator/IntersectionSegmentTool.h
new file mode 100644
--- /dev/null
+++ b/Engine/Source/Runtime/Datasmith/CADKernel/Private/CADKernel/Mesh/Meshers/IsoTriangulator/IntersectionSegmentTool.h
@@ -0,0 +1,128 @@
+// Copyright Epic Games, Inc. All Rights Reserved.
+#pragma once
+
+#include "CADKernel/Core/Types.h"
+#include "CADKernel/Math/Boundary.h"
+#include "CADKernel/Math/Geometry.h"
+#include "CADKernel/UI/Message.h"
+
+namespace CADKernel
+{
+	class FGrid;
+	class FIsoNode;
+	class FIsoSegment;
+	class FPoint2D;
+
+	struct FSegment4IntersectionTools
+	{
+		const TSegment<FPoint2D> Segment2D;
+
+		/**
+		 * Segment's axis aligned bounding box
+		 */
+		const FSurfacicBoundary Boundary;
+
+		const FIsoSegment* IsoSegment;
+
+		/**
+		 * Uses as criterion to sort segments to optimize 
+		 * AxisMin = Boundary[EIso::IsoU].Min + Boundary[EIso::IsoV].Min 
+		 */
+		double AxisMin;
+
+		FSegment4IntersectionTools(const FGrid& Grid, const FIsoSegment& InSegment);
+
+		FSegment4IntersectionTools(const FPoint2D& StartPoint, const FPoint2D& EndPoint)
+			: Segment2D(StartPoint, EndPoint)
+			, Boundary(StartPoint, EndPoint)
+			, IsoSegment(nullptr)
+		{
+			AxisMin = Boundary[EIso::IsoU].Min + Boundary[EIso::IsoV].Min;
+		}
+
+		/**
+		 * @return false if their Boundaries are not intersecting 
+		 */
+		bool CouldItIntersect(const FSegment4IntersectionTools& Segment) const
+		{
+			if ((Boundary[EIso::IsoU].Min > Segment.Boundary[EIso::IsoU].Max) || (Boundary[EIso::IsoV].Min > Segment.Boundary[EIso::IsoV].Max))
+			{
+				return false;
+			}
+
+			if ((Boundary[EIso::IsoU].Max < Segment.Boundary[EIso::IsoU].Min) || (Boundary[EIso::IsoV].Max < Segment.Boundary[EIso::IsoV].Min))
+			{
+				return false;
+			}
+			return true;
+		}
+	};
+
+	class FIntersectionSegmentTool
+	{
+		TArray<FSegment4IntersectionTools> Segments;
+		const FGrid& Grid;
+
+	public:
+		FIntersectionSegmentTool(const FGrid& InGrid)
+			: Grid(InGrid)
+		{
+		}
+
+		void Reserve(int32 InMaxNum)
+		{
+			Segments.Reserve(InMaxNum);
+		}
+
+		void SetNum(int32 NewMaxNum)
+		{
+			Segments.Reserve(NewMaxNum);
+		}
+
+		int32 Num()
+		{
+			return Segments.Num();
+		}
+
+		void AddSegment(const TArray<FIsoSegment*>& InNewSegments)
+		{
+			Segments.Reserve(InNewSegments.Num() + Segments.Num());
+			for (FIsoSegment* Segment : InNewSegments)
+			{
+				AddSegment(*Segment);
+			}
+		}
+
+		void AddSegment(const FIsoSegment& Segment)
+		{
+			Segments.Emplace(Grid, Segment);
+		}
+
+		bool DoesIntersect(const FIsoSegment& Segment) const;
+
+		bool DoesIntersect(const FPoint2D& StartPoint, const FPoint2D& EndPoint) const;
+
+		/**
+		 * Allow StartNode to be connected to one segment
+		 */
+		bool DoesIntersect(const FIsoNode& StartNode, const FPoint2D& EndPoint) const;
+
+		/**
+		 * Allow StartNode and EndNode to be connected to one segment
+		 */
+		bool DoesIntersect(const FIsoNode& StartNode, const FIsoNode& EndNode) const;
+
+		/**
+		 * segments are sorted by DMin increasing
+		 */
+		void Sort()
+		{
+			Algo::Sort(Segments, [](const FSegment4IntersectionTools& Segment1, const FSegment4IntersectionTools& Segment2) { return Segment1.AxisMin < Segment2.AxisMin; });
+		}
+
+#ifdef CADKERNEL_DEV
+		void Display(const TCHAR* Message) const;
+#endif
+	};
+} // namespace CADKernel
+
diff --git a/Engine/Source/Runtime/Datasmith/CADKernel/Private/CADKernel/Mesh/Meshers/IsoTriangulator/IsoNode.h b/Engine/Source/Runtime/Datasmith/CADKernel/Private/CADKernel/Mesh/Meshers/IsoTriangulator/IsoNode.h
new file mode 100644
--- /dev/null
+++ b/Engine/Source/Runtime/Datasmith/CADKernel/Private/CADKernel/Mesh/Meshers/IsoTriangulator/IsoNode.h
@@ -0,0 +1,427 @@
+// Copyright Epic Games, Inc. All Rights Reserved.
+#pragma once
+
+#include "CADKernel/Core/Types.h"
+#include "CADKernel/Mesh/MeshEnum.h"
+#include "CADKernel/Mesh/Structure/Grid.h"
+
+namespace CADKernel
+{
+	class FGrid;
+	class FIsoSegment;
+	class FPoint;
+	class FPoint2D;
+
+	enum class EIsoLink : uint8
+	{
+		IsoUNext = 0,
+		IsoVNext = 1,
+		IsoUPrevious = 2,
+		IsoVPrevious = 3,
+	};
+
+	enum class EIsoNodeStates : uint16
+	{
+		None = 0x0000u,	// No flags.
+
+		LinkedToPreviousU = 0x0001u,
+		LinkedToNextU     = 0x0002u,
+		LinkedToPreviousV = 0x0004u,
+		LinkedToNextV     = 0x0008u,
+
+		TriangleComplete  = 0x000Fu,
+
+		LinkedToLoop      = 0x0010u,
+		InnerMeshLoop     = 0x0020u,
+
+		FirstQuarter      = 0x0100u,
+		SecondQuarter     = 0x0200u,
+		ThirdQuarter      = 0x0400u,
+		FourthQuarter     = 0x0800u,
+
+		NearlyIsoU        = 0x1000u,  // Loop nodes
+		NearlyIsoV        = 0x2000u,  // Loop nodes
+
+		All               = 0xFFFFu
+	};
+
+	ENUM_CLASS_FLAGS(EIsoNodeStates);
+
+	/**
+	 * Inner Node of the TopologicalFace, these nodes are IsoNodes as they are build according to the TopologicalFace iso cutting coordinates UV
+	 */
+	class FIsoNode
+	{
+	protected:
+		TArray<FIsoSegment*> ConnectedSegments;
+		EIsoNodeStates States;
+
+		int32 Index; // Index of the node either in loop nodes either in inner nodes
+		int32 FaceIndex; // Index of the node in the face
+		int32 Id;
+
+	public:
+		FIsoNode(int32 InNodeIndex, int32 InFaceIndex, int32 InNodeId)
+			: ConnectedSegments()
+			, States(EIsoNodeStates::None)
+			, Index(InNodeIndex)
+			, FaceIndex(InFaceIndex)
+			, Id(InNodeId)
+		{
+			ConnectedSegments.Reserve(5);
+		};
+
+		virtual ~FIsoNode() = default;
+
+		const int32 GetIndex() const
+		{
+			return Index;
+		}
+
+		const int32 GetFaceIndex() const
+		{
+			return FaceIndex;
+		}
+
+		const int32 GetId() const
+		{
+			return Id;
+		}
+
+		const TArray<FIsoSegment*>& GetConnectedSegments() const
+		{
+			return ConnectedSegments;
+		}
+
+		FIsoSegment* GetSegmentConnectedTo(const FIsoNode* Node) const;
+
+		TArray<FIsoSegment*>& GetConnectedSegments()
+		{
+			return ConnectedSegments;
+		}
+
+		void ConnectSegment(FIsoSegment* Segment)
+		{
+			ConnectedSegments.Add(Segment);
+		}
+
+		void DisconnectSegment(FIsoSegment* Segment)
+		{
+			ConnectedSegments.Remove(Segment);
+		}
+
+		const virtual bool IsALoopNode() const = 0;
+
+		void SetAsLinkedToLoop()
+		{
+			States |= EIsoNodeStates::LinkedToLoop;
+		}
+
+		void SetLinkedToLoopInQuarter(int32 Quarter)
+		{
+			switch (Quarter)
+			{
+			case 0:
+				States |= EIsoNodeStates::FirstQuarter;
+			case 1:
+				States |= EIsoNodeStates::SecondQuarter;
+			case 2:
+				States |= EIsoNodeStates::ThirdQuarter;
+			default:
+				States |= EIsoNodeStates::FourthQuarter;
+			}
+		}
+
+		bool IsLinkedToLoopInQuarter(int32 Quarter) const // [0 Pi/2]
+		{
+			switch (Quarter)
+			{
+			case 0:
+				return (States & EIsoNodeStates::FirstQuarter) == EIsoNodeStates::FirstQuarter;
+			case 1:
+				return (States & EIsoNodeStates::SecondQuarter) == EIsoNodeStates::SecondQuarter;
+			case 2:
+				return (States & EIsoNodeStates::ThirdQuarter) == EIsoNodeStates::ThirdQuarter;
+			default:
+				return (States & EIsoNodeStates::FourthQuarter) == EIsoNodeStates::FourthQuarter;
+			}
+		}
+
+		bool IsLinkedToLoopInNearlyIso(int32 Iso) const
+		{
+			if (Iso > 3)
+			{
+				Iso -= 4;
+			}
+			switch ((EIsoLink)Iso)
+			{
+			case EIsoLink::IsoUNext:
+				return (States & EIsoNodeStates::LinkedToNextU) == EIsoNodeStates::LinkedToNextU;
+			case EIsoLink::IsoVNext:
+				return (States & EIsoNodeStates::LinkedToNextV) == EIsoNodeStates::LinkedToNextV;
+			case EIsoLink::IsoUPrevious:
+				return (States & EIsoNodeStates::LinkedToPreviousU) == EIsoNodeStates::LinkedToPreviousU;
+			default:
+				return (States & EIsoNodeStates::LinkedToPreviousV) == EIsoNodeStates::LinkedToPreviousV;
+			}
+		}
+
+		void SetLinkedToIso(EIsoLink Iso) // [3Pi/2 2Pi]
+		{
+			switch (Iso)
+			{
+			case EIsoLink::IsoUNext:
+				States |= EIsoNodeStates::LinkedToNextU;
+				break;
+			case EIsoLink::IsoVNext:
+				States |= EIsoNodeStates::LinkedToNextV;
+				break;
+			case EIsoLink::IsoUPrevious:
+				States |= EIsoNodeStates::LinkedToPreviousU;
+				break;
+			default:
+				States |= EIsoNodeStates::LinkedToPreviousV;
+			}
+		}
+
+		void SetLinkedToIso(int32 Iso) // [3Pi/2 2Pi]
+		{
+			if (Iso > 3)
+			{
+				Iso -= 4;
+			}
+			switch ((EIsoLink)Iso)
+			{
+			case EIsoLink::IsoUNext:
+				States |= EIsoNodeStates::LinkedToNextU;
+				break;
+			case EIsoLink::IsoVNext:
+				States |= EIsoNodeStates::LinkedToNextV;
+				break;
+			case EIsoLink::IsoUPrevious:
+				States |= EIsoNodeStates::LinkedToPreviousU;
+				break;
+			default:
+				States |= EIsoNodeStates::LinkedToPreviousV;
+			}
+		}
+
+		/**
+		 * Return the 2d coordinate of the node according to the space
+		 */
+		virtual const FPoint2D& Get2DPoint(EGridSpace Space, const FGrid& Grid) const = 0;
+
+		virtual const FPoint& Get3DPoint(const FGrid& Grid) const = 0;
+		virtual const FPoint& GetNormal(const FGrid& Grid) const = 0;
+
+		/**
+		 * Only for display purpose as it return a copy of the point
+		 */
+		virtual const FPoint GetPoint(EGridSpace Space, const FGrid& Grid) const = 0;
+
+	};
+
+	/**
+	 * Node of the loop of the TopologicalFace
+	 */
+	class FLoopNode : public FIsoNode
+	{
+	protected:
+		int32 LoopIndex;
+		FLoopNode* ConnectedLoopNodes[2];
+
+	public:
+		FLoopNode(int32 InLoopIndex, int32 InNodeIndex, int32 InFaceIndex, int32 NodeId)
+			: FIsoNode(InNodeIndex, InFaceIndex, NodeId)
+			, LoopIndex(InLoopIndex)
+			, ConnectedLoopNodes{ nullptr, nullptr }
+		{
+		};
+
+		virtual const bool IsALoopNode() const override
+		{
+			return true;
+		}
+
+		const int32 GetLoopIndex() const
+		{
+			return LoopIndex;
+		}
+
+		void SetNextConnectedNode(FLoopNode* NextNode)
+		{
+			ConnectedLoopNodes[1] = NextNode;
+		}
+		void SetPreviousConnectedNode(FLoopNode* PreviousNode)
+		{
+			ConnectedLoopNodes[0] = PreviousNode;
+		}
+
+		FLoopNode& GetPreviousNode() const
+		{
+			return *ConnectedLoopNodes[0];
+		}
+
+		FLoopNode& GetNextNode() const
+		{
+			return *ConnectedLoopNodes[1];
+		}
+
+		/**
+		 * Return the 2d coordinate of the node according to the space
+		 */
+		virtual const FPoint2D& Get2DPoint(EGridSpace Space, const FGrid& Grid) const override
+		{
+			return Grid.GetLoop2DPoint(Space, LoopIndex, Index);
+		}
+
+		virtual const FPoint& Get3DPoint(const FGrid& Grid) const override
+		{
+			return Grid.GetLoops3D()[LoopIndex][Index];
+		}
+
+		virtual const FPoint& GetNormal(const FGrid& Grid) const override
+		{
+			return Grid.GetLoopNormals()[LoopIndex][Index];
+		}
+
+		/**
+		 * Only for display purpose as it return a copy of the point
+		 */
+		virtual const FPoint GetPoint(EGridSpace Space, const FGrid& Grid) const override
+		{
+			switch (Space)
+			{
+			case EGridSpace::Default2D:
+			case EGridSpace::Scaled:
+			case EGridSpace::UniformScaled:
+				return Grid.GetLoop2DPoint(Space, LoopIndex, Index);
+			//case EGridSpace::Default3D:
+			//	return Grid.GetLoop3DPoint(LoopIndex, Index);
+			}
+			return FPoint::ZeroPoint;
+		}
+
+		void SetAsIsoU()
+		{
+			States |= EIsoNodeStates::NearlyIsoU;
+		}
+
+		bool IsIsoU() const
+		{
+			return (States & EIsoNodeStates::NearlyIsoU) == EIsoNodeStates::NearlyIsoU;
+		}
+
+		void SetAsIsoV()
+		{
+			States |= EIsoNodeStates::NearlyIsoV;
+		}
+
+		bool IsIsoV() const
+		{
+			return (States & EIsoNodeStates::NearlyIsoV) == EIsoNodeStates::NearlyIsoV;
+		}
+
+		constexpr bool IsIso(EIso Axe) const
+		{
+			if (Axe == IsoV)
+			{
+				return IsIsoV();
+			}
+			else
+			{
+				return IsIsoU();
+			}
+		}
+
+	};
+
+	class FIsoInnerNode : public FIsoNode
+	{
+	public:
+		FIsoInnerNode(int32 NodeIndex, int32 FaceIndex, int32 NodeId)
+			: FIsoNode(NodeIndex, FaceIndex, NodeId)
+		{
+		};
+
+		const bool IsIsolated() const
+		{
+			return (ConnectedSegments.Num() == 0);
+		}
+
+		virtual const bool IsALoopNode() const override
+		{
+			return false;
+		};
+
+		bool IsLinkedToPreviousU() const
+		{
+			return (States & EIsoNodeStates::LinkedToPreviousU) == EIsoNodeStates::LinkedToPreviousU;
+		}
+
+		bool IsLinkedToNextU() const
+		{
+			return (States & EIsoNodeStates::LinkedToNextU) == EIsoNodeStates::LinkedToNextU;
+		}
+
+		bool IsLinkedToPreviousV() const
+		{
+			return (States & EIsoNodeStates::LinkedToPreviousV) == EIsoNodeStates::LinkedToPreviousV;
+		}
+
+		bool IsLinkedToNextV() const
+		{
+			return (States & EIsoNodeStates::LinkedToNextV) == EIsoNodeStates::LinkedToNextV;;
+		}
+
+		/** the node is connected in its 4 directions*/
+		bool IsComplete() const
+		{
+			return (States & EIsoNodeStates::TriangleComplete) == EIsoNodeStates::TriangleComplete;
+		}
+
+		bool IsLinkedToBoundary() const
+		{
+			return (States & EIsoNodeStates::LinkedToLoop) == EIsoNodeStates::LinkedToLoop;
+		}
+
+		virtual const FPoint2D& Get2DPoint(EGridSpace Space, const FGrid& Grid) const override
+		{
+			return Grid.GetInner2DPoint(Space, Index);
+		}
+
+		virtual const FPoint& Get3DPoint(const FGrid& Grid) const override
+		{
+			return Grid.GetInner3DPoint(Index);
+		}
+
+		virtual const FPoint& GetNormal(const FGrid& Grid) const override
+		{
+			return Grid.GetPointNormal(Index);
+		}
+
+		/**
+		 * Only for display purpose as it return a copy of the point
+		 */
+		virtual const FPoint GetPoint(EGridSpace Space, const FGrid& Grid) const override
+		{
+			switch (Space)
+			{
+			case EGridSpace::Default2D:
+			case EGridSpace::Scaled:
+			case EGridSpace::UniformScaled:
+				return Grid.GetInner2DPoint(Space, Index);
+			//case EGridSpace::Default3D:
+			//	return Grid.GetInner3DPoint(Index);
+			}
+			return FPoint::ZeroPoint;
+		}
+
+		void OffsetId(int32 StartId)
+		{
+			Id += StartId;
+		}
+	};
+
+} // namespace CADKernel
+
diff --git a/Engine/Source/Runtime/Datasmith/CADKernel/Private/CADKernel/Mesh/Meshers/IsoTriangulator/IsoTriangulatorDisplay.cpp b/Engine/Source/Runtime/Datasmith/CADKernel/Private/CADKernel/Mesh/Meshers/IsoTriangulator/IsoTriangulatorDisplay.cpp
new file mode 100644
--- /dev/null
+++ b/Engine/Source/Runtime/Datasmith/CADKernel/Private/CADKernel/Mesh/Meshers/IsoTriangulator/IsoTriangulatorDisplay.cpp
@@ -0,0 +1,142 @@
+// Copyright Epic Games, Inc. All Rights Reserved.
+#include "CADKernel/Mesh/Meshers/IsoTriangulator.h"
+
+#ifdef CADKERNEL_DEV
+
+#include "CADKernel/Mesh/Meshers/IsoTriangulator/IntersectionSegmentTool.h"
+#include "CADKernel/Mesh/Structure/Grid.h"
+#include "CADKernel/Mesh/Meshers/IsoTriangulator/IsoSegment.h"
+#include "CADKernel/UI/Display.h"
+
+using namespace CADKernel;
+
+void FIsoTriangulator::Display(EGridSpace Space, const FIsoSegment& Segment, FIdent Ident, EVisuProperty Property, bool bDisplayOrientation) const
+{
+	CADKernel::DisplaySegment(Segment.GetFirstNode().GetPoint(Space, Grid), Segment.GetSecondNode().GetPoint(Space, Grid), Ident, Property, bDisplayOrientation);
+}
+
+void FIsoTriangulator::DisplayTriangle(EGridSpace Space, const FIsoNode& NodeA, const FIsoNode& NodeB, const FIsoNode& NodeC) const
+{
+	TArray<FPoint> Points;
+	Points.SetNum(3);
+	Points[0] = NodeA.GetPoint(Space, Grid);
+	Points[1] = NodeB.GetPoint(Space, Grid);
+	Points[2] = NodeC.GetPoint(Space, Grid);
+	DrawElement(2, Points, EVisuProperty::Element);
+	DisplaySegment(Points[0], Points[1]); 
+	DisplaySegment(Points[1], Points[2]);
+	DisplaySegment(Points[2], Points[0]);
+}
+
+void FIsoTriangulator::Display(EGridSpace Space, const FIsoNode& NodeA, const FIsoNode& NodeB, FIdent Ident, EVisuProperty Property) const
+{
+	return CADKernel::DisplaySegment(NodeA.GetPoint(Space, Grid), NodeB.GetPoint(Space, Grid), Ident, Property);
+}
+
+void FIsoTriangulator::Display(EGridSpace Space, const FIsoNode& Node, FIdent Ident, EVisuProperty Property) const
+{
+	return CADKernel::DisplayPoint(Node.GetPoint(Space, Grid), Property, Ident);
+}
+
+void FIsoTriangulator::Display(EGridSpace Space, const TCHAR* Message, const TArray<FIsoSegment*>& Segments, bool bDisplayNode, bool bDisplayOrientation, EVisuProperty Property) const
+{
+	if (!bDisplay)
+	{
+		return;
+	}
+
+	F3DDebugSession _(Message);
+#ifdef DEBUG_LIMIT
+	for (FIsoSegment* Segment : Segments)
+	{
+		if (Segment->IsALimit())
+		{
+			Display(Space, *Segment, 0, EVisuProperty::Iso, bDisplayOrientation);
+		}
+	}
+#endif
+	for (FIsoSegment* Segment : Segments)
+	{
+#ifdef DEBUG_LIMIT
+		if (!Segment->IsALimit())
+#endif
+		{
+			Display(Space, *Segment, 0, Property, bDisplayOrientation);
+		}
+	}
+
+	if (bDisplayNode)
+	{
+		for (FIsoSegment* Segment : Segments)
+		{
+			Display(Space, Segment->GetFirstNode(), Segment->GetFirstNode().GetFaceIndex());
+			Display(Space, Segment->GetSecondNode(), Segment->GetSecondNode().GetFaceIndex());
+		}
+	}
+	//Wait();
+}
+
+void FIsoTriangulator::DisplayIsoNodes(EGridSpace Space) const
+{
+	if (!bDisplay)
+	{
+		return;
+	}
+
+	Open3DDebugSession(TEXT("FIsoTrianguler::IsoNodes"));
+	for (const FLoopNode& Node : LoopNodes)
+	{
+		Display(Space, Node, Node.GetFaceIndex(), EVisuProperty::YellowPoint);
+	}
+	for (const FIsoNode& Node : InnerNodes)
+	{
+		Display(Space, Node, Node.GetFaceIndex());
+	}
+	Close3DDebugSession();
+}
+
+void FIsoTriangulator::DisplayPixel(const int32 IndexU, const int32 IndexV) const
+{
+	DisplayPixel(IndexV * Grid.GetCuttingCount(EIso::IsoU) + IndexU);
+}
+
+void FIsoTriangulator::DisplayPixel(const int32 Index) const
+{
+	FPoint Point = (Grid.GetInner2DPoint(EGridSpace::Default2D, Index) + Grid.GetInner2DPoint(EGridSpace::Default2D, Index + Grid.GetCuttingCount(EIso::IsoU) + 1)) * 0.5;
+	CADKernel::DisplayPoint(Point, EVisuProperty::GreenPoint);
+};
+
+void FIsoTriangulator::DisplayPixels(TArray<uint8>& Pixel) const
+{
+	if (!bDisplay)
+	{
+		return;
+	}
+
+	Open3DDebugSession(TEXT("FIsoTrianguler::Pixel"));
+	for (int32 Index = 0; Index <Grid.GetTotalCuttingCount(); ++Index)
+	{
+		if (Pixel[Index])
+		{
+			DisplayPixel(Index);
+		}
+	}
+	Close3DDebugSession();
+}
+
+void FIntersectionSegmentTool::Display(const TCHAR* Message) const
+{
+	if (!Grid.bDisplay)
+	{
+		return;
+	}
+	
+	Open3DDebugSession(Message);
+	for (const FSegment4IntersectionTools& Segment : Segments)
+	{
+		DisplaySegment(Segment.Segment2D[0], Segment.Segment2D[1]);
+	}
+	Close3DDebugSession();
+}
+
+#endif
\ No newline at end of file
diff --git a/Engine/Source/Runtime/Datasmith/CADKernel/Private/CADKernel/Mesh/Meshers/ParametricMesher.cpp b/Engine/Source/Runtime/Datasmith/CADKernel/Private/CADKernel/Mesh/Meshers/ParametricMesher.cpp
new file mode 100644
--- /dev/null
+++ b/Engine/Source/Runtime/Datasmith/CADKernel/Private/CADKernel/Mesh/Meshers/ParametricMesher.cpp
@@ -0,0 +1,1671 @@
+// Copyright Epic Games, Inc. All Rights Reserved.
+
+#include "CADKernel/Mesh/Meshers/ParametricMesher.h"
+
+#include "CADKernel/Core/KernelParameters.h"
+#include "CADKernel/Mesh/Criteria/CriteriaGrid.h"
+#include "CADKernel/Mesh/Criteria/Criterion.h"
+#include "CADKernel/Mesh/Meshers/IsoTriangulator.h"
+#include "CADKernel/Mesh/Meshers/MesherTools.h"
+#include "CADKernel/Mesh/Structure/EdgeMesh.h"
+#include "CADKernel/Mesh/Structure/FaceMesh.h"
+#include "CADKernel/Mesh/Structure/ModelMesh.h"
+#include "CADKernel/Mesh/Structure/VertexMesh.h"
+#include "CADKernel/Topo/Body.h"
+#include "CADKernel/Topo/Model.h"
+#include "CADKernel/Topo/Shell.h"
+#include "CADKernel/Topo/TopologicalFace.h"
+#include "CADKernel/Utils/Util.h"
+
+using namespace CADKernel;
+
+FMesherParameters::FMesherParameters()
+	: InconsistencyAngle(TEXT("inconsistencyAngle"), 20., *this)
+{}
+
+FParametricMesher::FParametricMesher(TSharedRef<FModelMesh> InMeshModel)
+	: MeshModel(InMeshModel)
+	, Parameters(MakeShared<FMesherParameters>())
+{
+}
+
+void FParametricMesher::MeshEntities(TArray<TSharedPtr<FEntity>>& InEntities)
+{
+	int32 FaceCount = 0;
+
+	for (TSharedPtr<FTopologicalFace>& Face : Faces)
+	{
+		Face->SetMarker1();
+	}
+
+	// count faces
+	for (TSharedPtr<FEntity>& Entity : InEntities)
+	{
+		TSharedPtr<FTopologicalEntity> TopologicalEntity = StaticCastSharedPtr<FTopologicalEntity>(Entity);
+		if(!TopologicalEntity.IsValid())
+		{
+			continue;
+		}
+		FaceCount += TopologicalEntity->FaceCount();
+	}
+	Faces.Reserve(Faces.Num() + FaceCount);
+
+	for (TSharedPtr<FTopologicalFace>& Face : Faces)
+	{
+		Face->ResetMarkers();
+	}
+
+	// Get independent Faces and spread body's shells orientation
+	for (TSharedPtr<FEntity>& Entity : InEntities)
+	{
+		TSharedPtr<FTopologicalEntity> TopologicalEntity = StaticCastSharedPtr<FTopologicalEntity>(Entity);
+		if (!TopologicalEntity.IsValid())
+		{
+			continue;
+		}
+
+		TopologicalEntity->SpreadBodyOrientation();
+		TopologicalEntity->GetFaces(Faces);
+	}
+
+	for (TSharedPtr<FTopologicalFace>& Face : Faces)
+	{
+		Face->ResetMarkers();
+	}
+
+	// Get independent elementary entities (Edge, Vertex)
+	for(TSharedPtr<FEntity>& Entity : InEntities)
+	{
+		switch (Entity->GetEntityType())
+		{
+		case EEntity::TopologicalEdge:
+			Edges.Add(StaticCastSharedPtr<FTopologicalEdge>(Entity));
+			break;
+
+		case EEntity::TopologicalVertex:
+			Vertices.Add(StaticCastSharedPtr<FTopologicalVertex>(Entity));
+			break;
+
+		default:
+			break;
+		}
+	}
+
+	MeshEntities();
+}
+
+void FParametricMesher::MeshEntities()
+{
+	FTimePoint StartTime = FChrono::Now();
+	FTimePoint ApplyCriteriaStartTime = FChrono::Now();
+
+	FProgress ProgressBar(Faces.Num() * 2, TEXT("Meshing Entities : Apply Surface Criteria"));
+
+	// ============================================================================================================
+	//      Apply Surface Criteria
+	// ============================================================================================================
+
+	for (TSharedPtr<FTopologicalFace> Face : Faces)
+	{
+		FProgress _(1, TEXT("Meshing Entities : Apply Surface Criteria"));
+
+		ensureCADKernel(Face.IsValid());
+		ensureCADKernel(!Face->IsDeleted());
+
+		ApplySurfaceCriteria(Face.ToSharedRef());
+	}
+
+	Chronos.ApplyCriteriaDuration = FChrono::Elapse(ApplyCriteriaStartTime);
+
+	FTimePoint MeshingStartTime = FChrono::Now();
+
+	// ============================================================================================================
+	//      Find quad surfaces 
+	// ============================================================================================================
+
+	TArray<FCostToFace> QuadTrimmedSurfaceSet;
+
+	if (Faces.Num() > 1)
+	{
+		TArray<TSharedPtr<FTopologicalFace>> OtherEntities;
+
+		FMessage::Printf(Log, TEXT("  Isolate QuadPatch\n"));
+		FTimePoint IsolateQuadPatchStartTime = FChrono::Now();
+
+		IsolateQuadFace(QuadTrimmedSurfaceSet, OtherEntities);
+
+		Chronos.IsolateQuadPatchDuration = FChrono::Elapse(IsolateQuadPatchStartTime);
+		FMessage::Printf(Log, TEXT("  %d Quad Surfaces found\n"), QuadTrimmedSurfaceSet.Num());
+	}
+
+	// ============================================================================================================
+	//      Mesh surfaces 
+	// ============================================================================================================
+
+	FMessage::Printf(Log, TEXT("  Mesh Surfaces\n"));
+
+	FTimePoint MeshStartTime = FChrono::Now();
+	MeshSurfaceByFront(QuadTrimmedSurfaceSet);
+	Chronos.GlobalMeshDuration = FChrono::Elapse(MeshStartTime);
+	Chronos.GlobalDuration = FChrono::Elapse(StartTime);
+
+	Chronos.PrintTimeElapse();
+}
+
+void FParametricMesher::ApplySurfaceCriteria(TSharedRef<FTopologicalFace> Surface)
+{
+	if (Surface->IsApplyCriteria())
+	{
+		return;
+	}
+
+	FCriteriaGrid Grid(Surface);
+	Grid.ApplyCriteria(GetMeshModel()->GetCriteria());
+
+	Surface->ChooseFinalDeltaUs();
+	Surface->SetApplyCriteria();
+
+	for (const TSharedPtr<FTopologicalLoop>& Loop : Surface->GetLoops())
+	{
+		for (const FOrientedEdge& OrientedEdge : Loop->GetEdges())
+		{
+			const TSharedPtr<FTopologicalEdge>& Edge = OrientedEdge.Entity;
+			if (!Edge.IsValid())
+			{
+				continue;
+			}
+			ApplyEdgeCriteria(Edge.ToSharedRef());
+		}
+	}
+}
+
+void FParametricMesher::ApplyEdgeCriteria(TSharedRef<FTopologicalEdge> Edge)
+{
+	TSharedRef<FTopologicalEdge> ActiveEdge = Edge->GetLinkActiveEdge();
+	if (ActiveEdge->IsApplyCriteria())
+	{
+		return;
+	}
+
+	ActiveEdge->ComputeCrossingPointCoordinates();
+
+	ActiveEdge->InitDeltaUs();
+
+	const TArray<double>& CrossingPointUs = ActiveEdge->GetCrossingPointUs();
+	TArray<double> Coordinates;
+	Coordinates.SetNum(CrossingPointUs.Num() * 2 - 1);
+	Coordinates[0] = CrossingPointUs[0];
+	for (int32 ICuttingPoint = 1; ICuttingPoint < ActiveEdge->GetCrossingPointUs().Num(); ICuttingPoint++)
+	{
+		Coordinates[2 * ICuttingPoint - 1] = (CrossingPointUs[ICuttingPoint - 1] + CrossingPointUs[ICuttingPoint]) * 0.5;
+		Coordinates[2 * ICuttingPoint] = CrossingPointUs[ICuttingPoint];
+	}
+
+	TArray<FCurvePoint> Points3D;
+	ActiveEdge->EvaluatePoints(Coordinates, 0, Points3D);
+
+	const TArray<TSharedPtr<FCriterion>>& Criteria = GetMeshModel()->GetCriteria();
+	for (const TSharedPtr<FCriterion> Criterion : Criteria)
+	{
+		Criterion->ApplyOnEdgeParameters(ActiveEdge, CrossingPointUs, Points3D);
+	}
+
+	ActiveEdge->ChooseFinalDeltaUs();
+	ActiveEdge->SetApplyCriteria();
+}
+
+void FParametricMesher::Mesh(TSharedRef<FTopologicalFace> Face)
+{
+	ensureCADKernel(!Face->IsDeleted());
+	ensureCADKernel(!Face->IsMeshed());
+
+	FMessage::Printf(EVerboseLevel::Log, TEXT("Meshing of surface %d\n"),  Face->GetId());
+
+	FProgress _(1, TEXT("Meshing Entities : Mesh Surface"));
+
+	if (BoolDisplayDebugMeshStep)
+	{
+		Open3DDebugSession(FString::Printf(TEXT("Mesh of surface %d"), Face->GetId()));
+	}
+
+	FTimePoint StartTime = FChrono::Now();
+
+	FGrid Grid(Face, MeshModel);
+
+	FTimePoint GenerateCloudStartTime = FChrono::Now();
+	GenerateCloud(Grid);
+	FDuration GenerateCloudDuration = FChrono::Elapse(GenerateCloudStartTime);
+
+	if(Grid.IsDegenerated())
+	{
+		if (BoolDisplayDebugMeshStep)
+		{
+			Close3DDebugSession();
+		}
+		FMessage::Printf(EVerboseLevel::Log, TEXT("The meshing of the surface %d failed due to a degenerated grid\n"), Face->GetId());
+		Face->SetMeshed();
+		return;
+	}
+
+	TSharedRef<FFaceMesh> SurfaceMesh = StaticCastSharedRef<FFaceMesh>(Face->GetOrCreateMesh(MeshModel));
+
+	FTimePoint IsoTriangulerStartTime = FChrono::Now();
+	FIsoTriangulator IsoTrianguler(Grid, SurfaceMesh);
+	if(IsoTrianguler.Triangulate())
+	{
+		if (Face->IsBackOriented())
+		{
+			SurfaceMesh->InverseOrientation();
+		}
+		MeshModel->AddMesh(SurfaceMesh);
+	}
+	Face->SetMeshed();
+
+	FDuration TriangulateDuration = FChrono::Elapse(IsoTriangulerStartTime);
+	FDuration Duration = FChrono::Elapse(StartTime);
+
+#ifdef CADKERNEL_DEV
+	Chronos.GlobalPointCloudDuration += Grid.Chronos.GeneratePointCloudDuration;
+	Chronos.GlobalGeneratePointCloudDuration += GenerateCloudDuration;
+	Chronos.GlobalTriangulateDuration += TriangulateDuration;
+	Chronos.GlobalDelaunayDuration += IsoTrianguler.Chronos.FindSegmentToLinkLoopToLoopByDelaunayDuration;
+	Chronos.GlobalMeshDuration += Duration;
+#endif
+
+	if (BoolDisplayDebugMeshStep)
+	{
+		Close3DDebugSession();
+	}
+}
+
+void FParametricMesher::GenerateCloud(FGrid& Grid)
+{
+	Grid.DefineCuttingParameters();
+	if(!Grid.GeneratePointCloud())
+	{
+		return;
+	}
+
+	bool bFindThinZone = false;
+	if (bFindThinZone)
+	{
+		FTimePoint StartTime = FChrono::Now();
+		Grid.SearchThinZones();
+
+		if (Grid.GetFace()->HasThinZone())
+		{
+#ifdef DEBUG_THIN_ZONES
+			{
+				F3DDebugSession _(TEXT("Thin Surface"));
+				Display(Grid.GetSurface());
+			}
+#endif
+			FTimePoint MeshThinZonesTime = FChrono::Now();
+			MeshThinZoneEdges(Grid);
+			Chronos.GlobalMeshThinZones += FChrono::Elapse(MeshThinZonesTime);
+		}
+		Chronos.GlobalThinZones += FChrono::Elapse(StartTime);
+	}
+
+	FTimePoint StartTime = FChrono::Now();
+	MeshFaceLoops(Grid);
+
+	Grid.ProcessPointCloud();
+
+	Chronos.GlobalMeshAndGetLoopNodes += FChrono::Elapse(StartTime);
+}
+
+void FParametricMesher::MeshFaceLoops(FGrid& Grid)
+{
+	TSharedRef<FTopologicalFace> Face = Grid.GetFace();
+
+	FTimePoint StartTime = FChrono::Now();
+
+	for (const TSharedPtr<FTopologicalLoop>& Loop : Face->GetLoops())
+	{
+		for (const FOrientedEdge& Edge : Loop->GetEdges())
+		{
+			Mesh(Edge.Entity.ToSharedRef(), Face);
+		}
+	}
+
+	Chronos.GlobalMeshEdges += FChrono::Elapse(StartTime);
+}
+
+void FillImposedIsoCuttingPoints(TArray<double>& UEdgeSetOfIntersectionWithIso, ECoordinateType CoordinateType, double ToleranceGeoEdge, const TSharedRef<FTopologicalEdge> Edge, TArray<FCuttingPoint>& OutImposedIsoVertexSet)
+{
+	int32 StartIndex = OutImposedIsoVertexSet.Num();
+
+	Algo::Sort(UEdgeSetOfIntersectionWithIso);
+	double PreviousU = -HUGE_VALUE;
+	for (double InterU : UEdgeSetOfIntersectionWithIso)
+	{
+		// Remove nearly duplicate 
+		if (InterU - PreviousU < ToleranceGeoEdge)
+		{
+			continue;
+		}
+		OutImposedIsoVertexSet.Emplace(InterU, CoordinateType);
+		PreviousU = InterU;
+	}
+
+	int32 Index;
+	if (OutImposedIsoVertexSet.Num() - StartIndex > 1)
+	{
+		OutImposedIsoVertexSet[StartIndex].IsoDeltaU = (OutImposedIsoVertexSet[StartIndex + 1].Coordinate - OutImposedIsoVertexSet[StartIndex].Coordinate) * AQuarter;
+		for (Index = StartIndex + 1; Index < OutImposedIsoVertexSet.Num() - 1; ++Index)
+		{
+			OutImposedIsoVertexSet[Index].IsoDeltaU = (OutImposedIsoVertexSet[Index + 1].Coordinate - OutImposedIsoVertexSet[Index - 1].Coordinate) * AEighth;
+		}
+		OutImposedIsoVertexSet[Index].IsoDeltaU = (OutImposedIsoVertexSet[Index].Coordinate - OutImposedIsoVertexSet[Index - 1].Coordinate) * AQuarter;
+	}
+	else if (OutImposedIsoVertexSet.Num() - StartIndex == 1)
+	{
+		int32 CuttingPointIndex = 0;
+		while (CuttingPointIndex < Edge->GetCrossingPointUs().Num() && Edge->GetCrossingPointUs()[CuttingPointIndex] <= OutImposedIsoVertexSet[StartIndex].Coordinate)
+		{
+			++CuttingPointIndex;
+		};
+		if (CuttingPointIndex > 0)
+		{
+			--CuttingPointIndex;
+		}
+		OutImposedIsoVertexSet[StartIndex].IsoDeltaU = Edge->GetDeltaUMaxs()[CuttingPointIndex] * AQuarter;
+	}
+}
+
+void FParametricMesher::Mesh(TSharedRef<FTopologicalVertex> InVertex)
+{
+	InVertex->GetOrCreateMesh(GetMeshModel());
+}
+
+void FParametricMesher::Mesh(TSharedRef<FTopologicalEdge> InEdge, TSharedRef<FTopologicalFace> TrimmedSurface)
+{
+	TSharedRef<FTopologicalEdge> ActiveEdge = StaticCastSharedRef<FTopologicalEdge>(InEdge->GetLinkActiveEntity());
+	if (ActiveEdge->IsMeshed())
+	{
+		return;
+	}
+
+	if (ActiveEdge->IsThinPeak())
+	{
+		TArray<FCuttingPoint>& FinalEdgeCuttingPointCoordinates = ActiveEdge->GetCuttingPoints();
+		FinalEdgeCuttingPointCoordinates.Emplace(ActiveEdge->GetStartCurvilinearCoordinates(), ECoordinateType::VertexCoordinate);
+		FinalEdgeCuttingPointCoordinates.Emplace(ActiveEdge->GetEndCurvilinearCoordinates(), ECoordinateType::VertexCoordinate);
+		GenerateEdgeElements(ActiveEdge);
+		return;
+	}
+
+	double ToleranceGeoEdge = ActiveEdge->GetCurve()->GetParametricTolerance();
+
+	const FSurfacicTolerance& ToleranceIso = TrimmedSurface->GetIsoTolerances();
+
+	// Get Edge intersection with inner surface mesh grid
+	TArray<double> UEdgeSetOfIntersectionWithIsoU;
+	TArray<double> UEdgeSetOfIntersectionWithIsoV;
+
+	const TArray<double>& SurfaceTabU = TrimmedSurface->GetCuttingCoordinatesAlongIso(EIso::IsoU);
+	const TArray<double>& SurfaceTabV = TrimmedSurface->GetCuttingCoordinatesAlongIso(EIso::IsoV);
+
+	TArray<FPoint2D> EdgeCrossingPoints2D;
+	const TArray<double>& EdgeCrossingPointsU = ActiveEdge->GetCrossingPointUs();
+
+	InEdge->ProjectTwinEdgePointsOn2DCurve(ActiveEdge, EdgeCrossingPointsU, EdgeCrossingPoints2D);
+
+#ifdef DEBUG_MESH_EDGE
+	{
+		F3DDebugSession _(FString::Printf(TEXT("EdgePointsOnDomain %d"), ActiveEdge->GetId()));
+		for (const FPoint& Point : EdgeCrossingPoints2D)
+		{
+			DisplayPoint(Point);
+		}
+		Wait();
+	}
+#endif
+
+#ifdef DEBUG_IntersectEdgeIsos
+	DebugIntersectEdgeIsos(ActiveEdge, TrimmedSurface, EdgeCrossingPoints2D, ToleranceIsoU, SurfaceTabU, EIso::IsoU);
+	DebugIntersectEdgeIsos(ActiveEdge, TrimmedSurface, EdgeCrossingPoints2D, ToleranceIsoV, SurfaceTabV, EIso::IsoV);
+#endif
+	IntersectEdgeIsos(ActiveEdge, EdgeCrossingPoints2D, ToleranceIso[EIso::IsoU], SurfaceTabU, EIso::IsoU, UEdgeSetOfIntersectionWithIsoU);
+	IntersectEdgeIsos(ActiveEdge, EdgeCrossingPoints2D, ToleranceIso[EIso::IsoV], SurfaceTabV, EIso::IsoV, UEdgeSetOfIntersectionWithIsoV);
+
+#ifdef DEBUG_IntersectEdgeIsos
+	{
+      Work only if ActiveEdge == InEdge
+		F3DDebugSession _(FString::Printf(TEXT("Edge %d Intersect with iso"), InEdge->Id));
+		TArray<FPoint> Intersections;
+		ActiveEdge->Approximate2DPoints(UEdgeSetOfIntersectionWithIsoU, Intersections);
+		for(const FPoint& Point : Intersections)
+		{
+			Display(Point);
+		}
+		Wait();
+	}
+#endif
+
+	FLinearBoundary EdgeBounds = ActiveEdge->GetBoundary();
+
+	TArray<double>& DeltaUs = ActiveEdge->GetDeltaUMaxs();
+
+	// build a edge mesh compiling inner surface cutting (based on criteria applied on the surface) and edge cutting (based on criteria applied on the curve)
+	TArray<FCuttingPoint> ImposedIsoCuttingPoints;
+	{
+		int32 NbImposedCuttingPoints = ActiveEdge->GetImposedCuttingPoints().Num() + UEdgeSetOfIntersectionWithIsoU.Num() + UEdgeSetOfIntersectionWithIsoV.Num() + 2;
+		ImposedIsoCuttingPoints.Reserve(NbImposedCuttingPoints);
+	}
+
+	ImposedIsoCuttingPoints.Emplace(EdgeBounds.GetMin(), ECoordinateType::VertexCoordinate, -1, ToleranceGeoEdge);
+
+	double MinDeltaU = HUGE_VALUE;
+	for (const double& DeltaU : DeltaUs)
+	{
+		MinDeltaU = FMath::Min(MinDeltaU, DeltaU);
+	}
+
+#ifdef DEBUG_MESH_EDGE
+	{
+		F3DDebugSession _(FString::Printf(TEXT("Edge %d"), InEdge->GetId()));
+		for (int32 Index = 0; Index < ActiveEdge->GetImposedCuttingPoints().Num(); ++Index)
+		{
+			const FImposedCuttingPoint& CuttingPoint = ActiveEdge->GetImposedCuttingPoints()[Index];
+			ImposedIsoCuttingPoints.Emplace(CuttingPoint.Coordinate, ECoordinateType::ImposedCoordinate, CuttingPoint.OppositNodeIndex, MinDeltaU * AThird);
+			FCurvePoint Point;
+			ActiveEdge->EvaluatePoint(CuttingPoint.Coordinate, 0, Point);
+			DisplayPoint(Point.Point);
+		}
+	}
+#endif
+
+	// Add Edge intersection with inner surface grid Isos
+	FillImposedIsoCuttingPoints(UEdgeSetOfIntersectionWithIsoU, IsoUCoordinate, ToleranceGeoEdge, ActiveEdge, ImposedIsoCuttingPoints);
+	FillImposedIsoCuttingPoints(UEdgeSetOfIntersectionWithIsoV, IsoVCoordinate, ToleranceGeoEdge, ActiveEdge, ImposedIsoCuttingPoints);
+
+	ImposedIsoCuttingPoints.Sort([](const FCuttingPoint& Point1, const FCuttingPoint& Point2) { return Point1.Coordinate < Point2.Coordinate; });
+
+	auto MergeImposedCuttingPoints = [&](int32& Index, int32& NewIndex, ECoordinateType NewType)
+	{
+		double DeltaU = FMath::Max(ImposedIsoCuttingPoints[NewIndex].IsoDeltaU, ImposedIsoCuttingPoints[Index].IsoDeltaU);
+		if (ImposedIsoCuttingPoints[NewIndex].Type <= ImposedCoordinate && ImposedIsoCuttingPoints[Index].Type <= ImposedCoordinate)
+		{
+			DeltaU /= 5.;
+		}
+
+		if (ImposedIsoCuttingPoints[NewIndex].Coordinate + DeltaU > ImposedIsoCuttingPoints[Index].Coordinate)
+		{
+			if (ImposedIsoCuttingPoints[Index].Type == VertexCoordinate)
+			{
+				ImposedIsoCuttingPoints[NewIndex].Coordinate = ImposedIsoCuttingPoints[Index].Coordinate;
+				ImposedIsoCuttingPoints[NewIndex].IsoDeltaU = ImposedIsoCuttingPoints[Index].IsoDeltaU;
+			}
+			else if (ImposedIsoCuttingPoints[NewIndex].Type == VertexCoordinate)
+			{
+			}
+			else if (ImposedIsoCuttingPoints[NewIndex].Type == ImposedCoordinate)
+			{
+				if (ImposedIsoCuttingPoints[Index].Type == ImposedCoordinate)
+				{
+					ImposedIsoCuttingPoints[NewIndex].Coordinate = (ImposedIsoCuttingPoints[NewIndex].Coordinate + ImposedIsoCuttingPoints[Index].Coordinate) * 0.5;
+				}
+			}
+			else if (ImposedIsoCuttingPoints[Index].Type == ImposedCoordinate)
+			{
+				ImposedIsoCuttingPoints[NewIndex].Coordinate = ImposedIsoCuttingPoints[Index].Coordinate;
+				ImposedIsoCuttingPoints[NewIndex].Type = ImposedCoordinate;
+				ImposedIsoCuttingPoints[NewIndex].IsoDeltaU = ImposedIsoCuttingPoints[Index].IsoDeltaU;
+			}
+			else if (ImposedIsoCuttingPoints[NewIndex].Type != ImposedIsoCuttingPoints[Index].Type)
+			{
+				ImposedIsoCuttingPoints[NewIndex].Coordinate = (ImposedIsoCuttingPoints[NewIndex].Coordinate + ImposedIsoCuttingPoints[Index].Coordinate) * 0.5;
+				ImposedIsoCuttingPoints[NewIndex].Type = IsoUVCoordinate;
+				ImposedIsoCuttingPoints[NewIndex].IsoDeltaU = FMath::Min(ImposedIsoCuttingPoints[NewIndex].IsoDeltaU, ImposedIsoCuttingPoints[Index].IsoDeltaU);
+			}
+
+			if (ImposedIsoCuttingPoints[NewIndex].Type <= ImposedCoordinate)
+			{
+				if (ImposedIsoCuttingPoints[NewIndex].OppositNodeIndex == -1)
+				{
+					ImposedIsoCuttingPoints[NewIndex].OppositNodeIndex = ImposedIsoCuttingPoints[Index].OppositNodeIndex;
+				}
+				else
+				{
+					ImposedIsoCuttingPoints[NewIndex].OppositNodeIndex2 = ImposedIsoCuttingPoints[Index].OppositNodeIndex;
+				}
+			}
+		}
+		else
+		{
+			++NewIndex;
+			ImposedIsoCuttingPoints[NewIndex] = ImposedIsoCuttingPoints[Index];
+		}
+	};
+
+	// If a pair of point isoU/isoV is too close, get the middle of the points
+	if (ImposedIsoCuttingPoints.Num() > 1)
+	{
+		int32 NewIndex = 0;
+		for (int32 Index = 1; Index < ImposedIsoCuttingPoints.Num(); ++Index)
+		{
+			if (ImposedIsoCuttingPoints[Index].Type > ECoordinateType::ImposedCoordinate)
+			{
+				bool bIsDelete = false;
+				for (const FLinearBoundary& ThinZone : ActiveEdge->GetThinZoneBounds())
+				{
+					if (ThinZone.Contains(ImposedIsoCuttingPoints[Index].Coordinate))
+					{
+						bIsDelete = true;
+						continue;
+					}
+				}
+				if (bIsDelete)
+				{
+					continue;
+				}
+			}
+
+			if (ImposedIsoCuttingPoints[NewIndex].Type == ECoordinateType::ImposedCoordinate || ImposedIsoCuttingPoints[Index].Type == ECoordinateType::ImposedCoordinate)
+			{
+				MergeImposedCuttingPoints(Index, NewIndex, ECoordinateType::ImposedCoordinate);
+			}
+			else if (ImposedIsoCuttingPoints[NewIndex].Type != ImposedIsoCuttingPoints[Index].Type)
+			{
+				MergeImposedCuttingPoints(Index, NewIndex, ECoordinateType::IsoUVCoordinate);
+			}
+			else
+			{
+				++NewIndex;
+				ImposedIsoCuttingPoints[NewIndex] = ImposedIsoCuttingPoints[Index];
+			}
+		}
+		ImposedIsoCuttingPoints.SetNum(NewIndex + 1);
+	}
+
+	if (ImposedIsoCuttingPoints.Num() > 1 && (EdgeBounds.GetMax() - ImposedIsoCuttingPoints.Last().Coordinate) < FMath::Min(ImposedIsoCuttingPoints.Last().IsoDeltaU, ActiveEdge->GetDeltaUMaxs().Last()))
+	{
+		ImposedIsoCuttingPoints.Last().Coordinate = EdgeBounds.GetMax();
+		ImposedIsoCuttingPoints.Last().Type = VertexCoordinate;
+	}
+	else
+	{
+		ImposedIsoCuttingPoints.Emplace(EdgeBounds.GetMax(), ECoordinateType::VertexCoordinate, -1, ActiveEdge->GetDeltaUMaxs().Last() * AQuarter);
+	}
+
+	// max vertex of the edge
+	int32 MaxNumberOfVertex = (int32)((EdgeBounds.GetMax() - EdgeBounds.GetMin()) / MinDeltaU) + 5;
+
+	// Final array of the edge mesh vertex 
+	TArray<FCuttingPoint>& FinalEdgeCuttingPointCoordinates = ActiveEdge->GetCuttingPoints();
+
+	FinalEdgeCuttingPointCoordinates.Empty(ImposedIsoCuttingPoints.Num() + MaxNumberOfVertex);
+
+#ifdef DEBUG_GetPreferredUVCoordinatesFromNeighbours
+	TArray<FCuttingPoint> Extremities;
+	Extremities.Reserve(2);
+	Extremities.Emplace(EdgeUMin, ECoordinateType::VertexCoordinate, -1, ToleranceGeoEdge);
+	Extremities.Emplace(EdgeUMax, ECoordinateType::VertexCoordinate, -1, ToleranceGeoEdge);
+
+	TArray<double> CuttingPoints2;
+	FMesherTools::ComputeFinalCuttingPointsWithImposedCuttingPoints(ActiveEdge->GetCrossingPointUs(), ActiveEdge->GetDeltaUMaxs(), Extremities, CuttingPoints2);
+#endif
+
+	if (ActiveEdge->IsDegenerated())
+	{
+		for(FCuttingPoint CuttingPoint : ImposedIsoCuttingPoints)
+		{
+			FinalEdgeCuttingPointCoordinates.Emplace(CuttingPoint.Coordinate, ECoordinateType::OtherCoordinate);
+		}
+	}
+	else
+	{
+		TArray<double> CuttingPoints;
+		FMesherTools::ComputeFinalCuttingPointsWithImposedCuttingPoints(ActiveEdge->GetCrossingPointUs(), ActiveEdge->GetDeltaUMaxs(), ImposedIsoCuttingPoints, CuttingPoints);
+		for (const double& Coordinate : CuttingPoints)
+		{
+			FinalEdgeCuttingPointCoordinates.Emplace(Coordinate, ECoordinateType::OtherCoordinate);
+		}
+
+#ifdef DEBUG_GetPreferredUVCoordinatesFromNeighbours
+		{
+			F3DDebugSession G(TEXT("Mesh(TSharedRef<FEdge> InEdge"));
+			{
+				F3DDebugSession G(TEXT("U From Iso"));
+				for (const FCuttingPoint& CuttingU : ImposedIsoCuttingPoints)
+				{
+					Display(FPoint(CuttingU.Coordinate, 0.0));
+				}
+			}
+			{
+				F3DDebugSession G(TEXT("U From Criteria"));
+				for (double CuttingU : CuttingPoints2)
+				{
+					Display(FPoint(CuttingU, 0.05), EVisuProperty::NonManifoldEdge);
+				}
+			}
+			{
+				F3DDebugSession G(TEXT("U Final (Criteria & Iso)"));
+				for (double CuttingU : CuttingPoints)
+				{
+					Display(FPoint(CuttingU, 0.1), EVisuProperty::PurplePoint);
+				}
+			}
+			//Wait();
+		}
+#endif
+
+		GenerateEdgeElements(ActiveEdge);
+	}
+
+}
+
+void FParametricMesher::GenerateEdgeElements(TSharedRef<FTopologicalEdge> Edge)
+{
+	{
+		TSharedRef<FTopologicalEdge> ActiveEdge = StaticCastSharedRef<FTopologicalEdge>(Edge->GetLinkActiveEntity());
+		if (ActiveEdge != Edge)
+		{
+			return GenerateEdgeElements(ActiveEdge);
+		}
+	}
+
+	TSharedRef<FEdgeMesh> EdgeMesh = Edge->GetOrCreateMesh(MeshModel);
+
+	int32 StartVertexNodeIndex = Edge->GetStartVertex()->GetOrCreateMesh(GetMeshModel())->GetMesh();
+	int32 EndVertexNodeIndex = Edge->GetEndVertex()->GetOrCreateMesh(GetMeshModel())->GetMesh();
+
+	TSharedRef<FTopologicalEdge> ActiveEdge = StaticCastSharedRef<FTopologicalEdge>(EdgeMesh->GetGeometricEntity());
+
+	TArray<double> CuttingPointCoordinates;
+	CuttingPointCoordinates.Reserve(ActiveEdge->GetCuttingPoints().Num());
+	for (const FCuttingPoint& CuttingPoint : ActiveEdge->GetCuttingPoints())
+	{
+		CuttingPointCoordinates.Add(CuttingPoint.Coordinate);
+	}
+
+	ensureCADKernel(CuttingPointCoordinates.Num() > 1);
+	CuttingPointCoordinates.RemoveAt(0);
+	CuttingPointCoordinates.Pop();
+
+	TArray<FPoint>& Coordinates = EdgeMesh->GetNodeCoordinates();
+	ActiveEdge->ApproximatePoints(CuttingPointCoordinates, Coordinates);
+
+	EdgeMesh->RegisterCoordinates();
+	EdgeMesh->Mesh(StartVertexNodeIndex, EndVertexNodeIndex);
+	MeshModel->AddMesh(EdgeMesh);
+	Edge->SetMeshed();
+}
+
+void FParametricMesher::IsolateQuadFace(TArray<FCostToFace>& QuadSurfaces, TArray<TSharedPtr<FTopologicalFace>>& OtherSurfaces) const
+{
+	TArray<TSharedPtr<FTopologicalFace>> FlatQuadsAndTriangles;
+	FlatQuadsAndTriangles.Reserve(Faces.Num());
+	QuadSurfaces.Reserve(Faces.Num() * 2);
+	OtherSurfaces.Reserve(Faces.Num());
+
+	for (const TSharedPtr<FTopologicalFace>& Face : Faces)
+	{
+		if (Face->IsDeleted())
+		{
+			continue;
+		}
+
+		if (Face->IsMeshed())
+		{
+			continue;
+		}
+
+		Face->ComputeSurfaceSideProperties();
+	}
+
+	for (const TSharedPtr<FTopologicalFace>& Face : Faces)
+	{
+		Face->DefineSurfaceType();
+		switch (Face->GetQuadType())
+		{
+		case EQuadType::Quadrangular:
+			double LocalMinCurvature;
+			double LocalMaxCurvature;
+			Sort(Face->GetCurvature(EIso::IsoU).Max, Face->GetCurvature(EIso::IsoV).Max, LocalMinCurvature, LocalMaxCurvature);
+			if (LocalMaxCurvature > LocalMinCurvature)
+			{
+				QuadSurfaces.Emplace(LocalMaxCurvature, Face.ToSharedRef());
+				if (LocalMinCurvature > LocalMinCurvature)
+				{
+					QuadSurfaces.Emplace(LocalMinCurvature, Face.ToSharedRef());
+				}
+			}
+			else
+			{
+				FlatQuadsAndTriangles.Add(Face);
+				OtherSurfaces.Add(Face);
+			}
+			break;
+		case EQuadType::Triangular:
+			FlatQuadsAndTriangles.Add(Face);
+			OtherSurfaces.Add(Face);
+			break;
+		case EQuadType::Unset:
+		default:
+			OtherSurfaces.Add(Face);
+		}
+	}
+
+	Algo::Sort(QuadSurfaces, [](FCostToFace& SurfaceA, FCostToFace& SurfaceB)
+		{
+			return (SurfaceA.Cost > SurfaceB.Cost);
+		}
+	);
+
+#ifdef DEBUG_ISOLATEQUADFACE
+	if (QuadSurfaces.Num() > 0)
+	{
+		Open3DDebugSession(TEXT("Quad Entities"));
+		for (const FCostToFace& Quad : QuadSurfaces)
+		{
+			//FGraphicSession _(FString::Printf(TEXT("Face %d %f"), Quad.Face->GetId(), Quad.Cost));
+			Display(Quad.Face);
+		}
+		Close3DDebugSession();
+	}
+
+	if (FlatQuadsAndTriangles.Num() > 0)
+	{
+		Open3DDebugSession(TEXT("Flat Quads & Triangles"));
+		for (const TSharedPtr<FTopologicalFace>& Face : FlatQuadsAndTriangles)
+		{
+			//FGraphicSession _(FString::Printf(TEXT("Face %d %f"), Face->GetId()));
+			Display(Face);
+		}
+		Close3DDebugSession();
+	}
+
+	if (OtherSurfaces.Num() > 0)
+	{
+		Open3DDebugSession(TEXT("Other Entities"));
+		for (const TSharedPtr<FTopologicalFace>& Face : OtherSurfaces)
+		{
+			//FGraphicSession _(FString::Printf(TEXT("Face %d %f"), Face->GetId()));
+			Display(StaticCastSharedPtr<FTopologicalFace>(Face));
+		}
+		Close3DDebugSession();
+	}
+#endif
+}
+
+void FParametricMesher::LinkQuadSurfaceForMesh(TArray<FCostToFace>& QuadTrimmedSurfaceSet, TArray<TArray<TSharedPtr<FTopologicalFace>>>& OutStrips)
+{
+	const double GeometricTolerance = 20. * MeshModel->GetGeometricTolerance();
+
+	OutStrips.Reserve(QuadTrimmedSurfaceSet.Num());
+
+	for (FCostToFace& Quad : QuadTrimmedSurfaceSet)
+	{
+		TSharedPtr<FTopologicalFace> Surface = Quad.Face;
+		const FSurfaceCurvature& Curvatures = Surface->GetCurvatures();
+
+		EIso Axe = (!RealCompare(Quad.Cost, Curvatures[EIso::IsoU].Max)) ? EIso::IsoU : EIso::IsoV;
+
+		if (Axe == EIso::IsoU)
+		{
+			if (Surface->HasMarker1())
+			{
+				continue;
+			}
+			Surface->SetMarker1();
+		}
+		else
+		{
+			if (Surface->HasMarker2())
+			{
+				continue;
+			}
+			Surface->SetMarker2();
+		}
+
+		TArray<TSharedPtr<FTopologicalFace>>* QuadStrip = &OutStrips.Emplace_GetRef();
+		QuadStrip->Reserve(QuadTrimmedSurfaceSet.Num());
+		QuadStrip->Add(Surface);
+
+		const TArray<FEdge2DProperties>& SideProperties = Surface->GetSideProperties();
+
+		int32 StartSideIndex = 0;
+		for (; StartSideIndex < 4; StartSideIndex++)
+		{
+			if (SideProperties[StartSideIndex].IsoType == Axe)
+			{
+				break;
+			}
+		}
+		if(StartSideIndex == 4)
+		{
+			continue;
+		}
+
+		bool bFirstStep = true;
+		int32 SideIndex = StartSideIndex;
+
+		while (Surface)
+		{
+			int32 EdgeIndex = Surface->GetStartEdgeIndexOfSide(SideIndex);
+			double SideLength = Surface->GetSideProperties()[SideIndex].Length3D;
+			TSharedPtr<FTopologicalEdge> Edge = Surface->GetLoops()[0]->GetEdges()[EdgeIndex].Entity;
+
+			Surface = nullptr;
+			TSharedPtr<FTopologicalEdge> NextEdge = Edge->GetFirstTwinEdge();
+			if (NextEdge)
+			{
+				Surface = (TSharedPtr<FTopologicalFace>)NextEdge->GetLoop()->GetFace();
+				ensureCADKernel(Surface);
+			}
+
+			if (Surface && (Surface->GetQuadType() == EQuadType::Quadrangular || Surface->GetQuadType() == EQuadType::Triangular))
+			{
+				// check side length
+				int32 LocalEdgeIndex = Surface->GetLoops()[0]->GetEdgeIndex(NextEdge);
+				SideIndex = Surface->GetSideIndex(LocalEdgeIndex);
+				double OtherSideLength = Surface->GetSideProperties()[SideIndex].Length3D;
+
+				Sort(OtherSideLength, SideLength);
+				if (SideLength - OtherSideLength > GeometricTolerance)
+				{
+					Surface = nullptr;
+				}
+			}
+			else
+			{
+				Surface = nullptr;
+			}
+
+			if (Surface)
+			{
+				// Set as processed in a direction
+				const TArray<FEdge2DProperties>& LocalSideProperties = Surface->GetSideProperties();
+				if (LocalSideProperties[SideIndex].IsoType == EIso::IsoU)
+				{
+					if (Surface->HasMarker1())
+					{
+						Surface = nullptr;
+					}
+					else
+					{
+						Surface->SetMarker1();
+					}
+				}
+				else
+				{
+					if (Surface->HasMarker2())
+					{
+						Surface = nullptr;
+					}
+					else
+					{
+						Surface->SetMarker2();
+					}
+				}
+			}
+
+			if (Surface)
+			{
+				// it's a quad or a tri => add
+				if (Surface->GetQuadType() != EQuadType::Other)
+				{
+					QuadStrip->Add(Surface);
+				}
+
+				if (Surface->GetQuadType() == EQuadType::Triangular)
+				{
+					// stop
+					Surface = nullptr;
+				}
+			}
+
+			if (!Surface)
+			{
+				if (bFirstStep)
+				{
+					bFirstStep = false;
+					Surface = (*QuadStrip)[0];
+					SideIndex = (StartSideIndex + 2) % 4;
+					continue;
+				}
+				else
+				{
+					break;
+				}
+			}
+
+			// find opposite side
+			SideIndex = (SideIndex + 2) % 4;
+		}
+
+		if (QuadStrip->Num() == 1)
+		{
+			OutStrips.Pop();
+		}
+	}
+
+	for (const TSharedPtr<FTopologicalFace>& Face : Faces)
+	{
+		Face->ResetMarkers();
+	}
+}
+
+void FParametricMesher::MeshSurfaceByFront(TArray<FCostToFace>& QuadTrimmedSurfaceSet)
+{
+	// Processed3 : Surfaces that have to be meshed are set Processed3
+	// Processed1 : Surfaces added in CandidateSurfacesForMesh
+	// Processed2 : Surfaces added in SecondChoiceOfCandidateSurfacesForMesh
+
+	FMessage::Printf(EVerboseLevel::Debug, TEXT("Start MeshSurfaceByFront\n"));
+
+	for (TSharedPtr<FTopologicalFace> Face : Faces)
+	{
+		Face->SetMarker3();
+	}
+
+	const double GeometricTolerance = 20. * MeshModel->GetGeometricTolerance();
+
+	TArray<TSharedPtr<FTopologicalFace>> CandidateFacesForMesh; // first in first out
+	CandidateFacesForMesh.Reserve(100);
+
+	TArray<TSharedPtr<FTopologicalFace>> SecondChoiceOfCandidateFacesForMesh; // first in first out
+	SecondChoiceOfCandidateFacesForMesh.Reserve(100);
+
+	TFunction<void(TSharedRef<FTopologicalFace>)> MeshFace = [&](TSharedRef<FTopologicalFace> Face)
+	{
+#ifdef DISPLAYDEBUGMESHFACEBYFACESTEP
+		Open3DDebugSession(TEXT("Surface " + Utils::ToFString(Face->GetId()));
+		Display(Face);
+		Close3DDebugSession();
+#endif
+		Mesh(Face);
+
+#ifdef DISPLAYDEBUGMESHFACEBYFACESTEP
+		Open3DDebugSession(TEXT("Mesh " + Utils::ToFString(Face->GetId()));
+		DisplayMesh((FFaceMesh&)Face->GetOrCreateMesh(GetMeshModel()));
+		Close3DDebugSession();
+		//Wait();
+#endif
+
+		if (Face->HasMarker1())
+		{
+			CandidateFacesForMesh.RemoveSingle(Face);
+		}
+		if (Face->HasMarker2())
+		{
+			SecondChoiceOfCandidateFacesForMesh.RemoveSingle(Face);
+		}
+
+		const TSharedPtr<FTopologicalLoop>& Loop = Face->GetLoops()[0];
+		for (const FOrientedEdge& OrientedEdge : Loop->GetEdges())
+		{
+			const TSharedPtr<FTopologicalEdge>& Edge = OrientedEdge.Entity;
+			Edge->SetMarker1(); // tmp for debug
+			for (TWeakPtr<FTopologicalEdge> WeakEdge : Edge->GetTwinsEntities())
+			{
+				TSharedPtr<FTopologicalEdge> NextEdge = WeakEdge.Pin();
+				if (NextEdge->HasMarker1())
+				{
+					continue;
+				}
+
+				TSharedPtr<FTopologicalFace> NextFace = NextEdge->GetFace();
+				if (!NextFace.IsValid())
+				{
+					continue;
+				}
+
+				if (!NextFace->HasMarker3())
+				{
+					// not in the scope of surface to mesh
+					continue;
+				}
+
+				int32 EdgeIndex;
+				int32 LoopIndex;
+				NextFace->GetEdgeIndex(NextEdge, LoopIndex, EdgeIndex);
+				if (LoopIndex > 0)
+				{
+					continue;
+				}
+				int32 SideIndex = NextFace->GetSideIndex(NextEdge);
+				if (SideIndex == -1)
+				{
+					continue;
+				}
+
+				FEdge2DProperties& SideProperty = NextFace->GetSideProperty(SideIndex);
+
+				double EdgeLength = NextEdge->Length();
+				SideProperty.MeshedLength += EdgeLength;
+				NextFace->AddMeshedLength(EdgeLength);
+				if ((SideProperty.Length3D - SideProperty.MeshedLength) < GeometricTolerance)
+				{
+					if (!SideProperty.bIsMesh)
+					{
+						SideProperty.bIsMesh = true;
+						NextFace->MeshedSideNum()++;
+					}
+
+					if (!NextFace->HasMarker1())
+					{
+						NextFace->SetMarker1();
+						CandidateFacesForMesh.Add(NextFace);
+					}
+				}
+				else
+				{
+					if (!NextFace->HasMarker2())
+					{
+						NextFace->SetMarker2();
+						SecondChoiceOfCandidateFacesForMesh.Add(NextFace);
+					}
+				}
+			}
+		}
+	};
+
+	TFunction<void(TSharedRef<FTopologicalFace>)> MeshFacesByFront = [&](TSharedRef<FTopologicalFace> Face)
+	{
+		if (Face->IsMeshed())
+		{
+			return;
+		}
+
+		MeshFace(Face);
+
+		while (CandidateFacesForMesh.Num() || SecondChoiceOfCandidateFacesForMesh.Num())
+		{
+			// the candidate are sorted according to the number of meshed side 
+			Algo::Sort(CandidateFacesForMesh, [](TSharedPtr<FTopologicalFace> Surface1, TSharedPtr<FTopologicalFace> Surface2) { return Surface1->MeshedSideNum() > Surface2->MeshedSideNum(); });
+
+			int32 IndexOfBestCandidate = -1;
+			double CandidateMeshedSideRatio = 0;
+
+			// The first choice will be done in the first set of surface with the max meshed side numbers.
+			if (CandidateFacesForMesh.Num())
+			{
+				int32 MaxMeshedSideNum = CandidateFacesForMesh[0]->MeshedSideNum();
+				for (int32 Index = 0; Index < CandidateFacesForMesh.Num(); ++Index)
+				{
+					if (CandidateFacesForMesh[Index]->IsMeshed())
+					{
+						CandidateFacesForMesh.RemoveAt(Index);
+						--Index;
+					}
+				}
+
+				// next face with side well meshed are preferred
+				int32 Index = 0;
+				for (; Index < CandidateFacesForMesh.Num(); ++Index)
+				{
+					TSharedPtr<FTopologicalFace> CandidateSurface = CandidateFacesForMesh[Index];
+					if (CandidateSurface->MeshedSideNum() < MaxMeshedSideNum)
+					{
+						break;
+					}
+
+					if (CandidateMeshedSideRatio < CandidateSurface->MeshedSideRatio())
+					{
+						CandidateMeshedSideRatio = CandidateSurface->MeshedSideRatio();
+						IndexOfBestCandidate = Index;
+					}
+				}
+
+				// if no candidate has been selected, the choice is done on all next surfaces
+				if (IndexOfBestCandidate == -1)
+				{
+					for (; Index < CandidateFacesForMesh.Num(); ++Index)
+					{
+						TSharedPtr<FTopologicalFace> CandidateSurface = CandidateFacesForMesh[Index];
+						if (CandidateMeshedSideRatio < CandidateSurface->MeshedSideRatio())
+						{
+							CandidateMeshedSideRatio = CandidateSurface->MeshedSideRatio();
+							IndexOfBestCandidate = Index;
+						}
+					}
+				}
+
+				if (IndexOfBestCandidate >= 0)
+				{
+					ensureCADKernel(CandidateFacesForMesh[IndexOfBestCandidate].IsValid());
+					MeshFace(CandidateFacesForMesh[IndexOfBestCandidate].ToSharedRef());
+					continue;
+				}
+			}
+
+			for (int32 Index = 0; Index < SecondChoiceOfCandidateFacesForMesh.Num(); ++Index)
+			{
+				TSharedPtr<FTopologicalFace> CandidateSurface = SecondChoiceOfCandidateFacesForMesh[Index];
+				if (CandidateMeshedSideRatio < CandidateSurface->MeshedSideRatio())
+				{
+					CandidateMeshedSideRatio = CandidateSurface->MeshedSideRatio();
+					IndexOfBestCandidate = Index;
+				}
+			}
+			if (IndexOfBestCandidate >= 0)
+			{
+				ensureCADKernel(SecondChoiceOfCandidateFacesForMesh[IndexOfBestCandidate].IsValid());
+				MeshFace(SecondChoiceOfCandidateFacesForMesh[IndexOfBestCandidate].ToSharedRef());
+			}
+		}
+	};
+
+	// the front is initialized with quad surface
+	for (const FCostToFace& Quad : QuadTrimmedSurfaceSet)
+	{
+		TSharedPtr<FTopologicalFace> Surface = Quad.Face;
+		MeshFacesByFront(Surface.ToSharedRef());
+	}
+
+	// the the other surface
+	for (TSharedPtr<FTopologicalFace> Face : Faces)
+	{
+		if (!Face->IsMeshed())
+		{
+			MeshFacesByFront(Face.ToSharedRef());
+		}
+	}
+
+}
+
+#ifdef DEBUG_INTERSECTEDGEISOS
+TMap<int32, int32> SurfaceDrawed;
+bool bDisplayIsoCurve = true;
+
+void FParametricMesher::DebugIntersectEdgeIsos(TSharedPtr<FTopologicalEdge> Edge, TSharedPtr<FTopologicalFace> Surface, const TArray<FPoint>& EdgeCrossingPoints2D, const double ToleranceIso, const TArray<double>& IsoCoordinates, EIso TypeIso)
+{
+	ensureCADKernel(Edge.IsValid());
+	if (SurfaceDrawed.Find(Surface->GetId()) == nullptr)
+	{
+		SurfaceDrawed.Add(Surface->GetId(), 0);
+	}
+
+	if (bDisplayIsoCurve && SurfaceDrawed[Surface->GetId()] < 2)
+	{
+		SurfaceDrawed[Surface->GetId()]++;
+
+		FSurfacicBoundary Bounds = Surface->GetBoundary();
+
+		{
+			F3DDebugSession _(FString::Printf(TEXT("Iso %s 3D %d"), TypeIso == EIso::IsoU ? TEXT("U") : TEXT("V"), Surface->GetId()));
+			for (double U : IsoCoordinates)
+			{
+				DisplayIsoCurve(Surface->GetCarrierSurface(), U, TypeIso);
+			}
+		}
+		F3DDebugSession _(FString::Printf(TEXT("Iso %s 2D %d"), TypeIso == EIso::IsoU ? TEXT("U") : TEXT("V"), Surface->GetId()));
+		if (TypeIso == EIso::IsoU)
+		{
+			for (double U : IsoCoordinates)
+			{
+				FPoint2D Start(U, Bounds.VMin);
+				FPoint2D End(U, Bounds.VMax);
+				DisplaySegment(Start, End, 0, EVisuProperty::Iso);
+			}
+		}
+		else
+		{
+			for (double V : IsoCoordinates)
+			{
+				FPoint2D Start(Bounds.UMin, V);
+				FPoint2D End(Bounds.UMax, V);
+				DisplaySegment(Start, End, 0, EVisuProperty::Iso);
+			}
+		}
+	}
+	{
+		F3DDebugSession _(FString::Printf(TEXT("Edge 2D %d"), Edge->GetId()));
+		Display(EdgeCrossingPoints2D[0]);
+		for (int32 Index = 1; Index < EdgeCrossingPoints2D.Num(); ++Index)
+		{
+			DisplaySegment(EdgeCrossingPoints2D[Index - 1], EdgeCrossingPoints2D[Index], 0, EVisuProperty::EdgeMesh);
+			Display(EdgeCrossingPoints2D[Index]);
+		}
+	}
+	{
+		F3DDebugSession _(FString::Printf(TEXT("Edge 3D %d"), Edge->GetId()));
+		Display(Edge);
+	}
+}
+#endif
+
+void FParametricMesher::IntersectEdgeIsos(TSharedPtr<FTopologicalEdge> Edge, const TArray<FPoint2D>& EdgeCrossingPoints2D, const double ToleranceIso, const TArray<double>& IsoCoordinates, EIso TypeIso, TArray<double>& Intersection)
+{
+	Intersection.Reserve(IsoCoordinates.Num());
+
+	{
+		double EdgeMinU = HUGE_VALUE;
+		double EdgeMaxU = -HUGE_VALUE;
+
+		for (const FPoint2D& EdgeCoord : EdgeCrossingPoints2D)
+		{
+			EdgeMinU = FMath::Min(EdgeMinU, EdgeCoord[(uint8)TypeIso]);
+			EdgeMaxU = FMath::Max(EdgeMaxU, EdgeCoord[(uint8)TypeIso]);
+		}
+
+		if ((EdgeMaxU - EdgeMinU) < ToleranceIso)
+		{
+			// the edge is on an IsoCurve on Iso axis
+			return;
+		}
+	}
+
+	double ToleranceEdge = Edge->GetCurve()->GetParametricTolerance();
+	const TArray<double>& EdgeCrossingPointsU = Edge->GetCrossingPointUs();
+
+	const FLinearBoundary& EdgeBounds = Edge->GetBoundary();
+	double MinU;
+	double MaxU;
+	for (int32 CrossingPointIndex = 0; CrossingPointIndex < EdgeCrossingPoints2D.Num() - 1; ++CrossingPointIndex)
+	{
+		if (FMath::Abs(EdgeCrossingPoints2D[CrossingPointIndex + 1][(uint8)TypeIso] - EdgeCrossingPoints2D[CrossingPointIndex][(uint8)TypeIso]) < SMALL_NUMBER)
+		{
+			continue;
+		}
+
+		Sort(EdgeCrossingPoints2D[CrossingPointIndex][(uint8)TypeIso], EdgeCrossingPoints2D[CrossingPointIndex + 1][(uint8)TypeIso], MinU, MaxU);
+		
+		for (double IsoCoord : IsoCoordinates)
+		{
+			// If the border is tangent to the iso, the intersection is not compute
+			if ((IsoCoord > MinU - SMALL_NUMBER) && (IsoCoord < MaxU))
+			{
+				double EdgeLocalSlop = ComputeUnorientedSlope(EdgeCrossingPoints2D[CrossingPointIndex], EdgeCrossingPoints2D[CrossingPointIndex + 1], 0);
+				if (TypeIso == EIso::IsoV)
+				{
+					if (EdgeLocalSlop < 0.1 || EdgeLocalSlop > 3.9)
+					{
+						continue;
+					}
+				}
+				else
+				{
+					if (EdgeLocalSlop < 2.1 && EdgeLocalSlop > 1.9)
+					{
+						continue;
+					}
+				}
+
+				double EdgeCoord = (IsoCoord - EdgeCrossingPoints2D[CrossingPointIndex][(uint8)TypeIso]) / (EdgeCrossingPoints2D[CrossingPointIndex + 1][(uint8)TypeIso] - EdgeCrossingPoints2D[CrossingPointIndex][(uint8)TypeIso]);
+				EdgeCoord *= (EdgeCrossingPointsU[CrossingPointIndex + 1] - EdgeCrossingPointsU[CrossingPointIndex]);
+				EdgeCoord += EdgeCrossingPointsU[CrossingPointIndex];
+				if ((EdgeCoord < (EdgeBounds.GetMin() + ToleranceEdge)) || (EdgeCoord > (EdgeBounds.GetMax() - ToleranceEdge)))
+				{
+					continue;
+				}
+				Intersection.Add(EdgeCoord);
+			}
+		}
+	}
+}
+
+// =========================================================================================================================================================================================================
+// =========================================================================================================================================================================================================
+// =========================================================================================================================================================================================================
+//
+//
+//                                                                            NOT YET REVIEWED
+//
+//
+// =========================================================================================================================================================================================================
+// =========================================================================================================================================================================================================
+// =========================================================================================================================================================================================================
+
+void FParametricMesher::MeshThinZoneEdges(FGrid& Grid)
+{
+#ifdef DEBUG_MESHTHINSURF
+	Open3DDebugSession(FString::Printf(TEXT("thin Surfaces cutting on surf %d"), Grid.GetFace()->GetId()));
+#endif
+
+	const TArray<FThinZone2D>& ThinZones = Grid.GetThinZones();
+
+	FTimePoint MeshStartTime = FChrono::Now();
+
+	for (const FThinZone2D& Zone : ThinZones)
+	{
+		bool bFirstSideIsPartiallyMeshed = Zone.GetFirstSide().IsPartiallyMeshed();
+		bool bSecondSideIsPartiallyMeshed = Zone.GetSecondSide().IsPartiallyMeshed();
+
+		if (bFirstSideIsPartiallyMeshed && bSecondSideIsPartiallyMeshed)
+		{
+			// the most meshed edge is meshed first
+			double FirstSideMeshedLength = Zone.GetFirstSide().GetMeshedLength();
+			double SecondSideMeshedLength = Zone.GetSecondSide().GetMeshedLength();
+			if (FirstSideMeshedLength > SecondSideMeshedLength)
+			{
+				bSecondSideIsPartiallyMeshed = false;
+			}
+			else
+			{
+				bFirstSideIsPartiallyMeshed = false;
+			}
+		}
+
+		if (!bFirstSideIsPartiallyMeshed && !bSecondSideIsPartiallyMeshed)
+		{
+			if (Zone.GetFirstSide().GetLength() > Zone.GetSecondSide().GetLength())
+			{
+				GetThinZoneBoundary(Zone.GetFirstSide());
+				GetThinZoneBoundary(Zone.GetSecondSide());
+
+				MeshThinZoneSide(Zone.GetFirstSide());
+			}
+			else
+			{
+				GetThinZoneBoundary(Zone.GetFirstSide());
+				GetThinZoneBoundary(Zone.GetSecondSide());
+
+				MeshThinZoneSide(Zone.GetSecondSide());
+			}
+		}
+		else if (bFirstSideIsPartiallyMeshed && !bSecondSideIsPartiallyMeshed)
+		{
+			MeshThinZoneSide(Zone.GetFirstSide());
+			GetThinZoneBoundary(Zone.GetSecondSide());
+		}
+		else if (!bFirstSideIsPartiallyMeshed && bSecondSideIsPartiallyMeshed)
+		{
+			MeshThinZoneSide(Zone.GetSecondSide());
+			GetThinZoneBoundary(Zone.GetFirstSide());
+		}
+	}
+
+	// if the extremity of the thin zone are connected by a short edges path, the edges path are not discretized to avoid a well discretized edge connecting to thin sides
+
+#ifdef DEBUG_MESHTHINSURF
+	Close3DDebugSession();
+#endif
+
+#ifdef DEBUG_MESHTHINSURF
+	Open3DDebugSession(FString("Mesh of ThinZone 2D of surf " + Utils::ToFString(Grid.GetFace()->GetId())));
+	for (const FThinZone2D& Zone : ThinZones)
+	{
+		TArray<TSharedPtr<FTopologicalEdge>> OutEdges;
+		Zone.GetFirstSide().GetEdges(OutEdges);
+		for (TSharedPtr<FTopologicalEdge> Edge : OutEdges)
+		{
+			TSharedPtr<FTopologicalEdge> ActiveEdge = StaticCastSharedRef<FTopologicalEdge>(Edge->GetLinkActiveEntity());
+
+			TArray<double> ImposedCuttingPointEdgeCoordinates;
+			GetCuttingPointCoordinates(ActiveEdge->GetImposedCuttingPoints(), ImposedCuttingPointEdgeCoordinates);
+
+			TArray<double> ImposedCuttingPointU;
+			Edge->TransformActiveEdgeCoordinatesToLocalCoordinates(ImposedCuttingPointEdgeCoordinates, ImposedCuttingPointU);
+
+			TArray<FPoint> ImposedCuttingPoint2D;
+			Edge->Approximate2DPoints(ImposedCuttingPointU, ImposedCuttingPoint2D);
+			for (FPoint Point : ImposedCuttingPoint2D)
+			{
+				Display(Point);
+			}
+		}
+
+		Zone.GetSecondSide().GetEdges(OutEdges);
+		for (TSharedPtr<FTopologicalEdge> Edge : OutEdges)
+		{
+			TSharedPtr<FTopologicalEdge> ActiveEdge = StaticCastSharedRef<FTopologicalEdge>(Edge->GetLinkActiveEntity());
+
+			TArray<double> ImposedCuttingPointEdgeCoordinates;
+			GetCuttingPointCoordinates(ActiveEdge->GetImposedCuttingPoints(), ImposedCuttingPointEdgeCoordinates);
+
+			TArray<double> ImposedCuttingPointCoordinates;
+			Edge->TransformActiveEdgeCoordinatesToLocalCoordinates(ImposedCuttingPointEdgeCoordinates, ImposedCuttingPointCoordinates);
+
+			TArray<FPoint> ImposedCuttingPoint2D;
+			Edge->Approximate2DPoints(ImposedCuttingPointCoordinates, ImposedCuttingPoint2D);
+			for (FPoint Point : ImposedCuttingPoint2D)
+			{
+				Display(Point);
+			}
+		}
+	}
+	Close3DDebugSession();
+#endif
+	Chronos.GlobalMeshThinZones += FChrono::Elapse(MeshStartTime);
+
+}
+
+void AddActiveEdgeThinZone(TSharedPtr<FTopologicalEdge> Edge, TSharedPtr<FTopologicalEdge> ActiveEdge, FLinearBoundary& SideEdgeCoordinate)
+{
+	TArray<double> SideEdgeBound;
+	SideEdgeBound.SetNum(2);
+	SideEdgeBound[0] = SideEdgeCoordinate.GetMin();
+	SideEdgeBound[1] = SideEdgeCoordinate.GetMax();
+
+	TArray<double> ActiveEdgeThinZone;
+	Edge->TransformActiveEdgeCoordinatesToLocalCoordinates(SideEdgeBound, ActiveEdgeThinZone);
+	FLinearBoundary ThinZoneBoundary(ActiveEdgeThinZone[0], ActiveEdgeThinZone[1]);
+	ActiveEdge->AddThinZone(ThinZoneBoundary);
+};
+
+void FParametricMesher::GetThinZoneBoundary(const FThinZoneSide& Side)
+{
+	TSharedPtr<FTopologicalEdge> Edge = nullptr;
+	TSharedPtr<FTopologicalEdge> ActiveEdge = nullptr;
+	FLinearBoundary SideEdgeCoordinate;
+
+	for (const FEdgeSegment* EdgeSegment : Side.GetSegments())
+	{
+		double UMin = EdgeSegment->GetCoordinate(ELimit::Start);
+		double UMax = EdgeSegment->GetCoordinate(ELimit::End);
+		Sort(UMin, UMax);
+
+		if (Edge != EdgeSegment->GetEdge())
+		{
+			if (Edge)
+			{
+				AddActiveEdgeThinZone(Edge, ActiveEdge, SideEdgeCoordinate);
+			}
+
+			Edge = EdgeSegment->GetEdge();
+			ActiveEdge = StaticCastSharedRef<FTopologicalEdge>(Edge->GetLinkActiveEntity());
+
+			SideEdgeCoordinate.Set(UMin, UMax);
+		}
+		else
+		{
+			SideEdgeCoordinate.ExtendTo(UMin, UMax);
+		}
+	};
+	AddActiveEdgeThinZone(Edge, ActiveEdge, SideEdgeCoordinate);
+}
+
+void FParametricMesher::MeshThinZoneSide(const FThinZoneSide& Side)
+{
+	typedef TFunction<bool(double, double)> CompareMethode;
+
+	TSharedPtr<FTopologicalEdge> Edge = nullptr;
+	TSharedPtr<FTopologicalEdge> ActiveEdge = nullptr;
+	int32 Index = 0;
+	int32 Increment = 1;
+	TArray<double> EdgeCuttingPointCoordinates;
+	FLinearBoundary SideEdgeCoordinate;
+	const TArray<int32>* NodeIndices = nullptr;
+
+	TFunction<void(int32&, int32&, const FEdgeSegment*, double, double)> AddImposedCuttingPoint = [&](int32& Index, int32& Increment, const FEdgeSegment* EdgeSegment, double UMin, double UMax)
+	{
+#ifdef DebugMeshThinSurf
+		DisplaySegment(EdgeSegment->GetExtemity(ELimit::End), EdgeSegment->GetExtemity(ELimit::Start));
+#endif
+		for (; Index >= 0 && Index < EdgeCuttingPointCoordinates.Num(); Index += Increment)
+		{
+			if (EdgeCuttingPointCoordinates[Index] < UMin)
+			{
+				break;
+			}
+			if (EdgeCuttingPointCoordinates[Index] > UMax)
+			{
+				break;
+			}
+			FPoint CuttingPoint3D = EdgeSegment->ComputeEdgePoint(EdgeCuttingPointCoordinates[Index]);
+
+			const FEdgeSegment* ClosedSegment = EdgeSegment->GetClosedSegment();
+			if (ClosedSegment == nullptr)
+			{
+				Wait();
+			}
+			double OppositeCuttingPointSegmentU;
+			FPoint OppositeCuttingPoint3D = ClosedSegment->ProjectPoint(CuttingPoint3D, OppositeCuttingPointSegmentU);
+
+			double OppositeCuttingPointU = 0;
+			TSharedPtr<FTopologicalEdge> OpositEdge = nullptr;
+			if (OppositeCuttingPointSegmentU == 0 && ClosedSegment->GetPrevious()->GetClosedSegment())
+			{
+				OppositeCuttingPoint3D = ClosedSegment->GetPrevious()->ProjectPoint(CuttingPoint3D, OppositeCuttingPointSegmentU);
+				OppositeCuttingPointU = ClosedSegment->GetPrevious()->ComputeEdgeCoordinate(OppositeCuttingPointSegmentU);
+				OpositEdge = ClosedSegment->GetPrevious()->GetEdge();
+			}
+			else if (OppositeCuttingPointSegmentU == 1 && ClosedSegment->GetNext()->GetClosedSegment())
+			{
+				OppositeCuttingPoint3D = ClosedSegment->GetNext()->ProjectPoint(CuttingPoint3D, OppositeCuttingPointSegmentU);
+				OppositeCuttingPointU = ClosedSegment->GetNext()->ComputeEdgeCoordinate(OppositeCuttingPointSegmentU);
+				OpositEdge = ClosedSegment->GetNext()->GetEdge();
+			}
+			else
+			{
+				OpositEdge = ClosedSegment->GetEdge();
+				OppositeCuttingPointU = ClosedSegment->ComputeEdgeCoordinate(OppositeCuttingPointSegmentU);
+			}
+
+			double OppositeActiveEdgeCuttingPointU = OpositEdge->TransformLocalCoordinateToActiveEdgeCoordinate(OppositeCuttingPointU);
+
+#ifdef DEBUG_MESHTHINSURF
+			DisplayPoint(CuttingPoint3D);
+			DisplaySegment(ClosedSegment->GetExtemity(ELimit::End), ClosedSegment->GetExtemity(ELimit::Start));
+			DisplayPoint(OppositeCuttingPoint3D);
+			DisplaySegment(CuttingPoint3D, OppositeCuttingPoint3D);
+#endif
+			OpositEdge->GetLinkActiveEdge()->AddImposedCuttingPointU(OppositeActiveEdgeCuttingPointU, (*NodeIndices)[Index]);
+		}
+	};
+
+	TFunction<void(TArray<double>&, double, int32&, CompareMethode)> FindFirstIndexForward = [](TArray<double>& EdgeCuttingPointU, double ULimit, int32& OutIndex, CompareMethode Compare)
+	{
+		for (; OutIndex < EdgeCuttingPointU.Num(); ++OutIndex)
+		{
+			if (Compare(ULimit, EdgeCuttingPointU[OutIndex]))
+			{
+				break;
+			}
+		}
+	};
+
+	TFunction<void(TArray<double>&, double, int32&, CompareMethode)> FindFirstIndexBackward = [](TArray<double>& EdgeCuttingPointU, double ULimit, int32& OutIndex, CompareMethode Compare)
+	{
+		for (; OutIndex >= 0; --OutIndex)
+		{
+			if (Compare(ULimit, EdgeCuttingPointU[OutIndex]))
+			{
+				break;
+			}
+		}
+	};
+
+	TFunction<void(const FEdgeSegment* EdgeSegment)> Process = [&](const FEdgeSegment* EdgeSegment)
+	{
+		double UMin = EdgeSegment->GetCoordinate(ELimit::Start);
+		double UMax = EdgeSegment->GetCoordinate(ELimit::End);
+		Sort(UMin, UMax);
+
+		if (Edge != EdgeSegment->GetEdge())
+		{
+			if (Edge)
+			{
+				AddActiveEdgeThinZone(Edge, ActiveEdge, SideEdgeCoordinate);
+#ifdef DebugMeshThinSurf
+				Close3DDebugSession();
+#endif
+			}
+
+			Edge = EdgeSegment->GetEdge();
+			ActiveEdge = StaticCastSharedRef<FTopologicalEdge>(Edge->GetLinkActiveEntity());
+
+			SideEdgeCoordinate.Set(UMin, UMax);
+
+			if (!ActiveEdge->IsMeshed())
+			{
+				Mesh(Edge.ToSharedRef(), Edge->GetFace());
+#ifdef DebugMeshThinSurf
+				Open3DDebugSession(TEXT("Mesh of Edge"));
+				DisplayMesh((const TSharedPtr<FEdgeMesh>&) ActiveEdge->GetMesh());
+				Close3DDebugSession();
+#endif
+			}
+#ifdef DebugMeshThinSurf
+			Open3DDebugSession(TEXT("Projection of mesh"));
+#endif
+			NodeIndices = &ActiveEdge->GetOrCreateMesh(MeshModel)->EdgeVerticesIndex;
+
+			TArray<double> CuttingPointCoordinates;
+			GetCuttingPointCoordinates(ActiveEdge->GetCuttingPoints(), CuttingPointCoordinates);
+
+			Edge->TransformActiveEdgeCoordinatesToLocalCoordinates(CuttingPointCoordinates, EdgeCuttingPointCoordinates);
+
+			if ((EdgeCuttingPointCoordinates[0] < EdgeCuttingPointCoordinates[1]) == (EdgeSegment->GetCoordinate(ELimit::Start) < EdgeSegment->GetCoordinate(ELimit::End)))
+			{
+				Index = 0;
+				if (EdgeCuttingPointCoordinates[0] < EdgeCuttingPointCoordinates[1])
+				{
+					FindFirstIndexForward(EdgeCuttingPointCoordinates, UMin, Index, [](double Value1, double Value2) {return (Value1 < Value2); });
+				}
+				else
+				{
+					FindFirstIndexForward(EdgeCuttingPointCoordinates, UMax, Index, [](double Value1, double Value2) {return (Value1 > Value2); });
+				}
+				Increment = 1;
+			}
+			else
+			{
+				Index = (int32)EdgeCuttingPointCoordinates.Num() - 1;
+				if (EdgeCuttingPointCoordinates[0] < EdgeCuttingPointCoordinates[1])
+				{
+					FindFirstIndexBackward(EdgeCuttingPointCoordinates, UMax, Index, [](double Value1, double Value2) {return (Value1 > Value2); });
+				}
+				else
+				{
+					FindFirstIndexBackward(EdgeCuttingPointCoordinates, UMax, Index, [](double Value1, double Value2) {return (Value1 < Value2); });
+				}
+				Increment = -1;
+			}
+		}
+		else
+		{
+			SideEdgeCoordinate.ExtendTo(UMin, UMax);
+		}
+
+		AddImposedCuttingPoint(Index, Increment, EdgeSegment, UMin, UMax);
+	};
+
+#ifdef DebugMeshThinSurf
+	Open3DDebugSession(TEXT("MeshThinZoneSide"));
+#endif
+
+	if (Side.IsFirstSide())
+	{
+		for (const FEdgeSegment* EdgeSegment : Side.GetSegments())
+		{
+			Process(EdgeSegment);
+		}
+	}
+	else
+	{
+		const TArray<FEdgeSegment*>& Segments = Side.GetSegments();
+		for (int32 SegmentIndex = Segments.Num() - 1; SegmentIndex >= 0; --SegmentIndex)
+		{
+			Process(Segments[SegmentIndex]);
+		}
+	}
+	AddActiveEdgeThinZone(Edge, ActiveEdge, SideEdgeCoordinate);
+
+#ifdef DebugMeshThinSurf
+	Close3DDebugSession();
+	Close3DDebugSession();
+#endif
+
+}
+
+
+void FParametricMesher::InitParameters(const FString& ParametersString)
+{
+	Parameters->SetFromString(ParametersString);
+}
diff --git a/Engine/Source/Runtime/Datasmith/CADKernel/Private/CADKernel/Mesh/Structure/EdgeSegment.h b/Engine/Source/Runtime/Datasmith/CADKernel/Private/CADKernel/Mesh/Structure/EdgeSegment.h
new file mode 100644
--- /dev/null
+++ b/Engine/Source/Runtime/Datasmith/CADKernel/Private/CADKernel/Mesh/Structure/EdgeSegment.h
@@ -0,0 +1,260 @@
+// Copyright Epic Games, Inc. All Rights Reserved.
+#pragma once
+
+#include "CADKernel/Core/Types.h"
+
+#include "CADKernel/Math/Geometry.h"
+#include "CADKernel/Core/HaveStates.h"
+#include "CADKernel/Math/Boundary.h"
+#include "CADKernel/Math/Point.h"
+#include "CADKernel/Math/SlopeUtils.h"
+
+namespace CADKernel
+{
+	class FTopologicalLoop;
+	class FTopologicalEdge;
+	class FGrid;
+	class FThinZone;
+	class FThinZone2DFinder;
+	class FTopologicalVertex;
+
+	class FEdgeSegment : public FHaveStates
+	{
+	private:
+		TSharedPtr<FTopologicalEdge> Edge;
+		double Coordinates[2];
+		FPoint2D Points[2];
+
+		FEdgeSegment* NextSegment;
+		FEdgeSegment* PreviousSegment;
+
+		FEdgeSegment* ClosedSegment[2] = { nullptr, nullptr };
+
+		FSurfacicBoundary Boundary;
+		double AxisMin;
+
+		double SquareDistanceToClosedSegment;
+		double Length;
+
+		FIdent ChainIndex;
+
+		FIdent Id;
+		static FIdent LastId;
+
+	public:
+		FEdgeSegment()
+			: NextSegment(nullptr)
+			, PreviousSegment(nullptr)
+			, AxisMin(0.)
+			, SquareDistanceToClosedSegment(HUGE_VALUE)
+			, Length(-1.)
+			, ChainIndex(0)
+			, Id(0)
+		{
+		};
+
+		virtual ~FEdgeSegment() = default;
+		
+		void SetBoundarySegment(bool bInIsInnerLoop, const TSharedPtr<FTopologicalEdge>& InEdge, double InStartU, double InEndU, const FPoint2D& InStartPoint, const FPoint2D& InEndPoint)
+		{
+			if (bInIsInnerLoop)
+			{
+				SetInner();
+			}
+
+			Edge = InEdge;
+			Coordinates[ELimit::Start] = InStartU;
+			Coordinates[ELimit::End] = InEndU;
+			Points[ELimit::Start] = InStartPoint;
+			Points[ELimit::End]   = InEndPoint;
+			NextSegment = nullptr;
+			PreviousSegment = nullptr;
+			ClosedSegment[0] = nullptr;
+			ClosedSegment[1] = nullptr;
+
+			SquareDistanceToClosedSegment = HUGE_VAL;
+			Length = Points[ELimit::Start].Distance(Points[ELimit::End]);
+
+			Id = ++LastId;
+			ChainIndex = 0;
+
+			Boundary.Set(Points[ELimit::Start], Points[ELimit::End]);
+
+			AxisMin = Boundary[EIso::IsoU].Min + Boundary[EIso::IsoV].Min;
+		};
+
+		double GetAxeMin() const
+		{
+			return AxisMin;
+		}
+
+		FIdent GetChainIndex() const
+		{
+			return ChainIndex;
+		}
+
+		void SetChainIndex(FIdent index)
+		{
+			ChainIndex = index;
+		}
+
+		bool IsInner() const
+		{
+			return ((States & EHaveStates::IsInner) == EHaveStates::IsInner);
+		}
+
+		void SetInner()
+		{
+			States |= EHaveStates::IsInner;
+		}
+
+		FIdent GetId()
+		{
+			return Id;
+		}
+
+		const TSharedPtr<FTopologicalEdge> GetEdge() const
+		{
+			return Edge;
+		}
+
+		//const TSharedPtr<FTopologicalLoop>& GetLoop() const
+		//{
+		//	return Edge->GetLoop();
+		//}
+
+		double GetLength() const
+		{
+			return Length;
+		}
+
+		FPoint GetCenter() const
+		{
+			return (Points[ELimit::Start] + Points[ELimit::End]) * 0.5;
+		}
+
+		FPoint ComputeEdgePoint(double EdgeParamU) const
+		{
+			double SegmentParamS = (EdgeParamU - Coordinates[ELimit::Start]) / (Coordinates[ELimit::End] - Coordinates[ELimit::Start]);
+			return Points[ELimit::Start] + (Points[ELimit::End] - Points[ELimit::Start]) * SegmentParamS;
+		}
+
+		constexpr const FPoint2D& GetExtemity(const ELimit Limit) const
+		{
+			if (Limit == Start)
+			{
+				return Points[ELimit::Start];
+			}
+			else
+			{
+				return Points[ELimit::End];
+			}
+		}
+
+		constexpr double GetCoordinate(const ELimit Limit) const
+		{
+			if(Limit == Start)
+			{
+				return Coordinates[ELimit::Start];
+			}
+			else
+			{
+				return Coordinates[ELimit::End];
+			}
+		}
+
+		/**
+		 * Compute the slope of the input Segment according to this.
+		 */
+		double ComputeUnorientedSlopeOf(FEdgeSegment* Segment)
+		{
+			double ReferenceSlope = ComputeSlope(Points[ELimit::Start], Points[ELimit::End]);
+			return ComputeUnorientedSlope(Segment->Points[ELimit::Start], Segment->Points[ELimit::End], ReferenceSlope);
+		}
+
+		/**
+		 * Compute the slope of the Segment defined by the two input points according to this.
+		 */
+		double ComputeUnorientedSlopeOf(FPoint2D& Middle, FPoint2D& Projection)
+		{
+			double ReferenceSlope = ComputeSlope(Points[ELimit::Start], Points[ELimit::End]);
+			return ComputeUnorientedSlope(Projection, Middle, ReferenceSlope);
+		}
+
+		FEdgeSegment* GetNext() const
+		{
+			return NextSegment;
+		}
+
+		FEdgeSegment* GetPrevious() const
+		{
+			return PreviousSegment;
+		}
+
+		FEdgeSegment* GetClosedSegment() const
+		{
+			return ClosedSegment[0];
+		}
+
+		void ResetClosedData()
+		{
+			if (ClosedSegment[0]->ClosedSegment[0] == this)
+			{
+				ClosedSegment[0]->ClosedSegment[0] = ClosedSegment[0]->ClosedSegment[1];
+			}
+			else
+			{
+				ClosedSegment[0]->ClosedSegment[1] = nullptr;
+			}
+
+			ClosedSegment[0] = nullptr;
+			SquareDistanceToClosedSegment = HUGE_VAL;
+		}
+
+		void SetClosedSegment(FEdgeSegment* InSegmentA, FEdgeSegment* InSegmentB, double InDistance, bool bFirstChoice)
+		{
+			ClosedSegment[0] = InSegmentA;
+			SquareDistanceToClosedSegment = InDistance;
+
+			if (bFirstChoice && InDistance < InSegmentA->SquareDistanceToClosedSegment)
+			{
+				InSegmentA->ClosedSegment[0] = this;
+				InSegmentA->SquareDistanceToClosedSegment = InDistance;
+			}
+			if (InSegmentB && InDistance < InSegmentB->SquareDistanceToClosedSegment)
+			{
+				InSegmentB->ClosedSegment[0] = this;
+				InSegmentB->SquareDistanceToClosedSegment = InDistance;
+			}
+		}
+
+		double GetClosedSquareDistance() const
+		{
+			return SquareDistanceToClosedSegment;
+		}
+
+		void SetNext(FEdgeSegment* Segment)
+		{
+			NextSegment = Segment;
+			Segment->SetPrevious(this);
+		}
+
+		double ComputeEdgeCoordinate(const double SegmentU) const
+		{
+			return Coordinates[ELimit::Start] + (Coordinates[ELimit::End] - Coordinates[ELimit::Start]) * SegmentU;
+		}
+
+		FPoint2D ProjectPoint(const FPoint2D& PointToProject, double& SegmentU) const
+		{
+			return ProjectPointOnSegment(PointToProject, Points[ELimit::Start], Points[ELimit::End], SegmentU, true);
+		}
+
+	private:
+		void SetPrevious(FEdgeSegment* Segment)
+		{
+
+			PreviousSegment = Segment;
+		}
+	};
+}
+
diff --git a/Engine/Source/Runtime/Datasmith/CADKernel/Private/CADKernel/Mesh/Structure/FaceMesh.cpp b/Engine/Source/Runtime/Datasmith/CADKernel/Private/CADKernel/Mesh/Structure/FaceMesh.cpp
new file mode 100644
--- /dev/null
+++ b/Engine/Source/Runtime/Datasmith/CADKernel/Private/CADKernel/Mesh/Structure/FaceMesh.cpp
@@ -0,0 +1,67 @@
+// Copyright Epic Games, Inc. All Rights Reserved.
+
+#include "CADKernel/Mesh/Structure/FaceMesh.h"
+
+#include "CADKernel/Mesh/Structure/EdgeMesh.h"
+#include "CADKernel/Mesh/Structure/VertexMesh.h"
+#include "CADKernel/Topo/TopologicalEdge.h"
+#include "CADKernel/Topo/TopologicalFace.h"
+#include "CADKernel/Topo/TopologicalLoop.h"
+#include "CADKernel/Topo/TopologicalVertex.h"
+
+using namespace CADKernel;
+
+void FFaceMesh::GetNodeIdToCoordinates(TMap<int32, const FPoint*>& NodeIdToCoordinates) const
+{
+	const TSharedRef<FTopologicalFace> Face = StaticCastSharedRef<FTopologicalFace>(GetGeometricEntity());
+	for (const TSharedPtr<FTopologicalLoop> Loop : Face->GetLoops())
+	{
+		const TArray<FOrientedEdge>& Edges = Loop->GetEdges();
+
+		for (const FOrientedEdge& Edge : Edges)
+		{
+			TSharedRef<FTopologicalEdge> ActiveEdge = Edge.Entity->GetLinkActiveEdge();
+			if (ActiveEdge->IsDegenerated()) 
+			{
+				continue;
+			}
+
+			const TSharedRef<FVertexMesh> StartVertexMesh = ActiveEdge->GetStartVertex()->GetMesh();
+			NodeIdToCoordinates.Add(StartVertexMesh->GetMesh(), &StartVertexMesh->GetNodeCoordinates()[0]);
+
+			const TSharedRef<FVertexMesh> EndVertexMesh = ActiveEdge->GetEndVertex()->GetMesh();
+			NodeIdToCoordinates.Add(EndVertexMesh->GetMesh(), &EndVertexMesh->GetNodeCoordinates()[0]);
+
+			const TSharedRef<FEdgeMesh> EdgeMesh = ActiveEdge->GetMesh();
+
+			TArray<int32> NodeIds = EdgeMesh->EdgeVerticesIndex;
+			const FPoint* EdgeNodeCoordinates = EdgeMesh->GetNodeCoordinates().GetData();
+			if(EdgeNodeCoordinates)
+			{
+				for (int32 Index = 1; Index <NodeIds.Num() - 1; Index++)
+				{
+					NodeIdToCoordinates.Add(NodeIds[Index], EdgeNodeCoordinates + Index - 1);
+				}
+			}
+		}
+	}
+
+	const FPoint* Coordinates = GetNodeCoordinates().GetData();
+	for (int32 Index = 0; Index < GetNodeCoordinates().Num(); ++Index)
+	{
+		NodeIdToCoordinates.Add(StartNodeId+Index, Coordinates + Index);
+	}
+}
+
+void FFaceMesh::InverseOrientation()
+{
+	for (int32 Index = 0; Index <TrianglesVerticesIndex.Num(); Index += 3)
+	{
+		Swap(TrianglesVerticesIndex[Index], TrianglesVerticesIndex[Index+1]);
+	}
+
+	for (FPoint& Normal : Normals)
+	{
+		Normal *= -1.;
+	}
+}
