CADKernel Source Code v0.01

diff --git a/Engine/Source/Runtime/Datasmith/CADKernel/Private/CADKernel/Mesh/Structure/Grid.cpp b/Engine/Source/Runtime/Datasmith/CADKernel/Private/CADKernel/Mesh/Structure/Grid.cpp
new file mode 100644
--- /dev/null
+++ b/Engine/Source/Runtime/Datasmith/CADKernel/Private/CADKernel/Mesh/Structure/Grid.cpp
@@ -0,0 +1,1936 @@
+// Copyright Epic Games, Inc. All Rights Reserved.
+
+#include "CADKernel/Mesh/Structure/Grid.h"
+
+#include "CADKernel/Geo/Sampling/SurfacicSampling.h"
+#include "CADKernel/Mesh/Meshers/MesherTools.h"
+#include "CADKernel/Mesh/Structure/EdgeMesh.h"
+#include "CADKernel/Mesh/Structure/VertexMesh.h"
+#include "CADKernel/Topo/TopologicalEdge.h"
+#include "CADKernel/Topo/TopologicalFace.h"
+#include "CADKernel/Utils/Util.h"
+
+#define DEBUG_GRID
+
+namespace CADKernel
+{
+	FGrid::FGrid(TSharedRef<FTopologicalFace>& InFace, TSharedRef<FModelMesh>& InMeshModel)
+		: Face(InFace)
+		, FaceTolerance(InFace->GetIsoTolerances())
+		, MeshModel(InMeshModel)
+		, ThinZoneFinder(*this)
+		, CuttingCoordinates(Face->GetCuttingPointCoordinates())
+	{
+		if (Face->GetId() == SURFACE_TO_DEBUG)
+		{
+			bDisplay = true;
+		}
+	}
+
+	void FGrid::PrintTimeElapse() const
+	{
+		Chronos.PrintTimeElapse();
+	}
+
+	void FGrid::ProcessPointCloud()
+	{
+		EGridSpace DisplaySpace = EGridSpace::Default2D;
+		FTimePoint StartTime = FChrono::Now();
+
+		if(!GetMeshOfLoops())
+		{
+			return;
+		}
+
+#ifdef DEBUG_GRID
+		DisplayLoop(TEXT("FGrid::Loop 2D with thin zone"), GetLoops2D(DisplaySpace), false, false);
+		DisplayLoop(TEXT("FGrid::Loop 2D with thin zone"), GetLoops2D(DisplaySpace), true, false);
+		DisplayLoop(TEXT("FGrid::Loop 3D"), GetLoops3D(), true, false);
+		//Wait();
+		DisplayInnerDomainPoints(TEXT("FGrid::Initial PointCloud 2D"), GetInner2DPoints(DisplaySpace));
+#endif
+		ScaleLoops();
+
+#ifdef DEBUG_GRID
+		DisplayInnerDomainPoints(TEXT("FGrid::PointCloud 2D UniformScaled"), GetInner2DPoints(EGridSpace::UniformScaled));
+		DisplayLoop(TEXT("FGrid::Loop 2D UniformScaled"), GetLoops2D(EGridSpace::UniformScaled), true, false);
+		DisplayInnerDomainPoints(TEXT("FGrid::PointCloud 2D Scaled"), GetInner2DPoints(EGridSpace::Scaled));
+		DisplayLoop(TEXT("FGrid::Loop 2D Scaled"), GetLoops2D(EGridSpace::Scaled), true, false);
+		//Wait(bDisplay);
+#endif
+
+		FindInnerFacePoints();
+
+#ifdef DEBUG_GRID
+		DisplayFindInnerDomainPoints(DisplaySpace);
+		//Wait(bDisplay);
+#endif
+		FindPointsCloseToLoop();
+#ifdef DEBUG_GRID
+		DisplayFindPointsCloseToLoop(DisplaySpace);
+		DisplayFindPointsCloseAndInsideToLoop(DisplaySpace);
+#endif
+		RemovePointsClosedToLoop();
+
+#ifdef DEBUG_GRID
+		DisplayFindPointsCloseAndInsideToLoop(DisplaySpace);
+		DisplayInnerDomainPoints(TEXT("FGrid::Final PointCloud 2D"), GetInner2DPoints(DisplaySpace));
+		//Wait(bDisplay);
+		//DisplayInnerDomainPoints("FGrid::Final PointCloud 3D"), Get3DPoints());
+#endif
+
+		// Removed of Thin zone boundary (the last boundaries). In case of thin zone, the number of 2d boundary will be biggest than 3d boundary one.
+		// Only EGridSpace::Default2D is needed.
+		FaceLoops2D[(int32)EGridSpace::Default2D].SetNum(FaceLoops3D.Num());
+
+#ifdef DEBUG_GRID
+		DisplayLoop(TEXT("FGrid::Final Loop 2D"), GetLoops2D(DisplaySpace), true, false);
+		//DisplayLoop(TEXT("FGrid::Final Loop 3D"), GetLoops3D(), true, false);
+#endif
+
+		Chronos.ProcessPointCloudDuration = FChrono::Elapse(StartTime);
+	}
+
+	void FGrid::DefineCuttingParameters()
+	{
+		FTimePoint StartTime = FChrono::Now();
+
+		FCuttingGrid Neighbors;
+		GetPreferredUVCoordinatesFromNeighbours(Neighbors);
+
+		DefineCuttingParameters(EIso::IsoU, Neighbors);
+		DefineCuttingParameters(EIso::IsoV, Neighbors);
+
+		CuttingSize = CuttingCoordinates.Count();
+
+		Chronos.DefineCuttingParametersDuration = FChrono::Elapse(StartTime);
+	}
+
+	void FGrid::DefineCuttingParameters(EIso Iso, FCuttingGrid& Neighbors)
+	{
+		FTimePoint StartTime = FChrono::Now();
+
+		const FSurfacicBoundary& Boundary = Face->GetBoundary();
+
+#ifdef DEBUG_GETPREFERREDUVCOORDINATESFROMNEIGHBOURS
+		TArray<double> CuttingPointTmp;
+#endif
+		if (Neighbors[Iso].Num())
+		{
+#ifdef DEBUG_GETPREFERREDUVCOORDINATESFROMNEIGHBOURS
+			TArray<FCuttingPoint> Extremities;
+			Extremities.Reserve(2);
+			Extremities.Emplace(Boundary.UVBoundaries[Iso].Min, ECoordinateType::VertexCoordinate, -1, 0.001);
+			Extremities.Emplace(Boundary.UVBoundaries[Iso].Max, ECoordinateType::VertexCoordinate, -1, 0.001);
+			FMesherTools::ComputeFinalCuttingPointsWithImposedCuttingPoints(Face->GetCrossingPointCoordinates(Iso), Face->GetCrossingPointDeltaMaxs(Iso), Extremities, CuttingPointTmp);
+#endif
+			FMesherTools::ComputeFinalCuttingPointsWithPreferredCuttingPoints(Face->GetCrossingPointCoordinates(Iso), Face->GetCrossingPointDeltaMaxs(Iso), Neighbors[Iso], Boundary[Iso], Face->GetCuttingCoordinatesAlongIso(Iso));
+		}
+		else
+		{
+			TArray<FCuttingPoint> Extremities;
+			Extremities.Reserve(2);
+			Extremities.Emplace(Boundary.UVBoundaries[Iso].Min, ECoordinateType::VertexCoordinate, -1, 0.001);
+			Extremities.Emplace(Boundary.UVBoundaries[Iso].Max, ECoordinateType::VertexCoordinate, -1, 0.001);
+			FMesherTools::ComputeFinalCuttingPointsWithImposedCuttingPoints(Face->GetCrossingPointCoordinates(Iso), Face->GetCrossingPointDeltaMaxs(Iso), Extremities, Face->GetCuttingCoordinatesAlongIso(Iso));
+		}
+
+#ifdef DEBUG_GETPREFERREDUVCOORDINATESFROMNEIGHBOURS
+		if (bDisplay)
+		{
+			F3DDebugSession _(TEXT("GetPreferredUVCoordinatesFromNeighbours"));
+			{
+				F3DDebugSession _(FString::Printf(TEXT("%s From Neighbours"), IsoNames[Iso]));
+				for (FCuttingPoint CuttingU : Neighbors[Iso])
+				{
+					DisplayPoint(FPoint(CuttingU.Coordinate, 0.0));
+				}
+			}
+			{
+				F3DDebugSession _(FString::Printf(TEXT("%s From Criteria"), IsoNames[Iso]));
+				for (double CuttingU : CuttingPointTmp)
+				{
+					DisplayPoint(FPoint(CuttingU, Boundary.UVBoundaries[Iso].Length() * 1/80), EVisuProperty::YellowPoint);
+				}
+			}
+			{
+				F3DDebugSession _(FString::Printf(TEXT("%s From Neighbours"), IsoNames[Iso]));
+				for (double CuttingU : Face->GetCuttingCoordinatesAlongIso(Iso))
+				{
+					DisplayPoint(FPoint(CuttingU, Boundary.UVBoundaries[Iso].Length() * 1/40), EVisuProperty::GreenPoint);
+				}
+			}
+			//Wait();
+		}
+#endif
+		CuttingCount[Iso] = CuttingCoordinates.IsoCount(Iso);
+
+		Chronos.DefineCuttingParametersDuration = FChrono::Elapse(StartTime);
+	}
+
+	void FGrid::GetPreferredUVCoordinatesFromNeighbours(FCuttingGrid& NeighboursCutting)
+	{
+#ifdef DEBUG_GETPREFERREDUVCOORDINATESFROMNEIGHBOURS
+		F3DDebugSession _(TEXT("GetPreferredUVCoordinatesFromNeighbours"));
+		{
+			F3DDebugSession _(TEXT("Surface 2D"));
+			Display2D(Face->GetCarrierSurface());
+		}
+#endif
+
+		int32 nbPoints = 0;
+		for (const TSharedPtr<FTopologicalLoop>& Loop : Face->GetLoops())
+		{
+			for (const FOrientedEdge& Edge : Loop->GetEdges())
+			{
+				nbPoints += (int32)Edge.Entity->GetOrCreateMesh(MeshModel)->GetNodeCoordinates().Num() + 1;
+			}
+		}
+
+		NeighboursCutting[EIso::IsoU].Reserve(nbPoints);
+		NeighboursCutting[EIso::IsoV].Reserve(nbPoints);
+
+		for (const TSharedPtr<FTopologicalLoop>& Loop : Face->GetLoops())
+		{
+			for (const FOrientedEdge& OrientedEdge : Loop->GetEdges())
+			{
+				const TSharedPtr<FTopologicalEdge>& Edge = OrientedEdge.Entity;
+
+#ifdef DEBUG_GetPreferredUVCoordinatesFromNeighbours
+				{
+					F3DDebugSession _(FString::Printf(TEXT("Edge %d"), Edge->GetId()));
+					Display2D(Edge);
+				}
+#endif
+
+				TSharedRef<FTopologicalEdge> ActiveEdge = Edge->GetLinkActiveEdge();
+				if (!ActiveEdge->IsMeshed())
+				{
+					continue;
+				}
+
+				TArray<double> ProjectedPointCoords;
+
+				if (ActiveEdge == Edge)
+				{
+					TArray<FCuttingPoint>& CuttingPoints = Edge->GetCuttingPoints();
+					if (CuttingPoints.Num() == 2)
+					{
+						continue;
+					}
+
+					ProjectedPointCoords.Reserve(CuttingPoints.Num());
+					for(const FCuttingPoint& Cutting : CuttingPoints)
+					{
+						ProjectedPointCoords.Emplace(Cutting.Coordinate);
+					}
+				}
+				else
+				{
+					const TSharedRef<FMesh> EdgeMesh = ActiveEdge->GetMesh();
+
+					const TArray<FPoint>& EdgeMeshNodes = EdgeMesh->GetNodeCoordinates();
+					if (EdgeMeshNodes.Num() == 0)
+					{
+						continue;
+					}
+	
+					ProjectedPointCoords.Reserve(EdgeMeshNodes.Num() + 2);
+					bool bSameDirection = Edge->IsSameDirection(ActiveEdge) ? true : false;
+
+					Edge->ProjectTwinEdgePoints(EdgeMeshNodes, bSameDirection, ProjectedPointCoords);
+					ProjectedPointCoords.Insert(Edge->GetStartCurvilinearCoordinates(), 0);
+					ProjectedPointCoords.Add(Edge->GetEndCurvilinearCoordinates());
+				}
+
+				TArray<FPoint2D> EdgePoints2D;
+				Edge->Approximate2DPoints(ProjectedPointCoords, EdgePoints2D);
+
+				const TArray<FCuttingPoint>& CuttingPointTypes = ActiveEdge->GetCuttingPoints();
+				if (ProjectedPointCoords.Num() == CuttingPointTypes.Num())
+				{
+#ifdef DEBUG_GETPREFERREDUVCOORDINATESFROMNEIGHBOURS
+					F3DDebugSession _(TEXT("Nodes"));
+#endif
+					for (int32 Index = 0; Index < EdgePoints2D.Num(); ++Index)
+					{
+						switch (CuttingPointTypes[Index].Type)
+						{
+						case ECoordinateType::VertexCoordinate:
+							NeighboursCutting[EIso::IsoU].Emplace(EdgePoints2D[Index].U, VertexCoordinate);
+							NeighboursCutting[EIso::IsoV].Emplace(EdgePoints2D[Index].V, VertexCoordinate);
+#ifdef DEBUG_GETPREFERREDUVCOORDINATESFROMNEIGHBOURS
+							DisplayPoint(EdgePoints2D[Index], EVisuProperty::GreenPoint);
+#endif
+							break;
+						case ECoordinateType::IsoUCoordinate:
+						case ECoordinateType::IsoVCoordinate:
+						case ECoordinateType::IsoUVCoordinate:
+							NeighboursCutting[EIso::IsoU].Emplace(EdgePoints2D[Index].U, IsoUCoordinate);
+							NeighboursCutting[EIso::IsoV].Emplace(EdgePoints2D[Index].V, IsoVCoordinate);
+#ifdef DEBUG_GETPREFERREDUVCOORDINATESFROMNEIGHBOURS
+							DisplayPoint(EdgePoints2D[Index], EVisuProperty::YellowPoint);
+#endif
+							break;
+						case ECoordinateType::ImposedCoordinate:
+							NeighboursCutting[EIso::IsoU].Emplace(EdgePoints2D[Index].U, OtherCoordinate);
+							NeighboursCutting[EIso::IsoV].Emplace(EdgePoints2D[Index].V, OtherCoordinate);
+#ifdef DEBUG_GETPREFERREDUVCOORDINATESFROMNEIGHBOURS
+							DisplayPoint(EdgePoints2D[Index]);
+#endif
+							break;
+						case ECoordinateType::OtherCoordinate:
+						default:
+							NeighboursCutting[EIso::IsoU].Emplace(EdgePoints2D[Index].U, OtherCoordinate);
+							NeighboursCutting[EIso::IsoV].Emplace(EdgePoints2D[Index].V, OtherCoordinate);
+#ifdef DEBUG_GETPREFERREDUVCOORDINATESFROMNEIGHBOURS
+							DisplayPoint(EdgePoints2D[Index]);
+#endif
+						}
+					}
+				}
+			}
+		}
+
+		TFunction<void(TArray<FCuttingPoint>&)> SortAndRemoveDuplicated = [](TArray<FCuttingPoint>& Neighbours)
+		{
+			if (Neighbours.Num() == 0)
+			{
+				return;
+			}
+
+			Algo::Sort(Neighbours, [](const FCuttingPoint& Point1, const FCuttingPoint& Point2) -> bool
+				{
+					return (Point1.Coordinate) < (Point2.Coordinate);
+				}
+			);
+
+			int32 NewIndex = 0;
+			for (int32 Index = 1; Index < Neighbours.Num(); ++Index)
+			{
+				if (FMath::Abs(Neighbours[Index].Coordinate - Neighbours[NewIndex].Coordinate) < SMALL_NUMBER)
+				{
+					continue;
+				}
+				NewIndex++;
+				Neighbours[NewIndex] = Neighbours[Index];
+			}
+			NewIndex++;
+			Neighbours.SetNum(NewIndex);
+		};
+
+		SortAndRemoveDuplicated(NeighboursCutting[EIso::IsoU]);
+		SortAndRemoveDuplicated(NeighboursCutting[EIso::IsoV]);
+	}
+
+	bool FGrid::GeneratePointCloud()
+	{
+		FTimePoint StartTime = FChrono::Now();
+
+		ComputeMaxDeltaUV();
+		if(MaxDeltaUV[EIso::IsoU] < FaceTolerance[EIso::IsoU] || MaxDeltaUV[EIso::IsoV] < FaceTolerance[EIso::IsoV])
+		{
+			SetAsDegenerated();
+			return false;
+		}
+
+		IsInsideFace.Init(true, CuttingSize);
+		IsCloseToLoop.Init(0, CuttingSize);
+
+		CountOfInnerNodes = CuttingSize;
+		for (int32 Index = 0; Index < (int32)EGridSpace::Max; ++Index)
+		{
+			Points2D[Index].SetNum(CuttingSize);
+		}
+		Points3D.SetNum(CuttingSize);
+		Normals.SetNum(CuttingSize);
+
+		int32 Index = 0;
+		for (int32 IPointV = 0; IPointV < CuttingCount[EIso::IsoV]; ++IPointV)
+		{
+			for (int32 IPointU = 0; IPointU < CuttingCount[EIso::IsoU]; ++IPointU, ++Index)
+			{
+				Points2D[(int32)EGridSpace::Default2D][Index].Set(Face->GetCuttingCoordinatesAlongIso(EIso::IsoU)[IPointU], Face->GetCuttingCoordinatesAlongIso(EIso::IsoV)[IPointV]);
+			}
+		}
+
+		Face->EvaluateGrid(*this);
+
+		ComputeMaxElementSize();
+
+		ScaleGrid();
+
+		Chronos.GeneratePointCloudDuration += FChrono::Elapse(StartTime);
+		return true;
+	}
+
+	void FGrid::ComputeMaxElementSize()
+	{
+		MaxElementSize[EIso::IsoV] = 0;
+		for (int32 IndexU = 0; IndexU < CuttingCount[EIso::IsoU]; ++IndexU)
+		{
+			int32 Index = IndexU;
+			for (int32 IndexV = 1; IndexV < CuttingCount[EIso::IsoV]; ++IndexV, Index += CuttingCount[EIso::IsoU])
+			{
+				MaxElementSize[EIso::IsoV] = FMath::Max(Points3D[Index].SquareDistance(Points3D[Index + CuttingCount[EIso::IsoU]]), MaxElementSize[EIso::IsoV]);
+			}
+		}
+
+		MaxElementSize[EIso::IsoU] = 0;
+		for (int32 IndexV = 0, Index = 0; IndexV < CuttingCount[EIso::IsoV]; ++IndexV)
+		{
+			for (int32 IndexU = 1; IndexU < CuttingCount[EIso::IsoU]; ++IndexU, ++Index)
+			{
+				MaxElementSize[EIso::IsoU] = FMath::Max(Points3D[Index].SquareDistance(Points3D[Index + 1]), MaxElementSize[EIso::IsoU]);
+			}
+			++Index;
+		}
+
+		MaxElementSize[EIso::IsoU] = sqrt(MaxElementSize[EIso::IsoU]);
+		MaxElementSize[EIso::IsoV] = sqrt(MaxElementSize[EIso::IsoV]);
+		MinOfMaxElementSize = FMath::Min(MaxElementSize[EIso::IsoU], MaxElementSize[EIso::IsoV]);
+	}
+
+	void FGrid::FindPointsCloseToLoop()
+	{
+		FTimePoint StartTime = FChrono::Now();
+
+		int32 IndexLoop = 0;
+		int32 IndexU = 1;
+		int32 IndexV = 1;
+		int32 Index = 1;
+
+		TFunction<void()> IncreaseU = [&]()
+		{
+			if (IndexU < CuttingCount[EIso::IsoU] - 1)
+			{
+				IndexU++;
+				Index++;
+			}
+		};
+
+		TFunction<void()> IncreaseV = [&]()
+		{
+			if (IndexV < CuttingCount[EIso::IsoV] - 1)
+			{
+				IndexV++;
+				Index += CuttingCount[EIso::IsoU];
+			}
+		};
+
+		TFunction<void()> DecreaseU = [&]()
+		{
+			if (IndexU > 1)
+			{
+				IndexU--;
+				Index--;
+			}
+		};
+
+		TFunction<void()> DecreaseV = [&]()
+		{
+			if (IndexV > 1)
+			{
+				IndexV--;
+				Index -= CuttingCount[EIso::IsoU];
+			}
+		};
+
+		for (const TArray<FPoint2D>& Loop : FaceLoops2D[(int32)EGridSpace::Default2D])
+		{
+			const FPoint2D* PointA = &Loop.Last();
+
+			IndexU = 1;
+			for (; IndexU < CuttingCount[EIso::IsoU] - 1; ++IndexU)
+			{
+				if (CuttingCoordinates[EIso::IsoU][IndexU] + SMALL_NUMBER_SQUARE > PointA->U)
+				{
+					break;
+				}
+			}
+
+			IndexV = 1;
+			for (; IndexV < CuttingCount[EIso::IsoV] - 1; ++IndexV)
+			{
+				if (CuttingCoordinates[EIso::IsoV][IndexV] + SMALL_NUMBER_SQUARE > PointA->V)
+				{
+					break;
+				}
+			}
+
+			Index = IndexV * CuttingCount[EIso::IsoU] + IndexU;
+
+			IsCloseToLoop[Index] = 1;
+			IsCloseToLoop[Index - 1] = 1;
+			IsCloseToLoop[Index - 1 - CuttingCount[EIso::IsoU]] = 1;
+			IsCloseToLoop[Index - CuttingCount[EIso::IsoU]] = 1;
+
+
+			for (int32 BIndex = 0; BIndex < Loop.Num(); ++BIndex)
+			{
+				const FPoint2D* PointB = &Loop[BIndex];
+
+				if ((CuttingCoordinates[EIso::IsoV][IndexV - 1] - SMALL_NUMBER_SQUARE < PointB->V) && (PointB->V < CuttingCoordinates[EIso::IsoV][IndexV] + SMALL_NUMBER_SQUARE))
+				{
+					if ((CuttingCoordinates[EIso::IsoU][IndexU - 1] < PointB->U) && (PointB->U < CuttingCoordinates[EIso::IsoU][IndexU]))
+					{
+						PointA = PointB;
+						continue;
+					}
+
+					if (PointA->U < PointB->U)
+					{
+						while ((CuttingCoordinates[EIso::IsoU][IndexU] < PointB->U) && (IndexU < CuttingCount[EIso::IsoU] - 1))
+						{
+							IncreaseU();
+							IsCloseToLoop[Index] = 1;
+							IsCloseToLoop[Index - CuttingCount[EIso::IsoU]] = 1;
+						}
+					}
+					else
+					{
+						while ((CuttingCoordinates[EIso::IsoU][IndexU - 1] > PointB->U) && (IndexU > 1))
+						{
+							DecreaseU();
+							IsCloseToLoop[Index - 1] = 1;
+							IsCloseToLoop[Index - CuttingCount[EIso::IsoU] - 1] = 1;
+						}
+					}
+					PointA = PointB;
+					continue;
+				}
+
+				if ((CuttingCoordinates[EIso::IsoU][IndexU - 1] < PointB->U) && (PointB->U < CuttingCoordinates[EIso::IsoU][IndexU]))
+				{
+					if (PointA->V < PointB->V)
+					{
+						while ((CuttingCoordinates[EIso::IsoV][IndexV] < PointB->V) && (IndexV < CuttingCount[EIso::IsoV] - 1))
+						{
+							IncreaseV();
+							IsCloseToLoop[Index] = 1;
+							IsCloseToLoop[Index - 1] = 1;
+						}
+					}
+					else
+					{
+						while ((CuttingCoordinates[EIso::IsoV][IndexV - 1] > PointB->V) && (IndexV > 1))
+						{
+							DecreaseV();
+							IsCloseToLoop[Index - CuttingCount[EIso::IsoU]] = 1;
+							IsCloseToLoop[Index - CuttingCount[EIso::IsoU] - 1] = 1;
+						}
+					}
+					PointA = PointB;
+					continue;
+				}
+
+				double ABv = PointB->V - PointA->V;
+				double ABu = PointB->U - PointA->U;
+
+				if (FMath::Abs(ABu) > FMath::Abs(ABv))
+				{
+					// Py = ABy/ABx*Px + (Ay - ABy/ABx*Ax)
+					// Py = ABy_ABx*Px + (Ay_ABy_ABx_Ax)
+					double ABy_ABx = ABv / ABu;
+					double Ay_ABy_ABx_Ax = PointA->V - ABy_ABx * PointA->U;
+
+					if (ABu > 0)
+					{
+						while (CuttingCoordinates[EIso::IsoU][IndexU] + SMALL_NUMBER_SQUARE < PointB->U)
+						{
+							IncreaseU();
+
+							double CoordinateV_IndexU = 0;
+							if (CuttingCoordinates[EIso::IsoU][IndexU] > PointB->U)
+							{
+								CoordinateV_IndexU = PointB->V;
+							}
+							else
+							{
+								CoordinateV_IndexU = ABy_ABx * CuttingCoordinates[EIso::IsoU][IndexU] + Ay_ABy_ABx_Ax;
+								if (ABv < 0)
+								{
+									if (IndexV > 2)
+									{
+										if (CoordinateV_IndexU + SMALL_NUMBER_SQUARE < CuttingCoordinates[EIso::IsoV][IndexV - 2])
+										{
+											CoordinateV_IndexU = CuttingCoordinates[EIso::IsoV][IndexV - 2] + SMALL_NUMBER_SQUARE;
+											DecreaseU();
+										}
+									}
+								}
+								else
+								{
+									if (IndexV < CuttingCount[EIso::IsoV] - 1)
+									{
+										if (CoordinateV_IndexU + SMALL_NUMBER_SQUARE > CuttingCoordinates[EIso::IsoV][IndexV + 1])
+										{
+											CoordinateV_IndexU = CuttingCoordinates[EIso::IsoV][IndexV + 1] - SMALL_NUMBER_SQUARE;
+											DecreaseU();
+										}
+									}
+								}
+							}
+
+							if (CoordinateV_IndexU - SMALL_NUMBER_SQUARE > CuttingCoordinates[EIso::IsoV][IndexV])
+							{
+								IncreaseV();
+
+								IsCloseToLoop[Index] = 1;
+								IsCloseToLoop[Index - 1] = 1;
+								IsCloseToLoop[Index - CuttingCount[EIso::IsoU]] = 1;
+							}
+							else if (CoordinateV_IndexU + SMALL_NUMBER_SQUARE < CuttingCoordinates[EIso::IsoV][IndexV - 1])
+							{
+								DecreaseV();
+
+								IsCloseToLoop[Index] = 1;
+								IsCloseToLoop[Index - CuttingCount[EIso::IsoU]] = 1;
+								IsCloseToLoop[Index - CuttingCount[EIso::IsoU] - 1] = 1;
+							}
+							else
+							{
+								IsCloseToLoop[Index] = 1;
+								IsCloseToLoop[Index - CuttingCount[EIso::IsoU]] = 1;
+							}
+						}
+					}
+					else
+					{
+						do
+						{
+							DecreaseU();
+
+							double CoordinateV_IndexU = 0; //= ABy_ABx * CoordinateU[IndexU - 1] + Ay_ABy_ABx_Ax;
+							if (CuttingCoordinates[EIso::IsoU][IndexU - 1] < PointB->U)
+							{
+								CoordinateV_IndexU = PointB->V;
+							}
+							else
+							{
+								CoordinateV_IndexU = ABy_ABx * CuttingCoordinates[EIso::IsoU][IndexU - 1] + Ay_ABy_ABx_Ax;
+								if (ABv < 0)
+								{
+									if (IndexV > 2)
+									{
+										if (CoordinateV_IndexU + SMALL_NUMBER_SQUARE < CuttingCoordinates[EIso::IsoV][IndexV - 2])
+										{
+											CoordinateV_IndexU = CuttingCoordinates[EIso::IsoV][IndexV - 2] + SMALL_NUMBER_SQUARE;
+											IncreaseU();
+										}
+									}
+								}
+								else
+								{
+									if (IndexV < CuttingCount[EIso::IsoV] - 1)
+									{
+										if (CoordinateV_IndexU + SMALL_NUMBER_SQUARE > CuttingCoordinates[EIso::IsoV][IndexV + 1])
+										{
+											CoordinateV_IndexU = CuttingCoordinates[EIso::IsoV][IndexV + 1] - SMALL_NUMBER_SQUARE;
+											IncreaseU();
+										}
+									}
+								}
+							}
+
+							if (CoordinateV_IndexU - SMALL_NUMBER_SQUARE > CuttingCoordinates[EIso::IsoV][IndexV])
+							{
+								IsCloseToLoop[Index - 1] = 1;
+
+								IncreaseV();
+								IsCloseToLoop[Index - 1] = 1;
+								IsCloseToLoop[Index] = 1;
+							}
+							else if (CoordinateV_IndexU + SMALL_NUMBER_SQUARE < CuttingCoordinates[EIso::IsoV][IndexV - 1])
+							{
+								DecreaseV();
+
+								IsCloseToLoop[Index - 1] = 1;
+								IsCloseToLoop[Index - CuttingCount[EIso::IsoU]] = 1;
+								IsCloseToLoop[Index - CuttingCount[EIso::IsoU] - 1] = 1;
+							}
+							else
+							{
+								IsCloseToLoop[Index - 1] = 1;
+								IsCloseToLoop[Index - CuttingCount[EIso::IsoU] - 1] = 1;
+							}
+						} while (CuttingCoordinates[EIso::IsoU][IndexU - 1] - SMALL_NUMBER_SQUARE > PointB->U);
+					}
+				}
+				else
+				{
+					// Px = ABx/ABy*Xy + Ax-ABx/ABy*Ay
+					// Px = ABx_ABy*Xy + Ax_ABx_ABy_Ay
+					double ABu_ABv = ABu / ABv;
+					double Au_ABu_ABv_Av = PointA->U - ABu_ABv * PointA->V;
+
+					if (ABv > 0)
+					{
+						while (CuttingCoordinates[EIso::IsoV][IndexV] + SMALL_NUMBER_SQUARE < PointB->V)
+						{
+							IncreaseV();
+							double CoordinateU_IndexV = 0;// = ABx_ABy * CoordinateV[IndexV - 1] + Ax_ABx_ABy_Ay;
+							if (CuttingCoordinates[EIso::IsoV][IndexV] > PointB->V)
+							{
+								CoordinateU_IndexV = PointB->U;
+							}
+							else
+							{
+								CoordinateU_IndexV = ABu_ABv * CuttingCoordinates[EIso::IsoV][IndexV] + Au_ABu_ABv_Av;
+								if (ABu < 0)
+								{
+									if (IndexU > 2)
+									{
+										if (CoordinateU_IndexV + SMALL_NUMBER_SQUARE < CuttingCoordinates[EIso::IsoU][IndexU - 2])
+										{
+											CoordinateU_IndexV = CuttingCoordinates[EIso::IsoU][IndexU - 2] + SMALL_NUMBER_SQUARE;
+											DecreaseV();
+										}
+									}
+								}
+								else
+								{
+									if (IndexU < CuttingCount[EIso::IsoU] - 1)
+									{
+										if (CoordinateU_IndexV + SMALL_NUMBER_SQUARE > CuttingCoordinates[EIso::IsoU][IndexU + 1])
+										{
+											CoordinateU_IndexV = CuttingCoordinates[EIso::IsoU][IndexU + 1] - SMALL_NUMBER_SQUARE;
+											DecreaseV();
+										}
+									}
+								}
+							}
+
+							if ((IndexU < CuttingCount[EIso::IsoU] - 1) && CoordinateU_IndexV - SMALL_NUMBER_SQUARE > CuttingCoordinates[EIso::IsoU][IndexU])
+							{
+								IncreaseU();
+								IsCloseToLoop[Index] = 1;
+								IsCloseToLoop[Index - 1] = 1;
+								IsCloseToLoop[Index - CuttingCount[EIso::IsoU]] = 1;
+							}
+							else if ((IndexU > 1) && CoordinateU_IndexV + SMALL_NUMBER_SQUARE < CuttingCoordinates[EIso::IsoU][IndexU - 1])
+							{
+								DecreaseU();
+
+								IsCloseToLoop[Index] = 1;
+								IsCloseToLoop[Index - 1] = 1;
+								IsCloseToLoop[Index - 1 - CuttingCount[EIso::IsoU]] = 1;
+							}
+							else
+							{
+								IsCloseToLoop[Index] = 1;
+								IsCloseToLoop[Index - 1] = 1;
+							}
+						}
+					}
+					else
+					{
+						do
+						{
+							DecreaseV();
+
+							double CoordinateU_IndexV = 0;// = ABx_ABy * CoordinateV[IndexV - 1] + Ax_ABx_ABy_Ay;
+							if (CuttingCoordinates[EIso::IsoV][IndexV - 1] < PointB->V)
+							{
+								CoordinateU_IndexV = PointB->U;
+							}
+							else
+							{
+								CoordinateU_IndexV = ABu_ABv * CuttingCoordinates[EIso::IsoV][IndexV - 1] + Au_ABu_ABv_Av;
+								if (ABu < 0)
+								{
+									if (IndexU > 2)
+									{
+										if (CoordinateU_IndexV + SMALL_NUMBER_SQUARE < CuttingCoordinates[EIso::IsoU][IndexU - 2])
+										{
+											CoordinateU_IndexV = CuttingCoordinates[EIso::IsoU][IndexU - 2] + SMALL_NUMBER_SQUARE;
+											IncreaseV();
+										}
+									}
+								}
+								else
+								{
+									if (IndexU < CuttingCount[EIso::IsoU] - 1)
+									{
+										if (CoordinateU_IndexV + SMALL_NUMBER_SQUARE > CuttingCoordinates[EIso::IsoU][IndexU + 1])
+										{
+											CoordinateU_IndexV = CuttingCoordinates[EIso::IsoU][IndexU + 1] - SMALL_NUMBER_SQUARE;
+											IncreaseV();
+										}
+									}
+								}
+							}
+
+							if (CoordinateU_IndexV - SMALL_NUMBER_SQUARE > CuttingCoordinates[EIso::IsoU][IndexU])
+							{
+								IncreaseU();
+								IsCloseToLoop[Index] = 1;
+								IsCloseToLoop[Index - CuttingCount[EIso::IsoU] - 1] = 1;
+								IsCloseToLoop[Index - CuttingCount[EIso::IsoU]] = 1;
+							}
+							else if (CoordinateU_IndexV + SMALL_NUMBER_SQUARE < CuttingCoordinates[EIso::IsoU][IndexU - 1])
+							{
+								DecreaseU();
+								IsCloseToLoop[Index - 1] = 1;
+								IsCloseToLoop[Index - CuttingCount[EIso::IsoU] - 1] = 1;
+								IsCloseToLoop[Index - CuttingCount[EIso::IsoU]] = 1;
+							}
+							else
+							{
+								IsCloseToLoop[Index - CuttingCount[EIso::IsoU]] = 1;
+								IsCloseToLoop[Index - CuttingCount[EIso::IsoU] - 1] = 1;
+							}
+						} while (CuttingCoordinates[EIso::IsoV][IndexV - 1] - SMALL_NUMBER_SQUARE > PointB->V);
+					}
+				}
+				PointA = PointB;
+			}
+		}
+		Chronos.FindPointsCloseToLoopDuration += FChrono::Elapse(StartTime);
+	}
+
+	void FGrid::RemovePointsClosedToLoop()
+	{
+		FTimePoint StartTime = FChrono::Now();
+
+		struct FGridSegment
+		{
+			const FPoint2D* StartPoint;
+			const FPoint2D* EndPoint;
+			double StartPointWeight;
+			double EndPointWeight;
+			double UMin;
+			double VMin;
+			double UMax;
+			double VMax;
+
+			FGridSegment(const FPoint2D& SPoint, const FPoint2D& EPoint)
+				: StartPoint(&SPoint)
+				, EndPoint(&EPoint)
+			{
+				StartPointWeight = StartPoint->U + StartPoint->V;
+				EndPointWeight = EndPoint->U + EndPoint->V;
+				if (StartPointWeight > EndPointWeight)
+				{
+					std::swap(StartPointWeight, EndPointWeight);
+					std::swap(StartPoint, EndPoint);
+				}
+				if (StartPoint->U < EndPoint->U)
+				{
+					UMin = StartPoint->U;
+					UMax = EndPoint->U;
+				}
+				else
+				{
+					UMin = EndPoint->U;
+					UMax = StartPoint->U;
+				}
+				if (StartPoint->V < EndPoint->V)
+				{
+					VMin = StartPoint->V;
+					VMax = EndPoint->V;
+				}
+				else
+				{
+					VMin = EndPoint->V;
+					VMax = StartPoint->V;
+				}
+			}
+		};
+
+		TArray<FGridSegment> LoopSegments;
+		{
+			int32 SegmentNum = 0;
+			for (const TArray<FPoint2D>& Loop : FaceLoops2D[(int32)EGridSpace::Default2D])
+			{
+				SegmentNum += (int32)Loop.Num();
+			}
+			LoopSegments.Reserve(SegmentNum);
+
+			for (TArray<FPoint2D>& Loop : FaceLoops2D[(int32)EGridSpace::Default2D])
+			{
+				for (int32 Index = 0; Index < Loop.Num() - 1; ++Index)
+				{
+					LoopSegments.Emplace(Loop[Index], Loop[Index + 1]);
+				}
+			}
+
+			Algo::Sort(LoopSegments, [](const FGridSegment& Seg1, const FGridSegment& Seg2) -> bool
+				{
+					return (Seg1.EndPointWeight) < (Seg2.EndPointWeight);
+				}
+			);
+		}
+
+		// Sort point border grid
+		TArray<double> GridPointWeight;
+		TArray<int32> IndexOfPointsNearAndInsideLoop;
+		TArray<int32> SortedPointIndexes;
+		{
+			IndexOfPointsNearAndInsideLoop.Reserve(CuttingSize);
+			for (int32 Index = 0; Index < CuttingSize; ++Index)
+			{
+				if (IsCloseToLoop[Index] && IsInsideFace[Index])
+				{
+					IndexOfPointsNearAndInsideLoop.Add(Index);
+				}
+			}
+
+			GridPointWeight.Reserve(IndexOfPointsNearAndInsideLoop.Num());
+			SortedPointIndexes.Reserve(IndexOfPointsNearAndInsideLoop.Num());
+			for (const int32& Index : IndexOfPointsNearAndInsideLoop)
+			{
+				GridPointWeight.Add(Points2D[(int32)EGridSpace::Default2D][Index].U + Points2D[(int32)EGridSpace::Default2D][Index].V);
+			}
+			for (int32 Index = 0; Index < IndexOfPointsNearAndInsideLoop.Num(); ++Index)
+			{
+				SortedPointIndexes.Add(Index);
+			}
+			SortedPointIndexes.Sort([&GridPointWeight](const int32& Index1, const int32& Index2) { return GridPointWeight[Index1] < GridPointWeight[Index2]; });
+		}
+
+		double DeltaUVMax = 0;
+		{
+			double MaxDeltaU = 0;
+			for (int32 Index = 0; Index < CuttingCount[EIso::IsoU] - 1; ++Index)
+			{
+				MaxDeltaU = FMath::Max(MaxDeltaU, FMath::Abs(CuttingCoordinates[EIso::IsoU][Index + 1] - CuttingCoordinates[EIso::IsoU][Index]));
+			}
+
+			double MaxDeltaV = 0;
+			for (int32 Index = 0; Index < CuttingCount[EIso::IsoV] - 1; ++Index)
+			{
+				MaxDeltaV = FMath::Max(MaxDeltaV, FMath::Abs(CuttingCoordinates[EIso::IsoV][Index + 1] - CuttingCoordinates[EIso::IsoV][Index]));
+			}
+			DeltaUVMax = FMath::Max(MaxDeltaU, MaxDeltaV);
+		}
+
+		TFunction<void(const int32, double&, double&)> GetDeltaUV = [&](const int32 Index, double& DeltaU, double& DeltaV)
+		{
+			int32 IndexU = Index % CuttingCount[EIso::IsoU];
+			int32 IndexV = Index / CuttingCount[EIso::IsoU];
+
+			if (IndexU == 0)
+			{
+				DeltaU = FMath::Abs(CuttingCoordinates[EIso::IsoU][1] - CuttingCoordinates[EIso::IsoU][0]);
+			}
+			else if (IndexU == CuttingCount[EIso::IsoU] - 1)
+			{
+				DeltaU = FMath::Abs(CuttingCoordinates[EIso::IsoU][CuttingCount[EIso::IsoU] - 1] - CuttingCoordinates[EIso::IsoU][CuttingCount[EIso::IsoU] - 2]);
+			}
+			else
+			{
+				DeltaU = FMath::Abs(CuttingCoordinates[EIso::IsoU][IndexU + 1] - CuttingCoordinates[EIso::IsoU][IndexU - 1]) * .5;
+			}
+
+			if (IndexV == 0)
+			{
+				DeltaV = FMath::Abs(CuttingCoordinates[EIso::IsoV][1] - CuttingCoordinates[EIso::IsoV][0]);
+			}
+			else if (IndexV == CuttingCount[EIso::IsoV] - 1)
+			{
+				DeltaV = FMath::Abs(CuttingCoordinates[EIso::IsoV][CuttingCount[EIso::IsoV] - 1] - CuttingCoordinates[EIso::IsoV][CuttingCount[EIso::IsoV] - 2]);
+			}
+			else
+			{
+				DeltaV = FMath::Abs(CuttingCoordinates[EIso::IsoV][IndexV + 1] - CuttingCoordinates[EIso::IsoV][IndexV - 1]) * .5;
+			}
+		};
+
+		const double DeltaUVMinSquare = FMath::Square(1. / 3.);
+
+		int32 SegmentIndex = 0;
+		for (const int32& SortedIndex : SortedPointIndexes)
+		{
+			int32 Index = IndexOfPointsNearAndInsideLoop[SortedIndex];
+			const FPoint2D& Point2D = Points2D[(int32)EGridSpace::Default2D][Index];
+
+			double DeltaU;
+			double DeltaV;
+			GetDeltaUV(Index, DeltaU, DeltaV);
+
+			for (; SegmentIndex < LoopSegments.Num(); ++SegmentIndex)
+			{
+				if (GridPointWeight[SortedIndex] < LoopSegments[SegmentIndex].EndPointWeight + DeltaUVMax)
+				{
+					break;
+				}
+			}
+
+			for (int32 SegmentIndex2 = SegmentIndex; SegmentIndex2 < LoopSegments.Num(); ++SegmentIndex2)
+			{
+				const FGridSegment& Segment = LoopSegments[SegmentIndex2];
+
+				if (GridPointWeight[SortedIndex] < Segment.StartPointWeight - DeltaUVMax)
+				{
+					continue;
+				}
+				if (Point2D.U + DeltaU < Segment.UMin)
+				{
+					continue;
+				}
+				if (Point2D.U - DeltaU > Segment.UMax)
+				{
+					continue;
+				}
+				if (Point2D.V + DeltaV < Segment.VMin)
+				{
+					continue;
+				}
+				if (Point2D.V - DeltaV > Segment.VMax)
+				{
+					continue;
+				}
+
+				double Coordinate;
+				FPoint2D Projection = ProjectPointOnSegment(Point2D, *Segment.StartPoint, *Segment.EndPoint, Coordinate, /*bRestrictCoodinateToInside*/ true);
+
+				// If Projected point is in the oval center on Point2D => the node is too close
+				double SqrDistance2D = FMath::Square((Point2D.U - Projection.U) / DeltaU);
+				SqrDistance2D += FMath::Square((Point2D.V - Projection.V) / DeltaV);
+				if (SqrDistance2D > DeltaUVMinSquare)
+				{
+					continue;
+				}
+
+				IsCloseToLoop[Index] = 0;
+				IsInsideFace[Index] = 0;
+				CountOfInnerNodes--;
+				break;
+			}
+		}
+
+		Chronos.RemovePointsClosedToLoopDuration += FChrono::Elapse(StartTime);
+	}
+
+	/**
+	 * For the surfacic normal at a StartPoint of the 3D degenerated curve (Not degenerated in 2d) 
+	 * The normal is swap if StartPoint is too close to the Boundary
+	 * The norm of the normal is defined as 1/20 of the parallel boundary Length
+	 */
+	void ScaleAndSwap(FPoint2D& Normal, const FPoint2D& StartPoint, const FSurfacicBoundary& Boundary)
+	{
+		Normal.Normalize();
+		FPoint2D MainDirection = Normal;
+		MainDirection.U /= Boundary[EIso::IsoU].Length();
+		MainDirection.V /= Boundary[EIso::IsoV].Length();
+
+		TFunction<void(const EIso)> SwapAndScale = [&](const EIso Iso)
+		{
+			if (MainDirection[Iso] > 0)
+			{
+				if (FMath::IsNearlyEqual(Boundary[Iso].Max, StartPoint[Iso]))
+				{
+					Normal *= -1.;
+				}
+			}
+			else
+			{
+				if (FMath::IsNearlyEqual(Boundary[IsoU].Min, StartPoint[Iso]))
+				{
+					Normal *= -1.;
+				}
+			}
+			Normal *= Boundary[Iso].Length() / 20;
+		};
+
+		if (MainDirection.U > MainDirection.V)
+		{
+			SwapAndScale(EIso::IsoU);
+		}
+		else
+		{
+			SwapAndScale(EIso::IsoU);
+		}
+	}
+
+	/**
+	 * Displace loop nodes inside to avoid that the nodes are outside the surface Boundary, so outside the grid
+	 */
+	void SlightlyDisplacedPolyline(TArray<FPoint2D>& D2Points, const FSurfacicBoundary& Boundary)
+	{
+		FPoint2D Normal;
+		for (int32 Index = 0; Index < D2Points.Num() - 1; ++Index)
+		{
+			FPoint2D Tangent = D2Points[Index + 1] - D2Points[Index];
+			Normal = Tangent.GetPerpendicularVector();
+
+			ScaleAndSwap(Normal, D2Points[Index], Boundary);
+			D2Points[Index] += Normal;
+		}
+		D2Points.Last() += Normal;
+	}
+
+	bool FGrid::GetMeshOfLoops()
+	{
+		int32 ThinZoneNum = 0;
+		if (Face->HasThinZone())
+		{
+			ThinZoneNum = (int32)ThinZoneFinder.GetThinZones().Num();
+		}
+
+		int32 LoopCount = Face->GetLoops().Num();
+		FaceLoops2D[(int32)EGridSpace::Default2D].Reserve(LoopCount + ThinZoneNum);
+
+		FaceLoops3D.Reserve(LoopCount);
+		NormalsOfFaceLoops.Reserve(LoopCount);
+		NodeIdsOfFaceLoops.Reserve(LoopCount);
+
+#ifdef DEBUG_GET_BOUNDARY_MESH
+		F3DDebugSession _(TEXT("GetLoopMesh"));
+#endif
+
+		for (const TSharedPtr<FTopologicalLoop> Loop : Face->GetLoops())
+		{
+			int32 LoopNodeCount = 0;
+
+			for (const FOrientedEdge& Edge : Loop->GetEdges())
+			{
+				LoopNodeCount += Edge.Entity->GetLinkActiveEdge()->GetCuttingPoints().Num() + 2;
+			}
+
+			TArray<FPoint2D>& Loop2D = FaceLoops2D[(int32)EGridSpace::Default2D].Emplace_GetRef();
+			Loop2D.Empty(LoopNodeCount);
+
+			TArray<FPoint>& Loop3D = FaceLoops3D.Emplace_GetRef();
+			Loop3D.Reserve(LoopNodeCount);
+
+			TArray<FPoint>& LoopNormals = NormalsOfFaceLoops.Emplace_GetRef();
+			LoopNormals.Reserve(LoopNodeCount);
+
+			TArray<int32>& LoopIds = NodeIdsOfFaceLoops.Emplace_GetRef();
+			LoopIds.Reserve(LoopNodeCount);
+
+			for (const FOrientedEdge& OrientedEdge : Loop->GetEdges())
+			{
+				const TSharedPtr<FTopologicalEdge>& Edge = OrientedEdge.Entity;
+				const TSharedRef<FTopologicalEdge>& ActiveEdge = Edge->GetLinkActiveEdge();
+
+				bool bSameDirection = Edge->IsSameDirection(ActiveEdge);
+
+				TArray<double> ActiveEdgeCuttingPointCoordinates;
+				{
+					TArray<FCuttingPoint>& CuttingPoints = ActiveEdge->GetCuttingPoints();
+					GetCuttingPointCoordinates(CuttingPoints, ActiveEdgeCuttingPointCoordinates);
+				}
+
+				FSurfacicPolyline CuttingPolyline(true);
+
+				if (Edge->IsDegenerated())
+				{
+					ensureCADKernel(ActiveEdge == Edge);
+					// else converts CuttingPoint2D form ActiveEdge to Edge (with Linear conversion)
+
+					Swap(CuttingPolyline.Coordinates, ActiveEdgeCuttingPointCoordinates);
+					Edge->Approximate2DPoints(CuttingPolyline.Coordinates, CuttingPolyline.Points2D);
+
+					CuttingPolyline.Points3D.Init(ActiveEdge->GetStartBarycenter(), CuttingPolyline.Coordinates.Num());
+
+					TArray<FPoint2D> D2Points = CuttingPolyline.Points2D;
+					const FSurfacicBoundary& Boundary = Edge->GetCurve()->GetCarrierSurface()->GetBoundary();
+					// to compute the normals, the 2D points are slightly displaced perpendicular to the curve
+					SlightlyDisplacedPolyline(D2Points, Boundary);
+					Edge->GetCurve()->GetCarrierSurface()->EvaluateNormals(D2Points, CuttingPolyline.Normals);
+				}
+				else
+				{
+					if (ActiveEdge != Edge)
+					{
+						ensureCADKernel(ActiveEdgeCuttingPointCoordinates.Num() > 1);
+						TArray<FPoint> CuttingPoint3D;
+						ActiveEdge->ApproximatePoints(ActiveEdgeCuttingPointCoordinates, CuttingPoint3D);
+
+#ifdef DEBUG_GET_BOUNDARY_MESH
+						{
+							F3DDebugSession _(FString::Printf(TEXT("Edge Cutting 3d points %d"), ActiveEdge->GetId()));
+							for (const FPoint& Point : CuttingPoint3D)
+							{
+								DisplayPoint(Point);
+							}
+						}
+#endif
+						TArray<double> ProjectedPointCoords;
+
+						CuttingPolyline.Coordinates.Reserve(CuttingPoint3D.Num());
+						Edge->ProjectTwinEdgePoints(CuttingPoint3D, bSameDirection, CuttingPolyline.Coordinates);
+					}
+					else
+					{
+						Swap(CuttingPolyline.Coordinates, ActiveEdgeCuttingPointCoordinates);
+					}
+
+					Edge->ApproximatePolyline(CuttingPolyline);
+				}
+
+				TArray<int32> EdgeVerticesIndex;
+				if (ActiveEdge->IsDegenerated())
+				{
+					EdgeVerticesIndex.Init(ActiveEdge->GetStartVertex()->GetLinkActiveEntity()->GetOrCreateMesh(MeshModel)->GetMesh(), CuttingPolyline.Coordinates.Num());
+				}
+				else
+				{
+					EdgeVerticesIndex = ((TSharedRef<FEdgeMesh>)ActiveEdge->GetOrCreateMesh(MeshModel))->EdgeVerticesIndex;
+				}
+
+#ifdef DEBUG_GET_BOUNDARY_MESH
+				{
+					F3DDebugSession _(FString::Printf(TEXT("Edge cutting points on surface %d"), ActiveEdge->GetId()));
+					for (const FPoint2D& Point2D : CuttingPolyline.Points2D)
+					{
+						DisplayPoint(Point2D);
+					}
+					//Wait();
+				}
+#endif
+
+				if (OrientedEdge.Direction != EOrientation::Front)
+				{
+					CuttingPolyline.Reverse();
+				}
+
+				if (bSameDirection != (OrientedEdge.Direction == EOrientation::Front))
+				{
+					Algo::Reverse(EdgeVerticesIndex);
+				}
+
+				ensureCADKernel(CuttingPolyline.Size() > 1);
+
+				Loop2D.Append(CuttingPolyline.Points2D);
+				Loop2D.Pop();
+
+				Loop3D.Emplace(ActiveEdge->GetStartVertex((OrientedEdge.Direction == EOrientation::Front) == bSameDirection)->GetLinkActiveEntity()->GetBarycenter());
+				Loop3D.Append(CuttingPolyline.Points3D.GetData() + 1, CuttingPolyline.Points3D.Num() - 2);
+
+				LoopNormals.Append(CuttingPolyline.Normals);
+				LoopNormals.Pop();
+
+				LoopIds.Append(EdgeVerticesIndex);
+				LoopIds.Pop();
+			}
+
+			if (Loop2D.Num() < 3) // degenerated loop
+			{
+				FaceLoops2D[(int32)EGridSpace::Default2D].Pop();
+				FaceLoops3D.Pop();
+				NormalsOfFaceLoops.Pop();
+				continue;
+			}
+		}
+
+		if (CheckIfDegenerated())
+		{
+			return false;
+		}
+
+		if (ThinZoneNum)
+		{
+			for (const FThinZone2D& ThinZone : ThinZoneFinder.GetThinZones())
+			{
+				int32 PointNum = (int32)ThinZone.GetFirstSide().GetSegments().Num();
+				PointNum += (int32)ThinZone.GetSecondSide().GetSegments().Num();
+				TArray<FPoint2D>& LoopPoints = FaceLoops2D[(int32)EGridSpace::Default2D].Emplace_GetRef();
+				LoopPoints.Reserve(PointNum + 4);
+
+				// First point Side1
+				{
+					const FEdgeSegment* Segment = ThinZone.GetFirstSide().GetSegments()[0];
+					LoopPoints.Emplace(Segment->GetEdge()->Approximate2DPoint(Segment->GetCoordinate(ELimit::Start)));
+				}
+
+				for (const FEdgeSegment* Segment : ThinZone.GetFirstSide().GetSegments())
+				{
+					LoopPoints.Emplace(Segment->GetEdge()->Approximate2DPoint(Segment->GetCoordinate(ELimit::End)));
+				}
+
+				// First point Side2
+				{
+					const FEdgeSegment* Segment = ThinZone.GetSecondSide().GetSegments().Last();
+					LoopPoints.Emplace(Segment->GetEdge()->Approximate2DPoint(Segment->GetCoordinate(ELimit::Start)));
+				}
+
+				const TArray<FEdgeSegment*>& Segments = ThinZone.GetSecondSide().GetSegments();
+				for (int32 Index = Segments.Num() - 1; Index >= 0; --Index)
+				{
+					LoopPoints.Emplace(Segments[Index]->GetEdge()->Approximate2DPoint(Segments[Index]->GetCoordinate(ELimit::End)));
+				}
+			}
+		}
+
+		// Fit boundaries to Surface bounds.
+		const FSurfacicBoundary& Bounds = Face->GetBoundary();
+		for (TArray<FPoint2D>& Loop : FaceLoops2D[(int32)EGridSpace::Default2D])
+		{
+			for (FPoint2D& Point : Loop)
+			{
+				Bounds.MoveInsideIfNot(Point);
+			}
+		}
+		return true;
+	}
+
+	void FGrid::ScaleGrid()
+	{
+		FTimePoint StartTime = FChrono::Now();
+
+		TFunction<double(const TArray<double>&)> GetMean = [](const TArray<double>& Lengths)
+		{
+			double MeanLength = 0;
+			for (double Length : Lengths)
+			{
+				MeanLength += Length;
+			}
+			MeanLength /= Lengths.Num();
+			return MeanLength;
+		};
+
+		TFunction<double(const TArray<double>&, const double)> StandardDeviation = [](const TArray<double>& Lengths, const double MeanLength)
+		{
+			double StandardDeviation = 0;
+			for (double Length : Lengths)
+			{
+				StandardDeviation += FMath::Square(Length - MeanLength);
+			}
+			StandardDeviation /= Lengths.Num();
+			StandardDeviation = sqrt(StandardDeviation);
+			return StandardDeviation;
+		};
+
+		TFunction<void(const TArray<double>&, const double, TArray<double>&)> ScaleCoordinates = [](const TArray<double>& InCoordinates, const double ScaleFactor, TArray<double>& OutCoordinatesScaled)
+		{
+			OutCoordinatesScaled.Reserve(InCoordinates.Num());
+
+			for (double Coordinate : InCoordinates)
+			{
+				OutCoordinatesScaled.Add(Coordinate * ScaleFactor);
+			}
+		};
+
+		TFunction<int32(const TArray<double>&, const double)> GetMiddleIndex = [](const TArray<double>& Coordinates, double Middle)
+		{
+			int32 StartIndexUp = 1;
+			for (; StartIndexUp < Coordinates.Num(); ++StartIndexUp)
+			{
+				if (Coordinates[StartIndexUp] > Middle)
+				{
+					break;
+				}
+			}
+			return StartIndexUp;
+		};
+
+		TArray<double> LengthsV;
+		LengthsV.SetNum(CuttingCount[EIso::IsoU]);
+		for (int32 IndexU = 0; IndexU < CuttingCount[EIso::IsoU]; ++IndexU)
+		{
+			int32 Index = IndexU;
+			double Length = 0;
+			for (int32 IndexV = 1; IndexV < CuttingCount[EIso::IsoV]; ++IndexV)
+			{
+				Length += Points3D[Index].Distance(Points3D[Index + CuttingCount[EIso::IsoU]]);
+				Index += CuttingCount[EIso::IsoU];
+			}
+			LengthsV[IndexU] = Length;
+		}
+
+		TArray<double> LengthsU;
+		LengthsU.SetNum(CuttingCount[EIso::IsoV]);
+		for (int32 IndexV = 0, Index = 0; IndexV < CuttingCount[EIso::IsoV]; ++IndexV)
+		{
+			double Length = 0;
+			for (int32 IndexU = 1; IndexU < CuttingCount[EIso::IsoU]; ++IndexU)
+			{
+				Length += Points3D[Index].Distance(Points3D[Index + 1]);
+				Index++;
+			}
+			Index++;
+			LengthsU[IndexV] = Length;
+		}
+
+		double MeanLengthV = GetMean(LengthsV);
+		double FactorV = MeanLengthV / (CuttingCoordinates[EIso::IsoV].Last() - CuttingCoordinates[EIso::IsoV][0]);
+
+		double MeanLengthU = GetMean(LengthsU);
+		double FactorU = MeanLengthU / (CuttingCoordinates[EIso::IsoU].Last() - CuttingCoordinates[EIso::IsoU][0]);
+
+		TArray<double> ScaledCoordinatesU;
+		ScaleCoordinates(CuttingCoordinates[EIso::IsoU], FactorU, ScaledCoordinatesU);
+
+		TArray<double> ScaledCoordinatesV;
+		ScaleCoordinates(CuttingCoordinates[EIso::IsoV], FactorV, ScaledCoordinatesV);
+
+		{
+			int32 NumUV = 0;
+			for (int32 IPointV = 0; IPointV < CuttingCount[EIso::IsoV]; ++IPointV)
+			{
+				for (int32 IPointU = 0; IPointU < CuttingCount[EIso::IsoU]; ++IPointU, ++NumUV)
+				{
+					Points2D[(int32)EGridSpace::UniformScaled][NumUV].Set(ScaledCoordinatesU[IPointU], ScaledCoordinatesV[IPointV]);
+				}
+			}
+		}
+
+		double StandardDeviationU = StandardDeviation(LengthsU, MeanLengthU);
+		double StandardDeviationV = StandardDeviation(LengthsV, MeanLengthV);
+
+		if (StandardDeviationV > StandardDeviationU)
+		{
+			double MiddleV = (CuttingCoordinates[EIso::IsoV].Last() + CuttingCoordinates[EIso::IsoV][0]) * 0.5;
+
+			FCoordinateGrid Grid;
+			Grid[EIso::IsoU] = CuttingCoordinates[EIso::IsoU];
+			Grid[EIso::IsoV].Add(MiddleV);
+
+			FSurfacicSampling MiddlePoints;
+			Face->EvaluatePointGrid(Grid, MiddlePoints);
+
+			int32 StartIndexUp = GetMiddleIndex(CuttingCoordinates[EIso::IsoV], MiddleV);
+			int32 StartIndexDown = StartIndexUp - 1;
+
+			int32 NumUV = 0;
+			for (int32 IPointU = 0; IPointU < CuttingCount[EIso::IsoU]; ++IPointU)
+			{
+				double Length = 0;
+				FPoint LastPoint = MiddlePoints.Points3D[IPointU];
+				for (int32 IPointV = StartIndexUp; IPointV < CuttingCount[EIso::IsoV]; ++IPointV)
+				{
+					NumUV = IPointV * CuttingCount[EIso::IsoU] + IPointU;
+					Length += LastPoint.Distance(Points3D[NumUV]);
+					Points2D[(int32)EGridSpace::Scaled][NumUV].Set(Points2D[(int32)EGridSpace::UniformScaled][NumUV].U, Length);
+					LastPoint = Points3D[NumUV];
+				}
+
+				Length = 0;
+				LastPoint = MiddlePoints.Points3D[IPointU];
+				for (int32 IPointV = StartIndexDown; IPointV >= 0; --IPointV)
+				{
+					NumUV = IPointV * CuttingCount[EIso::IsoU] + IPointU;
+					Length -= LastPoint.Distance(Points3D[NumUV]);
+					Points2D[(int32)EGridSpace::Scaled][NumUV].Set(Points2D[(int32)EGridSpace::UniformScaled][NumUV].U, Length);
+					LastPoint = Points3D[NumUV];
+				}
+			}
+		}
+		else
+		{
+			double MiddleU = (CuttingCoordinates[EIso::IsoU].Last() + CuttingCoordinates[EIso::IsoU][0]) * 0.5;
+
+			FCoordinateGrid Grid;
+			Grid[EIso::IsoU].Add(MiddleU);
+			Grid[EIso::IsoV] = CuttingCoordinates[EIso::IsoV];
+
+			FSurfacicSampling MiddlePoints;
+			Face->EvaluatePointGrid(Grid, MiddlePoints);
+
+			int32 StartIndexUp = GetMiddleIndex(CuttingCoordinates[EIso::IsoU], MiddleU);
+			int32 StartIndexDown = StartIndexUp - 1;
+
+			int32 NumUV = 0;
+			for (int32 IPointV = 0; IPointV < CuttingCount[EIso::IsoV]; ++IPointV)
+			{
+				double Length = 0;
+				FPoint LastPoint = MiddlePoints.Points3D[IPointV];
+				for (int32 IPointU = StartIndexUp; IPointU < CuttingCount[EIso::IsoU]; ++IPointU)
+				{
+					NumUV = IPointV * CuttingCount[EIso::IsoU] + IPointU;
+					Length += LastPoint.Distance(Points3D[NumUV]);
+					Points2D[(int32)EGridSpace::Scaled][NumUV].Set(Length, Points2D[(int32)EGridSpace::UniformScaled][NumUV].V);
+					LastPoint = Points3D[NumUV];
+				}
+
+				Length = 0;
+				LastPoint = MiddlePoints.Points3D[IPointV];
+				for (int32 IPointU = StartIndexDown; IPointU >= 0; --IPointU)
+				{
+					NumUV = IPointV * CuttingCount[EIso::IsoU] + IPointU;
+					Length -= LastPoint.Distance(Points3D[NumUV]);
+					Points2D[(int32)EGridSpace::Scaled][NumUV].Set(Length, Points2D[(int32)EGridSpace::UniformScaled][NumUV].V);
+					LastPoint = Points3D[NumUV];
+				}
+			}
+		}
+		Chronos.ScaleGridDuration = FChrono::Elapse(StartTime);
+	}
+
+	void FGrid::TransformPoints(EGridSpace DestinationSpace, const TArray<FPoint2D>& InPointsToScale, TArray<FPoint2D>& OutTransformedPoints) const
+	{
+		OutTransformedPoints.SetNum(InPointsToScale.Num());
+
+		int32 IndexU = 0;
+		int32 IndexV = 0;
+		for (int32 Index = 0; Index < InPointsToScale.Num(); ++Index)
+		{
+			const FPoint2D& Point = InPointsToScale[Index];
+
+			FindCoordinateIndex(CuttingCoordinates[EIso::IsoU], Point.U, IndexU);
+			FindCoordinateIndex(CuttingCoordinates[EIso::IsoV], Point.V, IndexV);
+
+			ComputeNewCoordinate(Points2D[(int32)DestinationSpace], IndexU, IndexV, Point, OutTransformedPoints[Index]);
+		}
+	}
+
+	void FGrid::SearchThinZones()
+	{
+		double Size = GetMinElementSize();
+		ThinZoneFinder.Set(Size / 3.);
+		ThinZoneFinder.SearchThinZones();
+	}
+
+	void FGrid::ScaleLoops()
+	{
+		FaceLoops2D[(int32)EGridSpace::Scaled].SetNum(FaceLoops2D[(int32)EGridSpace::Default2D].Num());
+		FaceLoops2D[(int32)EGridSpace::UniformScaled].SetNum(FaceLoops2D[(int32)EGridSpace::Default2D].Num());
+
+		for (int32 IndexBoudnary = 0; IndexBoudnary < FaceLoops2D[(int32)EGridSpace::Default2D].Num(); ++IndexBoudnary)
+		{
+			const TArray<FPoint2D>& Loop = FaceLoops2D[(int32)EGridSpace::Default2D][IndexBoudnary];
+			TArray<FPoint2D>& ScaledLoop = FaceLoops2D[(int32)EGridSpace::Scaled][IndexBoudnary];
+			TArray<FPoint2D>& UniformScaledLoop = FaceLoops2D[(int32)EGridSpace::UniformScaled][IndexBoudnary];
+
+			ScaledLoop.SetNum(Loop.Num());
+			UniformScaledLoop.SetNum(Loop.Num());
+
+			int32 IndexU = 0;
+			int32 IndexV = 0;
+			for (int32 Index = 0; Index < Loop.Num(); ++Index)
+			{
+				const FPoint2D& Point = Loop[Index];
+
+				FindCoordinateIndex(CuttingCoordinates[EIso::IsoU], Point.U, IndexU);
+				FindCoordinateIndex(CuttingCoordinates[EIso::IsoV], Point.V, IndexV);
+
+				ComputeNewCoordinate(Points2D[EGridSpace::Scaled], IndexU, IndexV, Point, ScaledLoop[Index]);
+				ComputeNewCoordinate(Points2D[EGridSpace::UniformScaled], IndexU, IndexV, Point, UniformScaledLoop[Index]);
+			}
+
+		}
+	}
+
+	void FGrid::ComputeMaxDeltaUV()
+	{
+		MaxDeltaUV[EIso::IsoU] = 0;
+		for (int32 Index = 1; Index < CuttingCoordinates[EIso::IsoU].Num(); ++Index)
+		{
+			double Delta = CuttingCoordinates[EIso::IsoU][Index] - CuttingCoordinates[EIso::IsoU][Index - 1];
+			MaxDeltaUV[EIso::IsoU] = FMath::Max(MaxDeltaUV[EIso::IsoU], Delta);
+		}
+
+		MaxDeltaUV[EIso::IsoV] = 0;
+		for (int32 Index = 1; Index < CuttingCoordinates[EIso::IsoV].Num(); ++Index)
+		{
+			double Delta = CuttingCoordinates[EIso::IsoV][Index] - CuttingCoordinates[EIso::IsoV][Index - 1];
+			MaxDeltaUV[EIso::IsoV] = FMath::Max(MaxDeltaUV[EIso::IsoV], Delta);
+		}
+	}
+
+	void FGrid::FindInnerFacePoints()
+	{
+		// FindInnerDomainPoints: Inner Points <-> bIsOfInnerDomain = true
+		// For each points count the number of intersection with the boundary in the four directions U+ U- V+ V-
+		// It for each the number is pair, the point is outside,
+		// If in 3 directions the point is inner, the point is inner else we have a doubt so it preferable to consider it outside. 
+		// Most of the time, there is a doubt if the point is to close of the boundary. So it will be removed be other criteria
+
+		FTimePoint StartTime = FChrono::Now();
+
+		TArray<char> NbIntersectUForward; // we need to know if intersect is pair 0, 2, 4... intersection of impair 1, 3, 5... false is pair, true is impaire 
+		TArray<char> NbIntersectUBackward;
+		TArray<char> NbIntersectVForward;
+		TArray<char> NbIntersectVBackward;
+		TArray<char> IntersectLoop;
+
+		IntersectLoop.Init(0, CuttingSize);
+		IsInsideFace.Init(1, CuttingSize);
+
+		NbIntersectUForward.Init(0, CuttingSize);
+		NbIntersectUBackward.Init(0, CuttingSize);
+		NbIntersectVForward.Init(0, CuttingSize);
+		NbIntersectVBackward.Init(0, CuttingSize);
+
+		// Loop node too close to one of CoordinateU or CoordinateV are moved a little to avoid floating error of comparison 
+		// This step is necessary instead of all points could be considered outside...
+		const double SmallToleranceU = SMALL_NUMBER;
+		const double SmallToleranceV = SMALL_NUMBER;
+
+		{
+			int32 IndexV = 0;
+			int32 IndexU = 0;
+			for (TArray<FPoint2D>& Loop : FaceLoops2D[(int32)EGridSpace::Default2D])
+			{
+				for (FPoint2D& Point : Loop)
+				{
+					while (IndexV != 0 && (Point.V < CuttingCoordinates[EIso::IsoV][IndexV]))
+					{
+						IndexV--;
+					}
+					for (; IndexV < CuttingCount[EIso::IsoV]; ++IndexV)
+					{
+						if (Point.V + SmallToleranceV < CuttingCoordinates[EIso::IsoV][IndexV])
+						{
+							break;
+						}
+						if (Point.V - SmallToleranceV > CuttingCoordinates[EIso::IsoV][IndexV])
+						{
+							continue;
+						}
+
+						if (IndexV == 0)
+						{
+							Point.V += SmallToleranceV;
+						}
+						else
+						{
+							Point.V -= SmallToleranceV;
+						}
+
+						break;
+					}
+					if (IndexV == CuttingCount[EIso::IsoV])
+					{
+						IndexV--;
+					}
+
+					while (IndexU != 0 && (Point.U < CuttingCoordinates[EIso::IsoU][IndexU]))
+					{
+						IndexU--;
+					}
+					for (; IndexU < CuttingCount[EIso::IsoU]; ++IndexU)
+					{
+						if (Point.U + SmallToleranceU < CuttingCoordinates[EIso::IsoU][IndexU])
+						{
+							break;
+						}
+						if (Point.U - SmallToleranceU > CuttingCoordinates[EIso::IsoU][IndexU])
+						{
+							continue;
+						}
+
+						if (IndexU == 0)
+						{
+							Point.U += SmallToleranceU;
+						}
+						else
+						{
+							Point.U -= SmallToleranceU;
+						}
+						break;
+					}
+					if (IndexU == CuttingCount[EIso::IsoU])
+					{
+						IndexU--;
+					}
+				}
+			}
+		}
+
+		DisplayLoop(TEXT("FGrid::Loop 2D After move according tol"), GetLoops2D(EGridSpace::Default2D), true, false);
+
+		// Intersection along U axis
+		for (const TArray<FPoint2D>& Loop : FaceLoops2D[(int32)EGridSpace::Default2D])
+		{
+			const FPoint2D* FirstSegmentPoint = &Loop.Last();
+			for (const FPoint2D& LoopPoint : Loop)
+			{
+				const FPoint2D* SecondSegmentPoint = &LoopPoint;
+				double UMin = FirstSegmentPoint->U;
+				double VMin = FirstSegmentPoint->V;
+				double Umax = SecondSegmentPoint->U;
+				double Vmax = SecondSegmentPoint->V;
+				Sort(UMin, Umax);
+				Sort(VMin, Vmax);
+
+				// AB^AP = ABu*APv - ABv*APu
+				// AB^AP = ABu*(Pv-Av) - ABv*(Pu-Au)
+				// AB^AP = Pv*ABu - Pu*ABv + Au*ABv - Av*ABu
+				// AB^AP = Pv*ABu - Pu*ABv + AuABVMinusAvABu
+				FPoint2D PointA;
+				FPoint2D PointB;
+				if (FirstSegmentPoint->V < SecondSegmentPoint->V)
+				{
+					PointA = *FirstSegmentPoint;
+					PointB = *SecondSegmentPoint;
+				}
+				else
+				{
+					PointA = *SecondSegmentPoint;
+					PointB = *FirstSegmentPoint;
+				}
+				double ABv = PointB.V - PointA.V;
+				double ABu = PointB.U - PointA.U;
+				double AuABVMinusAvABu = PointA.U * ABv - PointA.V * ABu;
+
+				int32 IndexV = 0;
+				int32 Index = 0;
+				for (; IndexV < CuttingCount[EIso::IsoV]; ++IndexV)
+				{
+					if (CuttingCoordinates[EIso::IsoV][IndexV] >= VMin)
+					{
+						break;
+					}
+					Index += CuttingCount[EIso::IsoU];
+				}
+
+				for (; IndexV < CuttingCount[EIso::IsoV]; ++IndexV)
+				{
+					if (CuttingCoordinates[EIso::IsoV][IndexV] > Vmax)
+					{
+						break;
+					}
+
+					for (int32 IndexU = 0; IndexU < CuttingCount[EIso::IsoU]; ++IndexU, ++Index)
+					{
+						//Index = IndexV * NumU + IndexU;
+						if (IntersectLoop[Index])
+						{
+							continue;
+						}
+
+						if (CuttingCoordinates[EIso::IsoU][IndexU] < UMin)
+						{
+							NbIntersectVForward[Index] = NbIntersectVForward[Index] > 0 ? 0 : 1;
+						}
+						else if (CuttingCoordinates[EIso::IsoU][IndexU] > Umax)
+						{
+							NbIntersectVBackward[Index] = NbIntersectVBackward[Index] > 0 ? 0 : 1;
+						}
+						else
+						{
+							double APvectAB = CuttingCoordinates[EIso::IsoV][IndexV] * ABu - CuttingCoordinates[EIso::IsoU][IndexU] * ABv + AuABVMinusAvABu;
+							if (APvectAB > SMALL_NUMBER)
+							{
+								NbIntersectVForward[Index] = NbIntersectVForward[Index] > 0 ? 0 : 1;
+							}
+							else if (APvectAB < SMALL_NUMBER)
+							{
+								NbIntersectVBackward[Index] = NbIntersectVBackward[Index] > 0 ? 0 : 1;
+							}
+							else
+							{
+								IntersectLoop[Index] = 1;
+							}
+						}
+					}
+				}
+				FirstSegmentPoint = SecondSegmentPoint;
+			}
+		}
+
+		// Intersection along V axis
+		for (const TArray<FPoint2D>& Loop : FaceLoops2D[(int32)EGridSpace::Default2D])
+		{
+			const FPoint2D* FirstSegmentPoint = &Loop.Last();
+			for (const FPoint2D& LoopPoint : Loop)
+			{
+				const FPoint2D* SecondSegmentPoint = &LoopPoint;
+				double UMin = FirstSegmentPoint->U;
+				double VMin = FirstSegmentPoint->V;
+				double Umax = SecondSegmentPoint->U;
+				double Vmax = SecondSegmentPoint->V;
+				Sort(UMin, Umax);
+				Sort(VMin, Vmax);
+
+				// AB^AP = ABu*APv - ABv*APu
+				// AB^AP = ABu*(Pv-Av) - ABv*(Pu-Au)
+				// AB^AP = Pv*ABu - Pu*ABv + Au*ABv - Av*ABu
+				// AB^AP = Pv*ABu - Pu*ABv + AuABVMinusAvABu
+				FPoint2D PointA;
+				FPoint2D PointB;
+				if (FirstSegmentPoint->U < SecondSegmentPoint->U)
+				{
+					PointA = *FirstSegmentPoint;
+					PointB = *SecondSegmentPoint;
+				}
+				else
+				{
+					PointA = *SecondSegmentPoint;
+					PointB = *FirstSegmentPoint;
+				}
+
+				double ABu = PointB.U - PointA.U;
+				double ABv = PointB.V - PointA.V;
+				double AuABVMinusAvABu = PointA.U * ABv - PointA.V * ABu;
+				int32 Index = 0;
+				for (int32 IndexU = 0; IndexU < CuttingCount[EIso::IsoU]; ++IndexU)
+				{
+					if (CuttingCoordinates[EIso::IsoU][IndexU] < UMin)
+					{
+						continue;
+					}
+
+					if (CuttingCoordinates[EIso::IsoU][IndexU] >= Umax)
+					{
+						continue;
+					}
+
+					Index = IndexU;
+					for (int32 IndexV = 0; IndexV < CuttingCount[EIso::IsoV]; ++IndexV, Index += CuttingCount[EIso::IsoU])
+					{
+						if (IntersectLoop[Index])
+						{
+							continue;
+						}
+
+						if (CuttingCoordinates[EIso::IsoV][IndexV] < VMin)
+						{
+							NbIntersectUForward[Index] = NbIntersectUForward[Index] > 0 ? 0 : 1;
+						}
+						else if (CuttingCoordinates[EIso::IsoV][IndexV] > Vmax)
+						{
+							NbIntersectUBackward[Index] = NbIntersectUBackward[Index] > 0 ? 0 : 1;
+						}
+						else
+						{
+							double APvectAB = CuttingCoordinates[EIso::IsoV][IndexV] * ABu - CuttingCoordinates[EIso::IsoU][IndexU] * ABv + AuABVMinusAvABu;
+							if (APvectAB > SMALL_NUMBER)
+							{
+								NbIntersectUBackward[Index] = NbIntersectUBackward[Index] > 0 ? 0 : 1;
+							}
+							else if (APvectAB < SMALL_NUMBER)
+							{
+								NbIntersectUForward[Index] = NbIntersectUForward[Index] > 0 ? 0 : 1;
+							}
+							else
+							{
+								IntersectLoop[Index] = 1;
+							}
+						}
+					}
+				}
+				FirstSegmentPoint = SecondSegmentPoint;
+			}
+		}
+
+		for (int32 Index = 0; Index < CuttingSize; ++Index)
+		{
+			if (IntersectLoop[Index])
+			{
+				IsInsideFace[Index] = 0;
+				CountOfInnerNodes--;
+				continue;
+			}
+
+			int32 IsInside = 0;
+			if (NbIntersectVForward[Index] > 0)
+			{
+				IsInside++;
+			}
+			if (NbIntersectVBackward[Index]  > 0)
+			{
+				IsInside++;
+			}
+			if (NbIntersectUForward[Index] > 0)
+			{
+				IsInside++;
+			}
+			if (NbIntersectUBackward[Index] > 0)
+			{
+				IsInside++;
+			}
+			if (IsInside < 3)
+			{
+				IsInsideFace[Index] = false;
+				CountOfInnerNodes--;
+			}
+		}
+
+		Chronos.FindInnerDomainPointsDuration += FChrono::Elapse(StartTime);
+	}
+
+	void FGrid::DisplayFindInnerDomainPoints(EGridSpace DisplaySpace) const
+	{
+		if (!bDisplay)
+		{
+			return;
+		}
+		int32 NbNum = 0;
+		Open3DDebugSession(TEXT("FGrid::FindInnerDomainPoints Inside Point"));
+		for (int32 Index = 0; Index < CuttingSize; ++Index)
+		{
+			if (IsInsideFace[Index])
+			{
+				DisplayPoint(Points2D[(int32)EGridSpace::Default2D][Index], Index);
+				NbNum++;
+			}
+		}
+		Close3DDebugSession();
+		ensureCADKernel(NbNum == CountOfInnerNodes);
+
+		Open3DDebugSession(TEXT("FGrid::FindInnerDomainPoints Outside Point"));
+		for (int32 Index = 0; Index < CuttingSize; ++Index)
+		{
+			if (!IsInsideFace[Index])
+			{
+				DisplayPoint(Points2D[(int32)EGridSpace::Default2D][Index], EVisuProperty::GreenPoint, Index);
+			}
+		}
+		Close3DDebugSession();
+	}
+
+	void FGrid::DisplayFindPointsCloseToLoop(EGridSpace DisplaySpace) const
+	{
+		if (!bDisplay)
+		{
+			return;
+		}
+
+		Open3DDebugSession(TEXT("FGrid::FindPointsClosedToLoop result"));
+		for (int32 Index = 0; Index < CuttingSize; ++Index)
+		{
+			if (IsCloseToLoop[Index])
+			{
+				DisplayPoint(Points2D[DisplaySpace][Index]);
+			}
+			else
+			{
+				DisplayPoint(Points2D[DisplaySpace][Index], EVisuProperty::YellowPoint);
+			}
+		}
+		Close3DDebugSession();
+	}
+
+	void FGrid::DisplayFindPointsCloseAndInsideToLoop(EGridSpace DisplaySpace) const
+	{
+		if (!bDisplay)
+		{
+			return;
+		}
+
+		Open3DDebugSession(TEXT("FGrid::FindPointsCloseAndInsideToLoop result"));
+		for (int32 Index = 0; Index < CuttingSize; ++Index)
+		{
+			if (IsInsideFace[Index])
+			{
+				if (IsCloseToLoop[Index])
+				{
+					DisplayPoint(Points2D[DisplaySpace][Index], EVisuProperty::BluePoint, Index);
+				}
+				else
+				{
+					DisplayPoint(Points2D[DisplaySpace][Index], EVisuProperty::YellowPoint, Index);
+				}
+			}
+		}
+		Close3DDebugSession();
+	}
+
+	bool FGrid::CheckIfDegenerated()
+	{
+		if (FaceLoops2D[(int32)EGridSpace::Default2D].Num() == 0)
+		{
+			SetAsDegenerated();
+			return true;
+		}
+
+		// if the external boundary is composed by only 2 points, the mesh of the surface is only an edge.
+		// The grid is degenerated.
+		if (FaceLoops2D[(int32)EGridSpace::Default2D][0].Num() < 3)
+		{
+			SetAsDegenerated();
+			return true;
+		}
+
+		return false;
+	}
+
+}
+
+
diff --git a/Engine/Source/Runtime/Datasmith/CADKernel/Private/CADKernel/Mesh/Structure/Grid.h b/Engine/Source/Runtime/Datasmith/CADKernel/Private/CADKernel/Mesh/Structure/Grid.h
new file mode 100644
--- /dev/null
+++ b/Engine/Source/Runtime/Datasmith/CADKernel/Private/CADKernel/Mesh/Structure/Grid.h
@@ -0,0 +1,617 @@
+// Copyright Epic Games, Inc. All Rights Reserved.
+#pragma once
+
+#include "CADKernel/Core/Chrono.h"
+#include "CADKernel/Geo/GeoEnum.h"
+#include "CADKernel/Geo/GeoPoint.h"
+#include "CADKernel/Math/Point.h"
+#include "CADKernel/Mesh/MeshEnum.h"
+#include "CADKernel/Mesh/Meshers/IsoTriangulator/DefineForDebug.h"
+#include "CADKernel/Mesh/Structure/ThinZone2D.h"
+#include "CADKernel/UI/Display.h"
+
+
+namespace CADKernel
+{
+	struct FCuttingPoint;
+	struct FCuttingGrid;
+	
+	class FModelMesh;
+	class FPoint;
+	class FPoint2D;
+	class FThinZone2DFinder;
+	class FTopologicalFace;
+
+	struct FGridChronos
+	{
+		FDuration DefineCuttingParametersDuration;
+		FDuration GeneratePointCloudDuration;
+		FDuration ProcessPointCloudDuration;
+		FDuration FindInnerDomainPointsDuration;
+		FDuration Build2DLoopDuration;
+		FDuration RemovePointsClosedToLoopDuration;
+		FDuration FindPointsCloseToLoopDuration;
+		FDuration ScaleGridDuration;
+
+		FGridChronos()
+			: DefineCuttingParametersDuration(FChrono::Init())
+			, GeneratePointCloudDuration(FChrono::Init())
+			, ProcessPointCloudDuration(FChrono::Init())
+			, FindInnerDomainPointsDuration(FChrono::Init())
+			, Build2DLoopDuration(FChrono::Init())
+			, RemovePointsClosedToLoopDuration(FChrono::Init())
+			, FindPointsCloseToLoopDuration(FChrono::Init())
+			, ScaleGridDuration(FChrono::Init())
+		{}
+
+		void PrintTimeElapse() const
+		{
+			FDuration GridDuration = FChrono::Elapse(FChrono::Now());
+			GridDuration += DefineCuttingParametersDuration;
+			GridDuration += GeneratePointCloudDuration;
+			GridDuration += FindInnerDomainPointsDuration;
+			GridDuration += Build2DLoopDuration;
+			GridDuration += RemovePointsClosedToLoopDuration;
+			GridDuration += FindPointsCloseToLoopDuration;
+			FChrono::PrintClockElapse(Log, TEXT(""), TEXT("Grid"), GridDuration);
+			FChrono::PrintClockElapse(Log, TEXT("  "), TEXT("DefineCuttingParameters"), DefineCuttingParametersDuration);
+			FChrono::PrintClockElapse(Log, TEXT("  "), TEXT("GeneratePointCloud"), GeneratePointCloudDuration);
+			FChrono::PrintClockElapse(Log, TEXT("  "), TEXT("GenerateDomainPoints"), ProcessPointCloudDuration);
+			FChrono::PrintClockElapse(Log, TEXT("    "), TEXT("FindInnerDomainPointsDuration"), FindInnerDomainPointsDuration);
+			FChrono::PrintClockElapse(Log, TEXT("    "), TEXT("FindPointsCloseToLoop"), FindPointsCloseToLoopDuration);
+			FChrono::PrintClockElapse(Log, TEXT("    "), TEXT("RemovePointsClosedToLoop"), RemovePointsClosedToLoopDuration);
+			FChrono::PrintClockElapse(Log, TEXT("    "), TEXT("Build2DLoopDuration"), Build2DLoopDuration);
+		}
+	};
+
+	class FGrid : public FHaveStates
+	{
+	protected:
+
+		const TSharedRef<FTopologicalFace> Face;
+		const FSurfacicTolerance& FaceTolerance;
+
+		TSharedRef<FModelMesh> MeshModel;
+
+		FThinZone2DFinder ThinZoneFinder;
+
+		/**
+		 * 2D Coordinate of Loop's nodes in each space
+		 */
+		TArray<TArray<FPoint2D>> FaceLoops2D[EGridSpace::Max];
+
+		/**
+		 * 3D Coordinate of Loop nodes in each space
+		 */
+		TArray<TArray<FPoint>> FaceLoops3D;
+
+		/**
+		 * Surface Normal at each boundary nodes
+		 */
+		TArray<TArray<FPoint>> NormalsOfFaceLoops;
+
+		TArray<TArray<int32>> NodeIdsOfFaceLoops;
+
+		/**
+		 * grid point cloud size
+		 */ 
+		int32 CuttingCount[2] = {0 , 0};
+		int32 CuttingSize = 0;
+
+		/**
+		 * count of node inside the face i.e. node inside the external loop and outer the inner loops 
+		 */
+		int32 CountOfInnerNodes = 0;
+
+		/*
+		 * Cutting coordinates of the face respecting the meshing criteria
+		 */
+		const FCoordinateGrid& CuttingCoordinates;
+
+		/*
+		 * Maximum difference of coordinate along the specified axis of two successive cutting points
+		 */
+		FPoint2D MaxDeltaUV = { 0., 0. };
+
+		/**
+		 * Maximum 3d distance along each axis of two successive cutting points
+		 */
+		FPoint2D MaxElementSize = { 0., 0. };
+
+		double MinOfMaxElementSize = 0;
+
+		/**
+		 * Array to flag each points as inside the face (if IsInsideDomain[Index] = 1) or outer 
+		 * @see IsNodeInsideFace
+		 */
+		TArray<char> IsInsideFace;
+
+		/**
+		 * Array to flag each points as close to a loop of the face if IsCloseToLoop[Index] = 1
+		 * @see IsNodeCloseToLoop
+		 */
+		TArray<char> IsCloseToLoop;
+
+		/**
+		 * 2D Coordinate of grid nodes in each space
+		 */
+		TArray<FPoint2D> Points2D[EGridSpace::Max];
+
+		/**
+		 * 3D Coordinate of inner nodes
+		 */
+		TArray<FPoint> Points3D;
+
+		/**
+		 * Surface Normal at each inner nodes
+		 */
+		TArray<FPoint> Normals;
+
+	public:
+		FGridChronos Chronos;
+
+	public:
+		FGrid(TSharedRef<FTopologicalFace>& InFace, TSharedRef<FModelMesh>& InShellMesh);
+
+		virtual ~FGrid() = default;
+
+		// ======================================================================================================================================================================================================================
+		// Meshing tools =========================================================================================================================================================================================================
+		// ======================================================================================================================================================================================================================
+		
+		/**
+		 * Return true if the grid is not consistent to build a mesh as composed of only two border nodes.
+		 */
+		bool CheckIfDegenerated();
+
+		/**
+		 * Defines the cutting coordinate of the grid according to mesh criteria and the existing mesh of bordering edges (loop's edges)
+		 * @see GetPreferredUVCoordinatesFromNeighbours
+		 */
+		void DefineCuttingParameters();
+		void DefineCuttingParameters(EIso Iso, FCuttingGrid& Neighbors);
+
+		/**
+		 * Computes the 2d Points, 3D points and Normals of the grid.
+		 * call of DefineCuttingParameters is mandatory called before
+		 * @return false if the grid is degenerated
+		 */
+		bool GeneratePointCloud();
+
+		/**
+		 * Process the generated point cloud to compute the scaled parametric spaces and to identify outer points of the face
+		 * @see IsNodeInsideFace
+		 * @see Points2D
+		 */
+		void ProcessPointCloud();
+
+		/**
+		 * Convert an array of points of "DefaultParametric" space into a scaled parametric space
+		 * @see FThinZone2DFinder::BuildBoundarySegments()
+		 */
+		void TransformPoints(EGridSpace DestinationSpace, const TArray<FPoint2D>& InPointsToScale, TArray<FPoint2D>& OutScaledPoints) const;
+
+	protected:
+
+		/**
+		 * Builds the scaled parametric spaces 
+		 * @see Points2D
+		 * @see GeneratePointCloud (ended GeneratePointCloud process)
+		 */
+		void ScaleGrid();
+
+		/**
+		 * Projects loop's points in the scaled parametric spaces
+		 * @see FaceLoops2D
+		 * @see ProcessPointClound (called in ProcessPointClound)
+		 */
+		void ScaleLoops();
+
+		/**
+		 * Convert Coordinate of "DefaultParametric" space into a scaled parametric space
+		 * @see ScaleLoops
+		 */
+		void ComputeNewCoordinate(const TArray<FPoint2D>& NewGrid, int32 IndexU, int32 IndexV, const FPoint2D& InPoint, FPoint2D& OutNewScaledPoint) const
+		{
+			OutNewScaledPoint =
+				NewGrid[(IndexV + 0) * CuttingCount[EIso::IsoU] + (IndexU + 0)] * (CuttingCoordinates[EIso::IsoU][IndexU + 1] - InPoint.U) * (CuttingCoordinates[EIso::IsoV][IndexV + 1] - InPoint.V) -
+				NewGrid[(IndexV + 0) * CuttingCount[EIso::IsoU] + (IndexU + 1)] * (CuttingCoordinates[EIso::IsoU][IndexU + 0] - InPoint.U) * (CuttingCoordinates[EIso::IsoV][IndexV + 1] - InPoint.V) -
+				NewGrid[(IndexV + 1) * CuttingCount[EIso::IsoU] + (IndexU + 0)] * (CuttingCoordinates[EIso::IsoU][IndexU + 1] - InPoint.U) * (CuttingCoordinates[EIso::IsoV][IndexV + 0] - InPoint.V) +
+				NewGrid[(IndexV + 1) * CuttingCount[EIso::IsoU] + (IndexU + 1)] * (CuttingCoordinates[EIso::IsoU][IndexU + 0] - InPoint.U) * (CuttingCoordinates[EIso::IsoV][IndexV + 0] - InPoint.V);
+			OutNewScaledPoint /= (CuttingCoordinates[EIso::IsoU][IndexU + 1] - CuttingCoordinates[EIso::IsoU][IndexU + 0]) * (CuttingCoordinates[EIso::IsoV][IndexV + 1] - CuttingCoordinates[EIso::IsoV][IndexV + 0]);
+		};
+
+		void FindCoordinateIndex(const TArray<double>& InCoordinates, double Coordinate, int32& OutIndex) const
+		{
+			ensureCADKernel(InCoordinates.IsValidIndex(OutIndex));
+			while (Coordinate < InCoordinates[OutIndex] && OutIndex > 0)
+			{
+				OutIndex--;
+			}
+			for (; OutIndex + 2 < InCoordinates.Num() && Coordinate > InCoordinates[OutIndex + 1]; ++OutIndex);
+			ensureCADKernel(InCoordinates.IsValidIndex(OutIndex));
+		};
+
+		/**
+		 * @see ProcessPointClound (called in ProcessPointClound)
+		 */
+		void FindInnerFacePoints();
+
+		/**
+		 * Define if a node is closes to a loop i.e. the loop cross the space [[IndexU - 1, IndexU + 1], [IndexV - 1, IndexV + 1]]
+		 * This algorithm is inspired of Bresenham's line algorithm
+		 * @see IsNodeCloseToLoop
+		 * @see ProcessPointClound (called in ProcessPointClound)
+		 */
+		void FindPointsCloseToLoop();
+
+		/**
+		 * Find all node close to the boundary.
+		 * For each node, the 3d distance to the boundary is estimated
+		 * If this distance is bigger than the Tolerance, the node is remove
+		 * Removing a node can make its neighbor removable,
+		 * Each removed node make its direct neighbor removable, the process is recursive as long as it is possible to remove a node
+		 *
+		 * @see ProcessPointClound (called in ProcessPointClound)
+		 */
+		void RemovePointsClosedToLoop();
+
+
+		/**
+		 * Gets the cutting coordinates of the existing mesh of bordering edges (loop's edges)
+		 * @see DefineCuttingParameters (called in DefineCuttingParameters)
+		 */
+		void GetPreferredUVCoordinatesFromNeighbours(FCuttingGrid& NeighboursCutting);
+
+		void ComputeMaxElementSize();
+		void ComputeMaxDeltaUV();
+
+		/**
+		 * @ return false if the mesh of the loop is degenerated
+		 */
+		bool GetMeshOfLoops();
+
+	public:
+
+		// ======================================================================================================================================================================================================================
+		// ThinZone properties for meshing scheduling optimisation ==============================================================================================================================================================
+		// ======================================================================================================================================================================================================================
+
+		void SearchThinZones();
+		const TArray<FThinZone2D>& GetThinZones() const
+		{
+			return ThinZoneFinder.GetThinZones();
+		}
+
+		// ======================================================================================================================================================================================================================
+		// GET Methodes =========================================================================================================================================================================================================
+		// ======================================================================================================================================================================================================================
+
+		constexpr const int32 GetCuttingCount(EIso Iso) const
+		{
+			return CuttingCount[Iso];
+		}
+
+		/**
+		 * Return the number of points (inner and outer points) of the grid (i.e. CuttingCount[IsoU] x CuttingCount[IsoV]).
+		 */
+		const int32 GetTotalCuttingCount() const
+		{
+			return CuttingSize;
+		}
+
+		const double GetTolerance(EIso Iso) const
+		{
+			return FaceTolerance[Iso];
+		}
+
+		/**
+		 * @return true if the node is inner the external loop and outer the inner loops 
+		 */
+		const bool IsNodeInsideFace(int32 IndexU, int32 IndexV) const
+		{
+			return IsInsideFace[GobalIndex(IndexU, IndexV)] == 1;
+		}
+
+		/**
+		 * @return true if the node is inner the external loop and outer the inner loops
+		 */
+		const bool IsNodeInsideFace(int32 Index) const
+		{
+			return IsInsideFace[Index] == 1;
+		}
+
+		/**
+		 * @return the count of inner nodes
+		 */
+		const int32 InnerNodesCount() const
+		{
+			return CountOfInnerNodes;
+		}
+
+		/**
+		 * @return true if the node is close to a loop i.e. the loop cross the space [[IndexU - 1, IndexU + 1], [IndexV - 1, IndexV + 1]]
+		 */
+		const bool IsNodeCloseToLoop(int32 IndexU, int32 IndexV) const
+		{
+			return IsCloseToLoop[GobalIndex(IndexU, IndexV)] == 1;
+		}
+
+		/**
+		 * @return true if the node is close to a loop i.e. the loop cross the space [[IndexU - 1, IndexU + 1], [IndexV - 1, IndexV + 1]]
+		 */
+		const bool IsNodeCloseToLoop(int32 Index) const
+		{
+			return IsCloseToLoop[Index] == 1;
+		}
+
+		/**
+		 * @return the FPoint2D (parametric coordinates) of the point at the Index of the grid in the defined grid space 
+		 * @see EGridSpace
+		 * @see Points2D
+		 */
+		const FPoint2D& GetInner2DPoint(EGridSpace Space, int32 Index) const
+		{
+			return Points2D[(int32)Space][Index];
+		}
+
+		/**
+		 * @return the FPoint2D (parametric coordinates) of the point at the Index of the grid in the defined grid space @see EGridSpace
+		 */
+		const FPoint2D& GetInner2DPoint(EGridSpace Space, int32 IndexU, int32 IndexV) const
+		{
+			return Points2D[(int32)Space][GobalIndex(IndexU, IndexV)];
+		}
+
+		/**
+		 * @return the FPoint (3D coordinates) of the point at the Index of the grid
+		 */
+		const FPoint& GetInner3DPoint(int32 Index) const
+		{
+			return Points3D[Index];
+		}
+
+		/**
+		 * @return the FPoint (3D coordinates) of the point at the Index of the grid
+		 */
+		const FPoint& GetInner3DPoint(int32 IndexU, int32 IndexV) const
+		{
+			return Points3D[GobalIndex(IndexU, IndexV)];
+		}
+
+		/**
+		 * @return the normal of the surface at the point at the Index of the grid
+		 */
+		const FPoint& GetPointNormal(int32 IndexU, int32 IndexV) const
+		{
+			return Normals[GobalIndex(IndexU, IndexV)];
+		}
+
+		/**
+		 * @return the normal of the surface at the point at the Index of the grid
+		 */
+		const FPoint& GetPointNormal(int32 Index) const
+		{
+			return Normals[Index];
+		}
+
+		constexpr const TArray<double>& GetCuttingCoordinatesAlongIso(EIso Iso) const
+		{
+			return CuttingCoordinates[Iso];
+		}
+
+		const FCoordinateGrid& GetCuttingCoordinates() const
+		{
+			return CuttingCoordinates;
+		}
+
+		/**
+		 * @return the array of 3d points of the grid
+		 */
+		TArray<FPoint>& GetInner3DPoints()
+		{
+			return Points3D;
+		}
+
+		/**
+		 * @return the array of 3d points of the grid
+		 */
+		const TArray<FPoint>& GetInner3DPoints() const
+		{
+			return Points3D;
+		}
+
+		/**
+		 * @return the array of 2d points of the grid in the defined space
+		 */
+		const TArray<FPoint2D>& GetInner2DPoints(EGridSpace Space) const
+		{
+			return Points2D[(int32)Space];
+		}
+
+		/**
+		 * @return the array of normal of the points of the grid
+		 */
+		TArray<FPoint>& GetNormals()
+		{
+			return Normals;
+		}
+
+		TSharedRef<FTopologicalFace> GetFace() const
+		{
+			return Face;
+		}
+
+		const TArray<TArray<int32>>& GetNodeIdsOfFaceLoops() const
+		{
+			return NodeIdsOfFaceLoops;
+		}
+
+		const FPoint2D& GetLoop2DPoint(EGridSpace Space, int32 LoopIndex, int32 Index) const
+		{
+			return FaceLoops2D[(int32)Space][LoopIndex][Index];
+		}
+
+		const FPoint& GetLoop3DPoint(int32 LoopIndex, int32 Index) const
+		{
+			return FaceLoops3D[LoopIndex][Index];
+		}
+
+		const int32 GetLoopCount() const
+		{
+			return FaceLoops2D[0].Num();
+		}
+
+		/**
+		 * @return the array of array of 2d points of the loops according to the defined grid space
+		 */
+		const TArray<TArray<FPoint2D>>& GetLoops2D(EGridSpace Space) const
+		{
+			return FaceLoops2D[(int32)Space];
+		}
+
+		/**
+		 * @return the array of array of 3d points of the loops
+		 */
+		TArray<TArray<FPoint>>& GetLoops3D()
+		{
+			return FaceLoops3D;
+		}
+
+		/**
+		 * @return the array of array of 3d points of the loops
+		 */
+		const TArray<TArray<FPoint>>& GetLoops3D() const
+		{
+			return FaceLoops3D;
+		}
+
+		/**
+		 * @return the array of array of normal of the points of the loops
+		 */
+		const TArray<TArray<FPoint>>& GetLoopNormals() const
+		{
+			return NormalsOfFaceLoops;
+		}
+
+		/**
+		 * @return the Index of the position in the arrays of a point [IndexU, IndexV] of the grid
+		 */
+		int32 GobalIndex(int32 IndexU, int32 IndexV) const
+		{
+			return IndexV * CuttingCount[EIso::IsoU] + IndexU;
+		}
+
+		/**
+		 * @return the maximum difference of coordinate along the specified axis of two two successive points
+		 */
+		double GetMaxDeltaU(EIso Iso) const
+		{
+			return MaxDeltaUV[Iso];
+		}
+
+		/**
+		 * @return the minimal size of the mesh elements in the grid
+		 */
+		double GetMinElementSize() const
+		{
+			return MinOfMaxElementSize;
+		}
+
+		// ======================================================================================================================================================================================================================
+		// Display Methodes   ================================================================================================================================================================================================
+		// ======================================================================================================================================================================================================================
+		bool bDisplay = false;
+
+		void DisplayFindInnerDomainPoints(EGridSpace DisplaySpace) const;
+
+		template<typename TPoint>
+		void DisplayPoints(FString Message, const TArray<TPoint>& Points) const
+		{
+			if (!bDisplay)
+			{
+				return;
+			}
+
+			Open3DDebugSession(Message);
+			for (int32 Index = 0; Index < Points.Num(); ++Index)
+			{
+				DisplayPoint(Points[Index], Index);
+			}
+			Close3DDebugSession();
+		}
+
+		template<typename TPoint>
+		void DisplayInnerDomainPoints(FString Message, const TArray<TPoint>& Points) const
+		{
+			if (!bDisplay)
+			{
+				return;
+			}
+
+			Open3DDebugSession(Message);
+			for (int32 Index = 0; Index < Points.Num(); ++Index)
+			{
+				if (IsInsideFace[Index])
+				{
+					if (IsCloseToLoop[Index])
+					{
+						DisplayPoint(Points[Index], EVisuProperty::BluePoint, Index);
+					}
+					else
+					{
+						DisplayPoint(Points[Index], EVisuProperty::GreenPoint, Index);
+					}
+				}
+			}
+			Close3DDebugSession();
+		}
+
+		template<typename TPoint>
+		void DisplayLoop(FString Message, const TArray<TArray<TPoint>>& Loops, bool bDisplayNodes, bool bMakeGroup) const
+		{
+			if (!bDisplay)
+			{
+				return;
+			}
+
+			Open3DDebugSession(*Message);
+			int32 LoopIndex = 0;
+			for (const TArray<TPoint>& Loop : Loops)
+			{
+				if (bMakeGroup)
+				{
+					Open3DDebugSession(FString::Printf(TEXT("Loop %d"), LoopIndex++));
+				}
+
+				const TPoint* FirstSegmentPoint = &Loop[0];
+				if (bDisplayNodes)
+				{
+					DisplayPoint(*FirstSegmentPoint, EVisuProperty::BluePoint, 0);
+				}
+
+				for (int32 Index = 1; Index < Loop.Num(); ++Index)
+				{
+					const TPoint* SecondSegmentPoint = &Loop[Index];
+					DisplaySegment(*FirstSegmentPoint, *SecondSegmentPoint);
+					if (bDisplayNodes)
+					{
+						DisplayPoint(*SecondSegmentPoint, EVisuProperty::BluePoint, Index);
+					}
+					FirstSegmentPoint = SecondSegmentPoint;
+				}
+				DisplaySegment(*FirstSegmentPoint, *Loop.begin());
+				if (bMakeGroup)
+				{
+					Close3DDebugSession();
+				}
+			}
+			Close3DDebugSession();
+		}
+
+		void DisplayFindPointsCloseToLoop(EGridSpace DisplaySpace) const;
+		void DisplayFindPointsCloseAndInsideToLoop(EGridSpace DisplaySpace) const;
+		void PrintTimeElapse() const;
+	};
+
+}
diff --git a/Engine/Source/Runtime/Datasmith/CADKernel/Private/CADKernel/Mesh/Structure/Mesh.cpp b/Engine/Source/Runtime/Datasmith/CADKernel/Private/CADKernel/Mesh/Structure/Mesh.cpp
new file mode 100644
--- /dev/null
+++ b/Engine/Source/Runtime/Datasmith/CADKernel/Private/CADKernel/Mesh/Structure/Mesh.cpp
@@ -0,0 +1,27 @@
+// Copyright Epic Games, Inc. All Rights Reserved.
+
+#include "CADKernel/Mesh/Structure/Mesh.h"
+
+#include "CADKernel/Mesh/Structure/ModelMesh.h"
+#include "CADKernel/Topo/TopologicalEntity.h"
+
+using namespace CADKernel;
+
+
+int32 FMesh::RegisterCoordinates()
+{
+	ModelMesh.Pin()->RegisterCoordinates(NodeCoordinates, StartNodeId, MeshModelIndex);
+	LastNodeIndex = StartNodeId + (int32)NodeCoordinates.Num();
+	return StartNodeId;
+}
+
+
+#ifdef CADKERNEL_DEV
+FInfoEntity& FMesh::GetInfo(FInfoEntity& Info) const
+{
+	return FEntityGeom::GetInfo(Info)
+		.Add(TEXT("geometric Entity"), GetGeometricEntity())
+		.Add(TEXT("mesh model"), GetMeshModel())
+		.Add(TEXT("Node Num"), (int32) NodeCoordinates.Num());
+}
+#endif
diff --git a/Engine/Source/Runtime/Datasmith/CADKernel/Private/CADKernel/Mesh/Structure/ModelMesh.cpp b/Engine/Source/Runtime/Datasmith/CADKernel/Private/CADKernel/Mesh/Structure/ModelMesh.cpp
new file mode 100644
--- /dev/null
+++ b/Engine/Source/Runtime/Datasmith/CADKernel/Private/CADKernel/Mesh/Structure/ModelMesh.cpp
@@ -0,0 +1,136 @@
+// Copyright Epic Games, Inc. All Rights Reserved.
+
+#include "CADKernel/Mesh/Structure/ModelMesh.h"
+
+#include "CADKernel/Mesh/Criteria/Criterion.h"
+#include "CADKernel/Mesh/Structure/EdgeMesh.h"
+#include "CADKernel/Mesh/Structure/Mesh.h"
+#include "CADKernel/Mesh/Structure/FaceMesh.h"
+#include "CADKernel/Mesh/Structure/VertexMesh.h"
+#include "CADKernel/Topo/TopologicalEntity.h"
+
+using namespace CADKernel;
+
+
+void FModelMesh::AddCriterion(TSharedPtr<FCriterion>& Criterion)
+{
+	Criteria.Add(Criterion);
+	switch (Criterion->GetCriterionType())
+	{
+	case ECriterion::MinSize:
+		MinSize = Criterion->Value();
+		break;
+	case ECriterion::MaxSize:
+		MaxSize = Criterion->Value();
+		break;
+	case ECriterion::Angle:
+		MaxAngle = Criterion->Value();
+		break;
+	case ECriterion::Sag:
+		Sag = Criterion->Value();
+		break;
+	case ECriterion::CADCurvature:
+		QuadAnalyse = true;
+		break;
+	}
+}
+
+#ifdef CADKERNEL_DEV
+FInfoEntity& FModelMesh::GetInfo(FInfoEntity& Info) const
+{
+	return FEntityGeom::GetInfo(Info)
+		.Add(TEXT("Surface Meshes"), FaceMeshes)
+		.Add(TEXT("Edge Meshes"), EdgeMeshes)
+		.Add(TEXT("Vertex Meshes"), VertexMeshes);
+}
+#endif
+
+const int32 FModelMesh::GetIndexOfVertexFromId(const int32 Ident) const
+{
+	for (const TSharedPtr<FVertexMesh>& VertexMesh : VertexMeshes)
+	{
+		if (Ident == VertexMesh->GetStartVertexId())
+		{
+			return VertexMesh->GetIndexInMeshModel();
+		}
+	}
+	return -1;
+}
+
+const TSharedPtr<FVertexMesh> FModelMesh::GetMeshOfVertexNodeId(const int32 Ident) const
+{
+	for (const TSharedPtr<FVertexMesh>& VertexMesh : VertexMeshes)
+	{
+		if (Ident == VertexMesh->GetStartVertexId())
+		{
+			return VertexMesh;
+		}
+	}
+	return nullptr;
+}
+
+const int32 FModelMesh::GetIndexOfEdgeFromId(const int32 Ident) const
+{
+	for (const TSharedPtr<FEdgeMesh>& EdgeMesh : EdgeMeshes)
+	{
+		if (Ident >= EdgeMesh->GetStartVertexId())
+		{
+			if (Ident <= EdgeMesh->GetLastVertexIndex())
+			{
+				return EdgeMesh->GetIndexInMeshModel();
+			}
+		}
+	}
+	return -1;
+}
+
+const int32 FModelMesh::GetIndexOfSurfaceFromId(const int32 Ident) const
+{
+	for (const TSharedPtr<FFaceMesh>& FaceMesh : FaceMeshes)
+	{
+		if (Ident >= FaceMesh->GetStartVertexId())
+		{
+			if (Ident <= FaceMesh->GetLastVertexIndex())
+			{
+				return FaceMesh->GetIndexInMeshModel();
+			}
+		}
+	}
+	return -1;
+}
+
+void FModelMesh::GetNodeCoordinates(TArray<FPoint>& NodeCoordinates) const
+{
+	NodeCoordinates.Reserve(LastIdUsed);
+
+	for (const TArray<FPoint>* PointArray : GlobalPointCloud)
+	{
+		NodeCoordinates.Insert(*PointArray, NodeCoordinates.Num());
+	}
+}
+
+void FModelMesh::GetNodeCoordinates(TArray<FVector>& NodeCoordinates) const
+{
+	NodeCoordinates.Reserve(LastIdUsed);
+
+	for (const TArray<FPoint>* PointArray : GlobalPointCloud)
+	{
+		for (const FPoint& Point :* PointArray)
+		{
+			NodeCoordinates.Emplace((float)Point.X, (float)Point.Y, (float)Point.Z);
+		}
+	}
+}
+
+const TArray<TSharedPtr<FMesh>>& FModelMesh::GetMeshes() const
+{
+	if (FaceMeshes.Num())
+	{
+		return (TArray<TSharedPtr<FMesh>>&) FaceMeshes;
+	}
+	if (EdgeMeshes.Num())
+	{
+		return (TArray<TSharedPtr<FMesh>>&) EdgeMeshes;
+	}
+	return (TArray<TSharedPtr<FMesh>>&) VertexMeshes;
+}
