CADKernel Source Code v0.01

diff --git a/Engine/Source/Runtime/Datasmith/CADKernel/Private/CADKernel/Mesh/Structure/ThinZone2D.cpp b/Engine/Source/Runtime/Datasmith/CADKernel/Private/CADKernel/Mesh/Structure/ThinZone2D.cpp
new file mode 100644
--- /dev/null
+++ b/Engine/Source/Runtime/Datasmith/CADKernel/Private/CADKernel/Mesh/Structure/ThinZone2D.cpp
@@ -0,0 +1,808 @@
+// Copyright Epic Games, Inc. All Rights Reserved.
+
+#include "CADKernel/Mesh/Structure/ThinZone2D.h"
+
+#include "CADKernel/Math/Point.h"
+#include "CADKernel/Mesh/Structure/Grid.h"
+#include "CADKernel/Mesh/Structure/EdgeSegment.h"
+#include "CADKernel/Topo/TopologicalEdge.h"
+#include "CADKernel/Topo/TopologicalFace.h"
+#include "CADKernel/Topo/TopologicalLoop.h"
+#include "CADKernel/UI/Display.h"
+
+using namespace CADKernel;
+
+FIdent FEdgeSegment::LastId = 0;
+
+void FThinZone2DFinder::FindClosedSegments()
+{
+	double MaxSegmentLength = 0.0;
+	for (const FEdgeSegment* Segment : LoopSegments)
+	{
+		double Length = Segment->GetLength();
+		if (Length > MaxSegmentLength) 
+		{
+			MaxSegmentLength = Length;
+		}
+	}
+	MaxSegmentLength *= 1.01;
+
+	double MaxSpace = FMath::Max(MaxSegmentLength, Tolerance * 1.01);
+	double MaxSquareSpace = 4 * FMath::Square(MaxSpace);
+
+	// Copy of loop segments to generate a sorted array of segments
+	SortedLoopSegments = LoopSegments;
+	Algo::Sort(SortedLoopSegments, [](FEdgeSegment* SegmentA, FEdgeSegment* SegmentB)
+		{
+			return (SegmentA->GetAxeMin() <SegmentB->GetAxeMin());
+		});
+
+
+#ifdef DEBUG_THIN_ZONES
+	Open3DDebugSession(TEXT("Closed Segment");
+#endif
+
+	// For each segment, the nearest segment is search
+	// If segment is from an inner loop, the nearest could not be from the same loop
+	int32 SegmentToBeginIndex = 0;
+	for (FEdgeSegment* Segment : SortedLoopSegments)
+	{
+		TSharedPtr<FTopologicalLoop> SegmentLoop = nullptr;
+		if (Segment->IsInner())
+		{
+			SegmentLoop = Segment->GetEdge()->GetLoop();
+		}
+		else
+		{
+			SegmentLoop = nullptr;
+		}
+
+		FPoint2D SegmentMiddle = Segment->GetCenter();
+
+		FPoint2D ClosedPoint, ClosedPoint2;
+		FEdgeSegment* ClosedSegment = nullptr;
+
+		double SegmentMin = Segment->GetAxeMin();
+		double MinSquareThickness = HUGE_VALUE;
+
+		for (int32 CandidateIndex = SegmentToBeginIndex; CandidateIndex <SortedLoopSegments.Num(); ++CandidateIndex)
+		{
+			FEdgeSegment* Candidate = SortedLoopSegments[CandidateIndex];
+			if (Candidate == Segment) 
+			{
+				continue;
+			}
+
+			// Inner boundary are not check for thin inner boundary
+			if (SegmentLoop == Candidate->GetEdge()->GetLoop()) 
+			{
+				continue;
+			}
+
+			double CandidateSegmentMin = Candidate->GetAxeMin();
+
+			// If min point of candiate segment + maxSpace (= max length of the segments + Tolerance) is smaller than Min point of current segment then the distance between both segments cannot be smaller the Tolerance
+			if ((CandidateSegmentMin + MaxSpace) < SegmentMin)
+			{
+				SegmentToBeginIndex = CandidateIndex;
+				continue;
+			}
+
+			// If Min point of current segment + maxSpace is smaller than Min point of candidate segment then the distance between both segments cannot be smaller the Tolerance.
+			// As segments are sorted, next segments are not close to the current segment
+			if (CandidateSegmentMin > (SegmentMin + MaxSpace))
+			{
+				break;
+			}
+
+			const FPoint2D& FirstPointCandidate = Candidate->GetExtemity(ELimit::Start);
+
+			// if the distance of FirstPointCandidate with the middle of current segment is biggest than MaxSpace, then the projection of the point cannot be smaller than the tolerance, 
+			double SquareDistance = SegmentMiddle.SquareDistance(FirstPointCandidate);
+			if (SquareDistance > MaxSquareSpace)
+			{
+				continue;
+			}
+
+			const FPoint2D& SecondPointCandidate = Candidate->GetExtemity(ELimit::End);
+			double Coordinate;
+			FPoint2D Projection = ProjectPointOnSegment(SegmentMiddle, FirstPointCandidate, SecondPointCandidate, Coordinate, true);
+
+			SquareDistance = SegmentMiddle.SquareDistance(Projection);
+			if (SquareDistance > SquareTolerance)
+			{
+				continue;
+			}
+
+			if (MinSquareThickness > SquareDistance)
+			{
+				// check the angle between segments. As they are opposite, the cosAngle as to be close to -1
+				double Slop = Segment->ComputeUnorientedSlopeOf(Candidate);
+				if (Slop < 3.33) // Angle <5Pi/6 (150 deg)
+				{
+					continue;
+				}
+
+				// check the angle between segment and Middle-Projection. As they are opposite, the cosAngle as to be close to 0
+				Slop = Segment->ComputeUnorientedSlopeOf(SegmentMiddle, Projection);
+				if (Slop < 1.33) // Angle <Pi/3 (60 deg)
+				{
+					continue;
+				}
+
+				MinSquareThickness = SquareDistance;
+				ClosedPoint = Projection;
+				ClosedSegment = Candidate;
+			}
+		}
+
+		if (ClosedSegment)
+		{
+			Segment->SetClosedSegment(ClosedSegment, nullptr, MinSquareThickness, true);
+#ifdef DEBUG_THIN_ZONES
+			DisplaySegment(Segment->GetFirstPoint(), Segment->GetSecondPoint());
+			DisplaySegment(ClosedSegment->GetFirstPoint(), ClosedSegment->GetSecondPoint());
+			Wait();
+#endif
+		}
+	}
+#ifdef DEBUG_THIN_ZONES
+	Close3DDebugSession();
+#endif
+}
+
+void FThinZone2DFinder::DisplayClosedSegments()
+{
+#ifdef DEBUG_THIN_ZONES
+	F3DDebugSession _(TEXT("Closed Segment"));
+	for (const FEdgeSegment* Segment : LoopSegments)
+	{
+		if (Segment->GetClosedSegment())
+		{
+			DisplaySegment(Segment->GetFirstPoint(), Segment->GetSecondPoint(), Segment->GetId(), PropertyElement);
+			FEdgeSegment* Parallel = Segment->GetClosedSegment();
+			DisplaySegment(Parallel->GetFirstPoint(), Parallel->GetSecondPoint(), Segment->GetId(), PropertyElement);
+
+			double Coordinate;
+			FPoint Projection = ProjectPointOnSegment(Segment->GetCenter(), Parallel->GetFirstPoint(), Parallel->GetSecondPoint(), &Coordinate, true);
+
+			DisplaySegment(Projection, Segment->GetCenter(), Segment->GetId(), EVisuProperty::Iso);
+		}
+	}
+#endif
+}
+
+void FThinZone2DFinder::CheckClosedSegments()
+{
+	for (FEdgeSegment* Segment : LoopSegments)
+	{
+		if (Segment->GetClosedSegment())
+		{
+			double NextSegmentThickness = FMath::Min(Segment->GetPrevious()->GetClosedSquareDistance(), Segment->GetNext()->GetClosedSquareDistance());
+			double SegmentThickness = Segment->GetClosedSquareDistance();
+			Sort(NextSegmentThickness, SegmentThickness);
+			SegmentThickness = sqrt(SegmentThickness);
+			NextSegmentThickness = sqrt(NextSegmentThickness);
+			double Delta = SegmentThickness - NextSegmentThickness;
+			if (Delta > (Segment->GetLength()))
+			{
+				Segment->ResetClosedData();
+			}
+		}
+	}
+}
+
+void FThinZone2DFinder::LinkClosedSegments()
+{
+	SortedLoopSegments.Empty();
+	SortedLoopSegments.Reserve(LoopSegments.Num());
+
+	FEdgeSegment* StartSegment = nullptr;
+	FEdgeSegment* Segment = nullptr;
+	FIdent Index = 0;
+	FIdent StartZoneIndex = 0;
+	TSharedPtr<FTopologicalLoop> LastBoundary = nullptr;
+	bool bThinZone = false;
+
+	for (FEdgeSegment* EdgeSegment : LoopSegments)
+	{
+		Segment = EdgeSegment;
+		if (LastBoundary != EdgeSegment->GetEdge()->GetLoop())
+		{
+			LastBoundary = EdgeSegment->GetEdge()->GetLoop();
+			Index++;
+			StartSegment = EdgeSegment;
+			StartZoneIndex = Index;
+			bThinZone = false;
+		}
+
+		if (Segment->GetClosedSegment() != nullptr)
+		{
+			Segment->SetChainIndex(Index);
+			SortedLoopSegments.Add(Segment);
+			bThinZone = true;
+
+			if (Segment->GetNext() == StartSegment)
+			{
+				FIdent currentZoneIndex = Segment->GetChainIndex();
+				if (StartZoneIndex != 0 && currentZoneIndex != StartZoneIndex)
+				{
+					do {
+						Segment->SetChainIndex(StartZoneIndex);
+						Segment = Segment->GetPrevious();
+					} while (Segment->GetChainIndex() == currentZoneIndex);
+				}
+			}
+		}
+		else
+		{
+			if (bThinZone)
+			{
+				Index++;
+				bThinZone = false;
+			}
+		}
+	}
+
+	// the number of ThinZone should be less than 2x the number of Chain, 
+	ThinZones.Reserve(Index * 2);
+#ifdef DEBUG_THIN_ZONES
+	{
+		F3DDebugSession _(TEXT("Linking"));
+		FIdent LastIndex = 0;
+		int32 Color = 1;
+		for (FEdgeSegment* EdgeSegment : SortedLoopSegments)
+		{
+			if (LastIndex != EdgeSegment->GetChainIndex())
+			{
+				if (Color > 1) Close3DDebugSession();
+				Color++;
+				LastIndex = EdgeSegment->GetChainIndex();
+				Open3DDebugSession(FString::Printf(TEXT("Chain %d\n"), LastIndex));
+				if (Color == 7) Color++;
+				if (Color == 11) Color = 2;
+			}
+			DisplaySegment(EdgeSegment->GetExtemity(ELimit::Start), EdgeSegment->GetExtemity(ELimit::End), EdgeSegment->GetId(), (EVisuProperty)Color);
+		}
+		if (SortedLoopSegments.Num() > 0)
+		{
+			Close3DDebugSession();
+		}
+	}
+#endif
+
+}
+
+
+void GetThinZoneSideConnectionsLength(TArray<FEdgeSegment*>& FirstSide, TArray<FEdgeSegment*>& SecondSide, double MaxLength, double* LengthBetweenExtremity, TArray<TSharedPtr<FTopologicalEdge>>* PeakEdges)
+{
+	LengthBetweenExtremity[0] = 0.0;
+	TSharedPtr<FTopologicalEdge> Edge = nullptr;
+	FEdgeSegment* Segment = FirstSide[0]->GetPrevious();
+	while (Segment != SecondSide[0])
+	{
+		if (Edge != Segment->GetEdge())
+		{
+			Edge = Segment->GetEdge();
+			PeakEdges[0].Add(Edge);
+		}
+		LengthBetweenExtremity[0] += Segment->GetLength();
+		Segment = Segment->GetPrevious();
+
+		if (LengthBetweenExtremity[0] > MaxLength)
+		{
+			LengthBetweenExtremity[0] = HUGE_VALUE;
+			PeakEdges[0].Empty(0);
+			break;
+		}
+	}
+
+	LengthBetweenExtremity[1] = 0.0;
+	Edge = nullptr;
+	Segment = FirstSide.Last()->GetNext();
+	while (Segment != SecondSide.Last())
+	{
+		if (Edge != Segment->GetEdge())
+		{
+			Edge = Segment->GetEdge();
+			PeakEdges[1].Add(Edge);
+		}
+
+		LengthBetweenExtremity[1] += Segment->GetLength();
+		Segment = Segment->GetNext();
+
+		if (LengthBetweenExtremity[1] > MaxLength)
+		{
+			LengthBetweenExtremity[1] = HUGE_VALUE;
+			PeakEdges[1].Empty(0);
+			break;
+		}
+	}
+}
+
+void FThinZone2DFinder::BuildThinZone()
+{
+	for (FEdgeSegment* EdgeSegment : SortedLoopSegments)
+	{
+		EdgeSegment->ResetMarker1();
+	}
+
+	TArray<FEdgeSegment*> FirstSide;
+	TArray<FEdgeSegment*> SecondSide;
+	FirstSide.Reserve(SortedLoopSegments.Num());
+	SecondSide.Reserve(SortedLoopSegments.Num());
+
+	FEdgeSegment* FirstSideSegment = nullptr;
+	FEdgeSegment* SecondSideSegment = nullptr;
+
+	for (FEdgeSegment* EdgeSegment : SortedLoopSegments)
+	{
+		bool bFirstSideIsClosed = false;
+		bool bSecondSideIsClosed = false;
+
+		FirstSide.Empty();
+		SecondSide.Empty();
+
+		if (EdgeSegment->HasMarker1()) 
+		{
+			continue;
+		}
+
+		FirstSideSegment = EdgeSegment;
+
+		// Find the first segment
+		FIdent FirstChainIndex = FirstSideSegment->GetChainIndex();
+		FEdgeSegment* Segment1 = FirstSideSegment;
+		while (Segment1->GetChainIndex() == FirstChainIndex && !Segment1->HasMarker1())
+		{
+			Segment1 = Segment1->GetPrevious();
+			if (Segment1 == FirstSideSegment)
+			{
+				Segment1 = Segment1->GetPrevious();
+				bFirstSideIsClosed = true;
+				break;
+			}
+		}
+		FirstSideSegment = Segment1->GetNext();
+
+		// Find the parallel chain
+		SecondSideSegment = FirstSideSegment->GetClosedSegment();
+		if (SecondSideSegment == nullptr)
+		{
+			SecondSideSegment = FirstSideSegment->GetNext()->GetClosedSegment();
+			if (SecondSideSegment == nullptr)
+			{
+				FirstSideSegment->HasMarker1();
+				continue;
+			}
+		}
+		FIdent SecondChainIndex = SecondSideSegment->GetChainIndex();
+		if (SecondChainIndex == 0) continue;
+
+		int32 SegmentNum = 0;
+		FEdgeSegment* Segment2 = SecondSideSegment;
+		if (!(bFirstSideIsClosed && (FirstChainIndex == SecondChainIndex)))
+		{
+			while (Segment2->GetChainIndex() == SecondChainIndex && !Segment2->HasMarker1() &&
+				((Segment2->GetClosedSegment() == FirstSideSegment) || (Segment2->GetClosedSegment() == FirstSideSegment->GetNext())
+					|| (FirstSideSegment->GetClosedSegment() == Segment2) || (FirstSideSegment->GetClosedSegment() == Segment2->GetPrevious())))
+			{
+				if (Segment2->HasMarker1())
+				{
+					break;
+				}
+
+				Segment2 = Segment2->GetNext();
+				if (SegmentNum > 0 && Segment2->GetPrevious() == SecondSideSegment)
+				{
+					bSecondSideIsClosed = true;
+					break;
+				}
+				SegmentNum++;
+			}
+			SecondSideSegment = Segment2->GetPrevious();
+		}
+
+		Segment1 = FirstSideSegment;
+		Segment2 = SecondSideSegment;
+
+		// both chains are browsed as long as they stay in front of each other
+		double MaxSquareThickness = 0.0;
+		do {
+			if (!FirstSideSegment->HasMarker1())
+			{
+				FirstSide.Add(FirstSideSegment);
+				FirstSideSegment->SetMarker1();
+				if (MaxSquareThickness <FirstSideSegment->GetClosedSquareDistance()) 
+				{
+					MaxSquareThickness = FirstSideSegment->GetClosedSquareDistance();
+				}
+			}
+
+			if (!SecondSideSegment->HasMarker1())
+			{
+				SecondSide.Add(SecondSideSegment);
+				SecondSideSegment->SetMarker1();
+				if (MaxSquareThickness <SecondSideSegment->GetClosedSquareDistance()) 
+				{
+					MaxSquareThickness = SecondSideSegment->GetClosedSquareDistance();
+				}
+			}
+
+			if (Segment1 == FirstSideSegment) 
+			{
+				Segment1 = FirstSideSegment->GetNext();
+			}
+			if (Segment2 == SecondSideSegment) 
+			{
+				Segment2 = SecondSideSegment->GetPrevious();
+			}
+
+			if ((!Segment1->HasMarker1()) && (Segment1->GetChainIndex() == FirstChainIndex) && ((Segment1->GetClosedSegment() == SecondSideSegment->GetNext()) || (Segment1->GetClosedSegment() == SecondSideSegment)))
+			{
+				FirstSideSegment = Segment1;
+			}
+			else if ((!Segment2->HasMarker1()) && (Segment2->GetChainIndex() == SecondChainIndex) && ((Segment2->GetClosedSegment() == FirstSideSegment->GetPrevious()) || (Segment2->GetClosedSegment() == FirstSideSegment)))
+			{
+				SecondSideSegment = Segment2;
+			}
+			else if ((!Segment1->HasMarker1()) && (Segment1->GetChainIndex() == FirstChainIndex) && (Segment1->GetClosedSegment() == SecondSideSegment->GetPrevious()))
+			{
+				FirstSideSegment = Segment1;
+			}
+			else if ((!Segment2->HasMarker1()) && (Segment2->GetChainIndex() == SecondChainIndex) && (Segment2->GetClosedSegment() == FirstSideSegment->GetNext()))
+			{
+				SecondSideSegment = Segment2;
+			}
+		} 
+		while ((FirstSideSegment == Segment1) || (SecondSideSegment == Segment2));
+
+		if ((FirstSide.Num() == 0) || (SecondSide.Num() == 0))
+		{
+			continue;
+		}
+
+		// Length of each side
+		double LengthSide1 = 0;
+		double MediumThickness1 = 0;
+		for (FEdgeSegment* Segment : FirstSide)
+		{
+			LengthSide1 += Segment->GetLength();
+			MediumThickness1 += Segment->GetClosedSquareDistance();
+		}
+		MediumThickness1 /= FirstSide.Num();
+		MediumThickness1 = sqrt(MediumThickness1);
+
+		double LengthSide2 = 0;
+		double MediumThickness2 = 0;
+		for (FEdgeSegment* Segment : SecondSide)
+		{
+			LengthSide2 += Segment->GetLength();
+			MediumThickness2 += Segment->GetClosedSquareDistance();
+		}
+		MediumThickness2 /= FirstSide.Num();
+		MediumThickness2 = sqrt(MediumThickness2);
+
+
+		if ((LengthSide1 <MediumThickness2 * 3) || (LengthSide2 <MediumThickness2 * 3))
+		{
+			continue;
+		}
+
+		const double MaxThickness = sqrt(MaxSquareThickness);
+		const double MaxLengthBetweendSideToBePeak = 3 * MaxThickness;
+
+		// make the thin zone
+		FThinZone2D& Zone = ThinZones.Emplace_GetRef(FirstSide, bFirstSideIsClosed, SecondSide, bSecondSideIsClosed, MaxThickness);
+
+		TFunction<void(TArray<TSharedPtr<FTopologicalEdge>>&)> DeleteCuttingPointOfPeak = [](const TArray<TSharedPtr<FTopologicalEdge>>& PeakEdges)
+		{
+			for (TSharedPtr<FTopologicalEdge> Edge : PeakEdges)
+			{
+				if (Edge->IsMeshed())
+				{
+					continue;
+				}
+				Edge->GetLinkActiveEdge()->SetThinPeak();
+			}
+		};
+
+		if (!Zone.GetFirstSide().IsInner() && !Zone.GetSecondSide().IsInner())
+		{
+			double LengthBetweenExtremity[2] = { HUGE_VALUE, HUGE_VALUE };
+			TArray<TSharedPtr<FTopologicalEdge>> PeakEdges[2];
+			GetThinZoneSideConnectionsLength(FirstSide, SecondSide, 10 * MaxThickness, LengthBetweenExtremity, PeakEdges);
+
+			if ((Zone.GetMaxLength() * 2 > (FirstLoopLength - MaxLengthBetweendSideToBePeak)) || (LengthBetweenExtremity[0] <MaxLengthBetweendSideToBePeak && LengthBetweenExtremity[1] <MaxLengthBetweendSideToBePeak))
+			{
+				Zone.SetCategory(EThinZone2DType::Global);
+				DeleteCuttingPointOfPeak(PeakEdges[0]);
+				DeleteCuttingPointOfPeak(PeakEdges[1]);
+			}
+			else if (LengthBetweenExtremity[0] <MaxLengthBetweendSideToBePeak)
+			{
+				Zone.SetCategory(EThinZone2DType::PeakStart);
+				if (Zone.GetFirstSide().GetLength() <MaxThickness * 10)
+				{
+					ThinZones.Pop();
+					continue;
+				}
+				DeleteCuttingPointOfPeak(PeakEdges[0]);
+			}
+			else if (LengthBetweenExtremity[1] <MaxLengthBetweendSideToBePeak)
+			{
+				Zone.SetCategory(EThinZone2DType::PeakEnd);
+				if (Zone.GetFirstSide().GetLength() <MaxThickness * 10)
+				{
+					ThinZones.Pop();
+					continue;
+				}
+				DeleteCuttingPointOfPeak(PeakEdges[1]);
+			}
+			else {
+				Zone.SetCategory(EThinZone2DType::Butterfly);
+			}
+		}
+		else
+		{
+			Zone.SetCategory(EThinZone2DType::BetweenLoops);
+		}
+
+		Zone.SetEdgesAsThinZone();
+	}
+}
+
+void FThinZone2DFinder::DisplayLoopSegments()
+{
+#ifdef DEBUG_THIN_ZONES
+	TSharedPtr<FTopologicalEdge> currentEdge = nullptr;
+	currentEdge = LoopSegments[0]->GetEdge();
+	F3DDebugSession _(TEXT("BoundarySegment"));
+
+	Open3DDebugSession(FString::Printf(TEXT("Edge %d\n"), currentEdge->GetId()));
+	for (FEdgeSegment* EdgeSegment : LoopSegments)
+	{
+		if (currentEdge != EdgeSegment->GetEdge())
+		{
+			currentEdge = EdgeSegment->GetEdge();
+			Close3DDebugSession();
+			Open3DDebugSession(FString::Printf(TEXT("Edge %d\n"), currentEdge->GetId()));
+		}
+		DisplaySegment(EdgeSegment->GetExtemity(ELimit::Start), EdgeSegment->GetExtemity(ELimit::End), EdgeSegment->GetId(), EVisuProperty::GreenCurve);
+	}
+	Close3DDebugSession();
+	//Wait();
+#endif
+}
+
+void FThinZone2DFinder::SearchThinZones()
+{
+#ifdef DEBUG_THIN_ZONES
+	{
+		F3DDebugSession _(TEXT("Thin Surface"));
+		Grid.DisplayPoints("Scaled Grid", Grid.GetInner2DPoints(EGridSpace::Scaled));
+	}
+#endif
+
+	//FMessage::Printf(DBG, TEXT("Searching thin zones on Surface %d\n", Surface->GetId());
+	FTimePoint StartTime = FChrono::Now();
+	BuildLoopSegments();
+	Chronos.BuildLoopSegmentsTime = FChrono::Elapse(StartTime);
+
+	DisplayLoopSegments();
+
+	StartTime = FChrono::Now();
+	FindClosedSegments();
+	Chronos.FindClosedSegmentTime = FChrono::Elapse(StartTime);
+	DisplayClosedSegments();
+
+	StartTime = FChrono::Now();
+	CheckClosedSegments();
+	Chronos.CheckClosedSegmentTime = FChrono::Elapse(StartTime);
+	DisplayClosedSegments();
+
+	StartTime = FChrono::Now();
+	LinkClosedSegments();
+	Chronos.LinkClosedSegmentTime = FChrono::Elapse(StartTime);
+
+	StartTime = FChrono::Now();
+	BuildThinZone();
+	Chronos.BuildThinZoneTime = FChrono::Elapse(StartTime);
+
+	if (ThinZones.Num() > 0)
+	{
+		Grid.GetFace()->SetHasThinZone();
+	}
+
+#ifdef DISPLAY_THIN_ZONES
+	if (ThinZones.Num() > 0)
+	{
+		F3DDebugSession _(TEXT("Zones fines"));
+
+		int32 index = 0;
+		for (const FThinZone2D& Zone : ThinZones)
+		{
+			index++;
+			EVisuProperty prop = PropertyElement;
+			switch (Zone.GetCategory())
+			{
+			default:
+			case FThinZone2D::UNDEFINED:
+				prop = EVisuProperty::Iso;
+				Open3DDebugSession(TEXT("Zone UNDEFINED"));
+				break;
+			case FThinZone2D::GLOBAL:
+				prop = PROPERTY_BORDER_GEO;
+				Open3DDebugSession(TEXT("Zone GLOBAL"));
+				break;
+			case FThinZone2D::PEAK_START:
+				prop = EVisuProperty::BorderEdge;
+				Open3DDebugSession(TEXT("Zone PEAK start"));
+				break;
+			case FThinZone2D::PEAK_END:
+				prop = EVisuProperty::BorderEdge;
+				Open3DDebugSession(TEXT("Zone PEAK end"));
+				break;
+			case FThinZone2D::BUTTERFLY:
+				prop = PropertyElement;
+				Open3DDebugSession(TEXT("Zone BUTTERFLY"));
+				break;
+			case FThinZone2D::BETWEEN_CONTOUR:
+				prop = EVisuProperty::NonManifoldEdge;
+				Open3DDebugSession(TEXT("Zone BETWEEN_CONTOUR"));
+				break;
+			default :
+				Open3DDebugSession(TEXT("Zone Unknown"));
+			}
+
+			{
+				F3DDebugSession _(TEXT("Side1"));
+				int32 ind = 0;
+				for (FEdgeSegment* EdgeSegment : Zone.GetFirstSide().GetSegments())
+				{
+					DisplaySegment(EdgeSegment->GetExtemity(ELimit::Start)(), EdgeSegment->GetExtemity(ELimit::End)(), ++ind, prop);
+				}
+			}
+			{
+				F3DDebugSession _(TEXT("Side2"));
+				int32 ind = 0;
+				for (FEdgeSegment* EdgeSegment : Zone.GetSecondSide().GetSegments())
+				{
+					DisplaySegment(EdgeSegment->GetExtemity(ELimit::Start)(), EdgeSegment->GetExtemity(ELimit::End)(), ++ind, prop);
+				}
+			}
+
+			Close3DDebugSession();
+		}
+	}
+#endif
+
+}
+void FThinZone2DFinder::BuildLoopSegments()
+{
+	double Length = 0;
+	double SegmentLength = Tolerance / 5.;
+
+	const TArray<TSharedPtr<FTopologicalLoop>>& Loops = Grid.GetFace()->GetLoops();
+
+	FirstLoopLength = -1.;
+	for (const TSharedPtr<FTopologicalLoop>& Loop : Loops)
+	{
+		for (const FOrientedEdge& Edge : Loop->GetEdges())
+		{
+			Length += Edge.Entity->Length();
+		}
+
+		if (FirstLoopLength < 0)
+		{
+			FirstLoopLength = Length;
+		}
+	}
+
+	double SegmentNum = 1.2 * Length / SegmentLength;
+	LoopSegments.Empty((int32)SegmentNum);
+
+	TSharedPtr<FTopologicalLoop> OuterLoop = Loops[0];
+	for (TSharedPtr<FTopologicalLoop> Loop : Loops)
+	{
+		const TArray<FOrientedEdge>& Edges = Loop->GetEdges();
+
+		FEdgeSegment* FirstSegment = nullptr;
+		FEdgeSegment* PrecedingSegment = nullptr;
+
+		for (const FOrientedEdge& Edge : Edges)
+		{
+			TArray<double> Coordinates;
+			Edge.Entity->Sample(SegmentLength, Coordinates);
+
+			TArray<FPoint2D> Points;
+			Edge.Entity->Approximate2DPoints(Coordinates, Points);
+
+			TArray<FPoint2D> ScaledPoints;
+			Grid.TransformPoints(EGridSpace::Scaled, Points, ScaledPoints);
+
+			int32 PointNum = Points.Num();
+			for (int32 ISegment = 0; ISegment < PointNum - 1; ISegment++)
+			{
+				int32 Index1 = (Edge.Direction == EOrientation::Front) ? ISegment : PointNum - ISegment - 1l;
+				int32 Index2 = (Edge.Direction == EOrientation::Front) ? ISegment + 1l : PointNum - ISegment - 2l;
+				FEdgeSegment& CurrentSeg = SegmentFatory.New();
+				CurrentSeg.SetBoundarySegment(OuterLoop != Loop, Edge.Entity, Coordinates[Index1], Coordinates[Index2], ScaledPoints[Index1], ScaledPoints[Index2]);
+
+				LoopSegments.Add(&CurrentSeg);
+				if (!FirstSegment)
+				{
+					FirstSegment = PrecedingSegment = &CurrentSeg;
+				}
+				else
+				{
+					PrecedingSegment->SetNext(&CurrentSeg);
+				}
+				PrecedingSegment = &CurrentSeg;
+			}
+		}
+		PrecedingSegment->SetNext(FirstSegment);
+	}
+	SortedLoopSegments = LoopSegments;
+}
+
+FThinZoneSide::FThinZoneSide(FThinZoneSide& InFrontSide, const TArray<FEdgeSegment*>& InSegments, bool bInIsFirstSide)
+	: FrontSide(InFrontSide)
+{
+	Segments = InSegments;
+	Length = 0;
+	for (const FEdgeSegment* Segment : Segments)
+	{
+		Length += Segment->GetLength();
+	}
+}
+
+void FThinZoneSide::SetEdgesAsThinZone()
+{
+	for (FEdgeSegment* Segment : Segments)
+	{
+		Segment->GetEdge()->SetThinZone();
+	}
+}
+
+bool FThinZoneSide::IsPartiallyMeshed() const
+{
+	TSharedPtr<FTopologicalEdge> Edge = nullptr;
+	for (FEdgeSegment* EdgeSegment : Segments)
+	{
+		if (Edge != EdgeSegment->GetEdge())
+		{
+			Edge = EdgeSegment->GetEdge();
+			if (Edge->GetLinkActiveEdge()->IsMeshed())
+			{
+				return true;
+			}
+		}
+	}
+	return false;
+}
+
+double FThinZoneSide::GetMeshedLength() const
+{
+	double LocalLength = 0;
+	TSharedPtr<FTopologicalEdge> Edge = nullptr;
+	for (FEdgeSegment* EdgeSegment : Segments)
+	{
+		if (Edge != EdgeSegment->GetEdge())
+		{
+			Edge = EdgeSegment->GetEdge();
+			if (Edge->GetLinkActiveEdge()->IsMeshed())
+			{
+				LocalLength += Edge->Length();
+			}
+		}
+	}
+	return LocalLength;
+}
+
+void FThinZone2D::SetEdgesAsThinZone()
+{
+	FirstSide.SetEdgesAsThinZone();
+	SecondSide.SetEdgesAsThinZone();
+}
+
diff --git a/Engine/Source/Runtime/Datasmith/CADKernel/Private/CADKernel/Mesh/Structure/ThinZone2D.h b/Engine/Source/Runtime/Datasmith/CADKernel/Private/CADKernel/Mesh/Structure/ThinZone2D.h
new file mode 100644
--- /dev/null
+++ b/Engine/Source/Runtime/Datasmith/CADKernel/Private/CADKernel/Mesh/Structure/ThinZone2D.h
@@ -0,0 +1,303 @@
+// Copyright Epic Games, Inc. All Rights Reserved.
+#pragma once
+
+#include "CADKernel/Core/Chrono.h"
+#include "CADKernel/Core/HaveStates.h"
+#include "CADKernel/Core/Factory.h"
+#include "CADKernel/Core/Types.h"
+#include "CADKernel/Math/Point.h"
+#include "CADKernel/Mesh/Structure/EdgeSegment.h"
+
+//#define DEBUG_THIN_ZONES
+//#define DISPLAY_THIN_ZONES // to display found thin zone
+
+//#define DEBUG_BORDER_THIN_ZONES
+//#define DEBUG_REMOVE_THIN_ZONE
+//#define DEBUG_REMOVABLE
+//#define DEBUG_BEST_VERTEX
+
+namespace CADKernel
+{
+	class FTopologicalLoop;
+	class FTopologicalEdge;
+	class FGrid;
+	class FThinZone2D;
+	class FThinZone2DFinder;
+
+	struct FThinZoneChronos
+	{
+		FDuration BuildLoopSegmentsTime;
+		FDuration DisplayBoundarySegmentTime;
+		FDuration FindClosedSegmentTime;
+		FDuration CheckClosedSegmentTime;
+		FDuration LinkClosedSegmentTime;
+		FDuration BuildThinZoneTime;
+	};
+
+	enum class EThinZone2DType : uint8
+	{
+		Undefined = 0,
+		Global,			// a Surface globally thin
+		PeakStart,		// an extremity of a Surface that is fine
+		PeakEnd,		// an extremity of a Surface that is fine
+		Butterfly,		// the outter loop that is like a bow tie (Butterfly)
+		BetweenLoops    // a bow tie between two different loops
+	};
+
+	class FThinZoneSide : public FHaveStates
+	{
+	private:
+		TArray<FEdgeSegment*> Segments;
+		FThinZoneSide& FrontSide;
+		double Length;
+
+	public:
+
+		FThinZoneSide(FThinZoneSide& InFrontSide, const TArray<FEdgeSegment*>& InSegments, bool bInIsFirstSide);
+		virtual ~FThinZoneSide() = default;
+
+		bool IsClosed() const
+		{
+			if (GetFirst() == nullptr)
+			{
+				return false;
+			}
+			return GetFirst()->GetPrevious() == GetLast();
+		}
+
+		FEdgeSegment* GetFirst() const
+		{
+			return Segments.Num() > 0 ? Segments[0] : nullptr;
+		}
+
+		FEdgeSegment* GetLast() const
+		{
+			return Segments.Num() > 0 ? Segments.Last() : nullptr;
+		}
+
+		void SetEdgesAsThinZone();
+
+		const TArray<FEdgeSegment*>& GetSegments() const
+		{
+			return Segments;
+		}
+
+		bool IsPartiallyMeshed() const;
+		double GetMeshedLength() const;
+
+		double GetLength() const
+		{
+			return Length;
+		}
+
+		bool IsInner() const
+		{
+			if (Segments.Num() == 0)
+			{
+				return true;
+			}
+			return Segments[0]->IsInner();
+		}
+
+		bool IsFirstSide() const
+		{
+			return ((States & EHaveStates::IsFirstSide) == EHaveStates::IsFirstSide);
+		}
+
+		void SetFirstSide() const
+		{
+			States |= EHaveStates::IsFirstSide;
+		}
+
+		void ResetFirstSide() const
+		{
+			States &= ~EHaveStates::IsFirstSide;
+		}
+	};
+
+	class FThinZone2D : public FHaveStates
+	{
+	private:
+
+		FThinZoneSide FirstSide;
+		FThinZoneSide SecondSide;
+
+		EThinZone2DType Category;
+
+		TArray<TSharedPtr<FTopologicalEdge>> PeakEdges[2];
+
+		double Thickness;
+
+	public:
+
+		FThinZone2D(const TArray<FEdgeSegment*>& InFirstSideSegments, bool bInIsClosed1, const TArray<FEdgeSegment*>& InSecondSideSegments, bool bInIsClosed2, double ZoneThickness)
+			: FirstSide(SecondSide, InFirstSideSegments, true)
+			, SecondSide(FirstSide, InSecondSideSegments, false)
+			, Category(EThinZone2DType::Undefined)
+			, Thickness(ZoneThickness)
+		{
+			if (bInIsClosed1)
+			{
+				SetFirstSideClosed();
+			}
+
+			if (bInIsClosed2)
+			{
+				SetFirstSideClosed();
+			}
+		}
+
+		virtual ~FThinZone2D() = default;
+
+		const double GetThickness()
+		{
+			return Thickness;
+		};
+
+		const TArray<TSharedPtr<FTopologicalEdge>>* GetPeakEdges() const
+		{
+			return PeakEdges;
+		}
+
+		TArray<TSharedPtr<FTopologicalEdge>>* GetPeakEdges()
+		{
+			return PeakEdges;
+		}
+
+		const FThinZoneSide& GetFirstSide() const
+		{
+			return FirstSide;
+		}
+
+		const FThinZoneSide& GetSecondSide() const
+		{
+			return SecondSide;
+		}
+
+		EThinZone2DType GetCategory() const
+		{
+			return Category;
+		}
+
+		void SetCategory(EThinZone2DType InType)
+		{
+			Category = InType;
+		}
+
+
+		void SetEdgesAsThinZone();
+
+		double GetMaxLength() const
+		{
+			return FMath::Max(FirstSide.GetLength(), SecondSide.GetLength());
+		}
+
+		bool IsRemoved() const
+		{
+			return ((States & EHaveStates::IsRemoved) == EHaveStates::IsRemoved);
+		}
+
+		void SetRemoved() const
+		{
+			States |= EHaveStates::IsRemoved;
+		}
+
+		void ResetRemoved() const
+		{
+			States &= ~EHaveStates::IsRemoved;
+		}
+
+		bool IsFirstSideClosed() const
+		{
+			return ((States & EHaveStates::FirstSideClosed) == EHaveStates::FirstSideClosed);
+		}
+
+		void SetFirstSideClosed() const
+		{
+			States |= EHaveStates::FirstSideClosed;
+		}
+
+		void ResetFirstSideClosed() const
+		{
+			States &= ~EHaveStates::FirstSideClosed;
+		}
+
+		bool IsSecondSideClosed() const
+		{
+			return ((States & EHaveStates::SecondSideClosed) == EHaveStates::SecondSideClosed);
+		}
+
+		void SetSecondSideClosed() const
+		{
+			States |= EHaveStates::SecondSideClosed;
+		}
+
+		void ResetSecondSideClosed() const
+		{
+			States &= ~EHaveStates::SecondSideClosed;
+		}
+
+		bool HasClosedSide() const 
+		{
+			constexpr EHaveStates ATLeastOneClosed = EHaveStates::FirstSideClosed | EHaveStates::FirstSideClosed;
+			return EnumHasAnyFlags(States, ATLeastOneClosed);
+		}
+
+	};
+
+	class FThinZone2DFinder
+	{
+		friend class FThinZone2D;
+
+	public:
+		FThinZoneChronos Chronos;
+
+	protected:
+		const FGrid& Grid;
+
+		double Tolerance;
+		double SquareTolerance;
+		double MinFeatureSize;
+
+		TArray<FThinZone2D> ThinZones;
+
+		TArray<FEdgeSegment*> LoopSegments;
+		TArray<FEdgeSegment*> SortedLoopSegments;
+
+		double FirstLoopLength = 0.0;
+
+		TFactory<FEdgeSegment> SegmentFatory;
+
+	public:
+
+		FThinZone2DFinder(const FGrid& InGrid)
+			: Grid(InGrid)
+			, SegmentFatory()
+		{
+		}
+
+		void Set(double InTolerance)
+		{
+			Tolerance = InTolerance;
+			SquareTolerance = FMath::Square(Tolerance);
+		}
+
+		void SearchThinZones();
+
+		const TArray<FThinZone2D>& GetThinZones() const
+		{
+			return ThinZones;
+		}
+
+	private:
+
+		void BuildLoopSegments();
+		void FindClosedSegments();
+		void CheckClosedSegments();
+		void LinkClosedSegments();
+		void BuildThinZone();
+
+		void DisplayClosedSegments();
+		void DisplayLoopSegments();
+	};
+}
diff --git a/Engine/Source/Runtime/Datasmith/CADKernel/Private/CADKernel/Topo/Body.cpp b/Engine/Source/Runtime/Datasmith/CADKernel/Private/CADKernel/Topo/Body.cpp
new file mode 100644
--- /dev/null
+++ b/Engine/Source/Runtime/Datasmith/CADKernel/Private/CADKernel/Topo/Body.cpp
@@ -0,0 +1,30 @@
+// Copyright Epic Games, Inc. All Rights Reserved.
+#include "CADKernel/Topo/Body.h"
+
+#include "CADKernel/Topo/Shell.h"
+
+using namespace CADKernel;
+
+
+void FBody::AddShell(TSharedRef<FShell> Shell)
+{
+	Shells.Add(Shell);
+}
+
+TSharedPtr<FEntityGeom> FBody::ApplyMatrix(const FMatrixH& InMatrix) const 
+{
+	TArray<TSharedPtr<FShell>> NewShells;
+	for (TSharedPtr<FShell> Shell : Shells)
+	{
+		NewShells.Add(StaticCastSharedPtr<FShell>(Shell->ApplyMatrix(InMatrix)));
+	}
+	return FEntity::MakeShared<FBody>(NewShells);
+}
+
+#ifdef CADKERNEL_DEV
+FInfoEntity& FBody::GetInfo(FInfoEntity& Info) const
+{
+	return FTopologicalEntity::GetInfo(Info).Add(TEXT("shells"), Shells);
+}
+#endif
+
diff --git a/Engine/Source/Runtime/Datasmith/CADKernel/Private/CADKernel/Topo/Joiner.cpp b/Engine/Source/Runtime/Datasmith/CADKernel/Private/CADKernel/Topo/Joiner.cpp
new file mode 100644
--- /dev/null
+++ b/Engine/Source/Runtime/Datasmith/CADKernel/Private/CADKernel/Topo/Joiner.cpp
@@ -0,0 +1,587 @@
+// Copyright Epic Games, Inc. All Rights Reserved.
+
+#include "CADKernel/Topo/Joiner.h"
+
+#include "CADKernel/Geo/Curves/RestrictionCurve.h"
+#include "CADKernel/Topo/TopologicalEdge.h"
+#include "CADKernel/Topo/TopologicalLink.h"
+#include "CADKernel/Topo/TopologicalFace.h"
+#include "CADKernel/Topo/TopologicalLoop.h"
+#include "CADKernel/Topo/TopologicalVertex.h"
+#include "CADKernel/UI/Display.h"
+#include "CADKernel/UI/Message.h"
+#include "CADKernel/Utils/Util.h"
+
+using namespace CADKernel;
+
+FJoiner::FJoiner(TArray<TSharedPtr<FTopologicalFace>>& InSurfaces, double InTolerance)
+	: Surfaces(InSurfaces)
+	, JoiningTolerance(InTolerance)
+	, JoiningToleranceSqr(FMath::Square(JoiningTolerance))
+{
+}
+
+void FJoiner::JoinSurfaces(bool bProcessOnlyBorderEdges, bool bProcessOnlyNonManifoldEdges)
+{
+	if (!(bProcessOnlyBorderEdges || bProcessOnlyNonManifoldEdges))
+	{
+		JoinVertices();
+		CheckSelfConnectedEdge();
+	}
+	else
+	{
+		TArray<TSharedPtr<FTopologicalVertex>> NewMergeVertices;
+		JoinBorderVertices(NewMergeVertices);
+		JoinEdges(NewMergeVertices);
+	}
+
+	JoinEdges(Vertices);
+	TArray<TSharedPtr<FTopologicalVertex>> VerticesToProcess;
+	MergeUnconnectedAdjacentEdges(VerticesToProcess, true);
+	JoinEdges(VerticesToProcess);
+}
+
+void FJoiner::GetVertices()
+{
+	Vertices.Reserve(10 * Surfaces.Num());
+
+	for (TSharedPtr<FTopologicalFace> Surface : Surfaces)
+	{
+		for (TSharedPtr<FTopologicalLoop> Loop : Surface->GetLoops())
+		{
+			for (const FOrientedEdge& OrientedEdge : Loop->GetEdges())
+			{
+				const TSharedPtr<FTopologicalEdge>& Edge = OrientedEdge.Entity;
+				if (!Edge->GetStartVertex()->GetLinkActiveEntity()->HasMarker1())
+				{
+					Edge->GetStartVertex()->GetLinkActiveEntity()->SetMarker1();
+					Vertices.Add(Edge->GetStartVertex()->GetLinkActiveEntity());
+				}
+				if (!Edge->GetEndVertex()->GetLinkActiveEntity()->HasMarker1())
+				{
+					Edge->GetEndVertex()->GetLinkActiveEntity()->SetMarker1();
+					Vertices.Add(Edge->GetEndVertex()->GetLinkActiveEntity());
+				}
+			}
+		}
+	}
+}
+
+void FJoiner::JoinVertices()
+{
+	GetVertices();
+
+	int32 VertexNum = (int32) Vertices.Num();
+
+	TArray<double> VerticesWeight;
+	VerticesWeight.Reserve(VertexNum);
+	TArray<int32> SortedVertexIndexes;
+	SortedVertexIndexes.Reserve(VertexNum);
+	TArray<TSharedPtr<FTopologicalVertex>> ActiveVertices;
+	ActiveVertices.Reserve(VertexNum);
+
+	for (TSharedPtr<FTopologicalVertex> Vertex: Vertices)
+	{
+		Vertex->ResetMarker1();
+		VerticesWeight.Add(Vertex->GetCoordinates().X + Vertex->GetCoordinates().Y + Vertex->GetCoordinates().Z);
+	}
+	for (int32 Index = 0; Index < Vertices.Num(); ++Index)
+	{
+		SortedVertexIndexes.Add(Index);
+	}
+	SortedVertexIndexes.Sort([&VerticesWeight](const int32& Index1, const int32& Index2) { return VerticesWeight[Index1] < VerticesWeight[Index2]; });
+
+	for (int32 IndexI = 0; IndexI < VertexNum; ++IndexI)
+	{
+		TSharedPtr<FTopologicalVertex> Vertex = Vertices[SortedVertexIndexes[IndexI]];
+		if (Vertex->HasMarker1())
+		{
+			continue;
+		}
+		ActiveVertices.Add(Vertex);
+		Vertex->SetMarker1();
+
+		double VertexWeigth = VerticesWeight[SortedVertexIndexes[IndexI]];
+		FPoint Barycenter = Vertex->GetBarycenter();
+		for (int32 IndexJ = IndexI+1; IndexJ < VertexNum; ++IndexJ)
+		{
+			TSharedPtr<FTopologicalVertex> OtherVertex = Vertices[SortedVertexIndexes[IndexJ]];
+			if (OtherVertex->HasMarker1())
+			{
+				continue;
+			}
+
+			double OtherVertexWeigth = VerticesWeight[SortedVertexIndexes[IndexJ]];
+			if ((OtherVertexWeigth - VertexWeigth) > JoiningTolerance)
+			{
+				break;
+			}
+
+			double DistanceSqr = OtherVertex->GetLinkActiveEntity()->SquareDistance(Barycenter);
+			if (DistanceSqr < JoiningToleranceSqr)
+			{
+				OtherVertex->SetMarker1();
+				Vertex->Link(OtherVertex.ToSharedRef());
+				Barycenter = Vertex->GetBarycenter();
+			}
+		}
+	}
+
+	for (TSharedPtr<FTopologicalVertex> Vertex : Vertices)
+	{
+		Vertex->ResetMarker1();
+	}
+	Swap(ActiveVertices, Vertices);
+}
+
+void FJoiner::JoinBorderVertices(TArray<TSharedPtr<FTopologicalVertex>>& MergedVertices)
+{
+	if (!Vertices.Num())
+	{
+		GetVertices();
+	}
+
+	int32 VertexNum = (int32)Vertices.Num();
+
+	TArray<double> VerticesWeight;
+	VerticesWeight.Reserve(VertexNum);
+	TArray<int32> SortedVertexIndexes;
+	SortedVertexIndexes.Reserve(VertexNum);
+
+	TArray<TSharedPtr<FTopologicalVertex>> ActiveVertices;
+	ActiveVertices.Reserve(VertexNum);
+
+	for (TSharedPtr<FTopologicalVertex> Vertex : Vertices)
+	{
+		Vertex->ResetMarker1();
+		VerticesWeight.Add(Vertex->GetCoordinates().X + Vertex->GetCoordinates().Y + Vertex->GetCoordinates().Z);
+	}
+	for (int32 Index = 0; Index < Vertices.Num(); ++Index)
+	{
+		SortedVertexIndexes.Add(Index);
+	}
+	SortedVertexIndexes.Sort([&VerticesWeight](const int32& Index1, const int32& Index2) { return VerticesWeight[Index1] < VerticesWeight[Index2]; });
+
+	int32 StartIndexJ = 1;
+	for (int32 IndexI = 0; IndexI <VertexNum; ++IndexI)
+	{
+		TSharedRef<FTopologicalVertex> Vertex = Vertices[SortedVertexIndexes[IndexI]].ToSharedRef();
+		if (!Vertex->IsBorderVertex())
+		{
+			continue;
+		}
+
+		double VertexWeigth = VerticesWeight[SortedVertexIndexes[IndexI]];
+		FPoint Barycenter = Vertex->GetBarycenter();
+		for (int32 IndexJ = StartIndexJ; IndexJ < VertexNum; ++IndexJ)
+		{
+			TSharedRef<FTopologicalVertex> OtherVertex = Vertices[SortedVertexIndexes[IndexJ]].ToSharedRef();
+			if (Vertex->GetLink() == OtherVertex->GetLink())
+			{
+				continue;
+			}
+
+			double OtherVertexWeigth = VerticesWeight[SortedVertexIndexes[IndexJ]];
+
+			if (OtherVertexWeigth + JoiningTolerance - VertexWeigth <0)
+			{
+				StartIndexJ = IndexJ;
+				continue;
+			}
+
+			if (OtherVertexWeigth - JoiningTolerance - VertexWeigth > 0)
+			{
+				break;
+			}
+
+			double DistanceSqr = OtherVertex->GetLinkActiveEntity()->SquareDistance(Barycenter);
+			if (DistanceSqr < JoiningToleranceSqr)
+			{
+				TArray<TSharedPtr<FTopologicalEdge>> CommonEdges;
+				Vertex->GetConnectedEdges(OtherVertex, CommonEdges);
+				if (CommonEdges.Num() > 0)
+				{
+					continue;
+				}
+
+				OtherVertex->SetMarker1();
+				Vertex->Link(OtherVertex);
+				Barycenter = Vertex->GetBarycenter();
+				ActiveVertices.Add(Vertex);
+			}
+		}
+	}
+
+	MergedVertices.Reserve(ActiveVertices.Num());
+	for (TSharedPtr<FTopologicalVertex> Vertex : ActiveVertices)
+	{
+		if (Vertex->GetLinkActiveEntity()->HasMarker1())
+		{
+			continue;
+		}
+		MergedVertices.Add(Vertex->GetLinkActiveEntity());
+		MergedVertices.Last()->SetMarker1();
+	}
+
+	for (TSharedPtr<FTopologicalVertex> Vertex : MergedVertices)
+	{
+		Vertex->ResetMarker1();
+	}
+}
+
+// =========================================================================================================================================================================================================
+// =========================================================================================================================================================================================================
+// =========================================================================================================================================================================================================
+//
+//
+//                                                                            NOT YET REVIEWED
+//
+//
+// =========================================================================================================================================================================================================
+// =========================================================================================================================================================================================================
+// =========================================================================================================================================================================================================
+
+
+void FJoiner::CheckSelfConnectedEdge()
+{
+	TArray<TSharedPtr<FTopologicalEdge>> CrunchedEdges;
+	CrunchedEdges.Reserve(100);
+
+	TArray<TSharedPtr<FTopologicalEdge>> IsolatedEdges;
+	IsolatedEdges.Reserve(100);
+
+	for (TSharedPtr<FTopologicalVertex> Vertex : Vertices)
+	{
+		if (Vertex->GetTwinsEntities().Num() > 1)
+		{
+			for (TWeakPtr<FTopologicalVertex> TwinVertex : Vertex->GetTwinsEntities())
+			{
+				TArray<TWeakPtr<FTopologicalEdge>> Edges = TwinVertex.Pin()->GetConnectedEdges();
+				for (TWeakPtr<FTopologicalEdge> WeakEdge : Edges)
+				{
+					TSharedPtr<FTopologicalEdge> Edge = WeakEdge.Pin();
+					if (!Edge->GetLoop().IsValid())
+					{
+						IsolatedEdges.Add(Edge);
+						continue;
+					}
+
+					if (Edge->HasMarker1())
+					{
+						continue;
+					}
+
+					if (Edge->GetStartVertex()->GetLink() == Edge->GetEndVertex()->GetLink())
+					{
+						ensureCADKernel(!Edge->IsDegenerated());
+						CrunchedEdges.Add(Edge);
+						Edge->SetMarker1();
+					}
+				}
+			}
+		}
+	}
+
+	TArray<TSharedPtr<FTopologicalVertex>> NewVertices;
+	NewVertices.Reserve((CrunchedEdges.Num() + IsolatedEdges.Num()) * 2);
+
+	TFunction<void(const TSharedPtr<TTopologicalLink<FTopologicalVertex>>&, const TSharedPtr<FTopologicalVertex>&, const TSharedPtr<FTopologicalVertex>&)> DeleteVertexOfLink = [&NewVertices](const TSharedPtr<TTopologicalLink<FTopologicalVertex>>& OldLink, const TSharedPtr<FTopologicalVertex>& Vertex1, const TSharedPtr<FTopologicalVertex>& Vertex2)
+	{
+		TSharedPtr<FTopologicalVertex> NewVertex = TSharedPtr<FTopologicalVertex>();
+		for (TWeakPtr<FTopologicalVertex> IVertex : OldLink->GetTwinsEntities())
+		{
+			ensureCADKernel(IVertex.IsValid());
+			IVertex.Pin()->ResetTopologicalLink();
+			if (IVertex == Vertex1 || IVertex == Vertex2)
+			{
+				continue;
+			}
+
+			if (!NewVertex)
+			{
+				NewVertex = IVertex.Pin();
+				NewVertices.Add(NewVertex);
+			}
+			else
+			{
+				NewVertex->Link(IVertex.Pin().ToSharedRef());
+			}
+		}
+	};
+
+	for (TSharedPtr<FTopologicalEdge> Edge : IsolatedEdges)
+	{
+		TSharedPtr<FTopologicalVertex> Vertex1 = Edge->GetStartVertex();
+		TSharedPtr<FTopologicalVertex> Vertex2 = Edge->GetEndVertex();
+		if (!Vertex1 && !Vertex2)
+		{
+			continue;
+		}
+
+		ensureCADKernel(!Vertex1 == !Vertex2);
+
+		TSharedPtr<TTopologicalLink<FTopologicalVertex>> OldLink1 = Vertex1->GetLink();
+		TSharedPtr<TTopologicalLink<FTopologicalVertex>> OldLink2 = Vertex2->GetLink();
+
+		Vertex1->RemoveConnectedEdge(Edge.ToSharedRef());
+		Vertex2->RemoveConnectedEdge(Edge.ToSharedRef());
+		if (!Vertex1->GetConnectedEdges().Num())
+		{
+			if (!Vertex2->GetConnectedEdges().Num())
+			{
+				DeleteVertexOfLink(OldLink1, Vertex1, Vertex2);
+				if (OldLink1 != OldLink2)
+				{
+					DeleteVertexOfLink(OldLink2, Vertex1, Vertex2);
+				}
+			}
+			else
+			{
+				DeleteVertexOfLink(OldLink1, Vertex1, nullptr);
+			}
+			continue;
+		}
+
+		if (!Vertex2->GetConnectedEdges().Num())
+		{
+			DeleteVertexOfLink(OldLink2, Vertex2, nullptr);
+		}
+	}
+
+	for(TSharedPtr<FTopologicalEdge> Edge : CrunchedEdges)
+	{
+		Edge->ResetMarkers();
+
+		TSharedPtr<FTopologicalVertex> Vertex1 = Edge->GetStartVertex();
+		TSharedPtr<FTopologicalVertex> Vertex2 = Edge->GetEndVertex();
+		TSharedPtr<TTopologicalLink<FTopologicalVertex>> OldLink = Vertex1->GetLink();
+
+		for (TWeakPtr<FTopologicalVertex> Vertex : OldLink->GetTwinsEntities())
+		{
+			Vertex.Pin()->ResetTopologicalLink();
+		}
+
+		for (TWeakPtr<FTopologicalVertex> Vertex : OldLink->GetTwinsEntities())
+		{
+			double Distance1 = Vertex1->Distance(Vertex.Pin().ToSharedRef());
+			double Distance2 = Vertex2->Distance(Vertex.Pin().ToSharedRef());
+			if (Distance1 <Distance2)
+			{
+				Vertex1->Link(Vertex.Pin().ToSharedRef());
+			}
+			else
+			{
+				Vertex2->Link(Vertex.Pin().ToSharedRef());
+			}
+		}
+
+		NewVertices.Add(Vertex1->GetLinkActiveEntity());
+		NewVertices.Add(Vertex2->GetLinkActiveEntity());
+	}
+
+	TArray<TSharedPtr<FTopologicalVertex>> ActiveVertices;
+	ActiveVertices.Reserve(Vertices.Num() + NewVertices.Num());
+
+	for (TSharedPtr<FTopologicalVertex> Vertex : NewVertices)
+	{
+		if (Vertex->GetLinkActiveEntity()->HasMarker1())
+		{
+			continue;
+		}
+		Vertex->GetLinkActiveEntity()->SetMarker1();
+		ActiveVertices.Add(Vertex->GetLinkActiveEntity());
+	}
+
+	int32 NumOfVertices = (int32)ActiveVertices.Num();
+
+	for (TSharedPtr<FTopologicalVertex> Vertex : Vertices)
+	{
+		if (Vertex->GetLinkActiveEntity()->HasMarker1())
+		{
+			continue;
+		}
+		ActiveVertices.Add(Vertex->GetLinkActiveEntity());
+	}
+
+	for (int32 Index = 0; Index <NumOfVertices; Index++)
+	{
+		ActiveVertices[Index]->ResetMarker1();
+	}
+
+	Swap(ActiveVertices, Vertices);
+}
+
+void FJoiner::FixCollapsedEdges()
+{
+	TArray<TSharedPtr<FTopologicalEdge>> CollapsedEdges;
+
+	int32 NbDegeneratedEdges = 0;
+	for (const TSharedPtr<FTopologicalFace>& Surface : Surfaces)
+	{
+		for (const TSharedPtr<FTopologicalLoop>& Loop : Surface->GetLoops())
+		{
+			for (const FOrientedEdge& OrientedEdge : Loop->GetEdges())
+			{
+				const TSharedPtr<FTopologicalEdge>& Edge = OrientedEdge.Entity;
+				for (const TWeakPtr<FTopologicalEdge>& TwinEdge : Edge->GetTwinsEntities())
+				{
+					TSharedPtr<FTopologicalFace> Face = TwinEdge.Pin()->GetFace();
+					if (!Face.IsValid())
+					{
+						continue;
+					}
+				}
+
+				if (Edge->GetStartVertex()->GetLink() == Edge->GetEndVertex()->GetLink())
+				{
+					if (!Edge->IsDegenerated())
+					{
+						for (const TWeakPtr<FTopologicalEdge>& TwinEdge : Edge->GetTwinsEntities())
+						{
+							TwinEdge.Pin()->SetAsDegenerated();
+						}
+					}
+				}
+			}
+		}
+	}
+
+	for (TSharedPtr<FTopologicalEdge>& Edge : CollapsedEdges)
+	{
+		TSharedPtr<FTopologicalLoop> Loop = Edge->GetLoop();
+		Loop->RemoveEdge(Edge);
+
+		TSharedPtr<FTopologicalVertex> StartVertex = Edge->GetStartVertex();
+		TSharedPtr<FTopologicalVertex> EndVertex = Edge->GetEndVertex();
+
+		Edge->Delete();
+
+		StartVertex->DeleteIfIsolated();
+		EndVertex->DeleteIfIsolated();
+
+		if (Loop->GetEdges().Num() == 0)
+		{
+			TSharedPtr<FTopologicalFace> Face = Loop->GetFace();
+			Face->RemoveLoop(Loop);
+		}
+	}
+}
+
+
+
+void FJoiner::MergeUnconnectedAdjacentEdges(TArray<TSharedPtr<FTopologicalVertex>>& NewEdgeVertices, bool bOnlyIfSameCurve)
+{
+// TODO 
+	ensureCADKernel(false);
+#ifdef TODO
+
+	NewEdgeVertices.Reserve(Vertices.Num());
+	for (TSharedPtr<FTopologicalVertex> Vertex : Vertices)
+	{
+		TArray<TWeakPtr<FTopologicalEdge>> Edges = Vertex->GetConnectedEdges();
+		if (Edges.Num() == 2)
+		{
+			// TODO Work in 2D
+			FPoint Edge0Tangent;
+			Edges[0].Pin()->GetTangentAt(Vertex, Edge0Tangent);
+
+			FPoint Edge1Tangent;
+			Edges[1].Pin()->GetTangentAt(Vertex, Edge1Tangent);
+
+			double CosAngle = Edge0Tangent.ComputeCosinus(Edge1Tangent);
+			if (CosAngle < -0.75)
+			{
+				TArray<EOrientation> EdgeDirections;
+				EdgeDirections.Reserve(2);
+				EdgeDirections.Add(Edges[0].Pin()->GetEndVertex() == Vertex ? EOrientation::Front : EOrientation::Back);
+				EdgeDirections.Add(Edges[1].Pin()->GetStartVertex() == Vertex ? EOrientation::Front : EOrientation::Back);
+				TSharedPtr<FTopologicalEdge> ExtendedEdge = nullptr;
+
+				// TODO Check better if the both curve are similar (2DCurve as to be similar)
+				if (Edges[0].Pin()->GetCurve() != Edges[1].Pin()->GetCurve())
+				{
+					//ExtendedEdge = FEdge::ReplaceEdgesByOne(Edges, EdgeDirections, GetTolerance());
+				}
+				else
+				{
+					//ExtendedEdge = FEdge::CreateEdgeToMerge2Edges(Edges, EdgeDirections);
+				}
+
+				if (ExtendedEdge)
+				{
+					NewEdgeVertices.Add(ExtendedEdge->GetStartVertex());
+					NewEdgeVertices.Add(ExtendedEdge->GetEndVertex());
+				}
+				continue;
+			}
+		}
+	}
+#endif
+}
+
+
+void FJoiner::JoinEdges(TArray<TSharedPtr<FTopologicalVertex>>& VerticesToProcess)
+{
+	ensureCADKernel(false);
+#ifdef TODO
+
+	TArray<TSharedPtr<FTopologicalEdge>> ProcessedEdges;
+	ProcessedEdges.Reserve(VerticesToProcess.Num());
+
+	TArray<TWeakPtr<FTopologicalEdge>> ConnectedEdges;
+	for (TSharedPtr<FTopologicalVertex> Vertex : VerticesToProcess)
+	{
+		ConnectedEdges.Empty();
+		Vertex->GetConnectedEdges(ConnectedEdges);
+		for (int32 EdgeI = 0; EdgeI < ConnectedEdges.Num(); ++EdgeI)
+		{
+			TSharedRef<FTopologicalEdge> Edge = ConnectedEdges[EdgeI].Pin()->GetLinkActiveEdge();
+			if (Edge->HasMarker1())
+			{
+				continue;
+			}
+			TSharedRef<FTopologicalVertex> EndVertex = Edge->GetOtherVertex(Vertex.ToSharedRef())->GetLinkActiveEntity();
+
+			for (int32 EdgeJ = EdgeI; EdgeJ < ConnectedEdges.Num(); ++EdgeJ)
+			{
+				TSharedRef<FTopologicalEdge> OtherEdge = ConnectedEdges[EdgeJ].Pin()->GetLinkActiveEdge();
+				if (OtherEdge == Edge)
+				{
+					continue;
+				}
+				if (OtherEdge->HasMarker1())
+				{
+					continue;
+				}
+				TSharedRef<FTopologicalVertex> OtherEdgeEndVertex = OtherEdge->GetOtherVertex(Vertex.ToSharedRef())->GetLinkActiveEntity();
+
+				if (OtherEdgeEndVertex == EndVertex)
+				{
+					FPoint StartTangentEdge;
+					Edge->GetTangentAt(Vertex, StartTangentEdge);
+					FPoint StartTangentOtherEdge;
+					OtherEdge->GetTangentAt(Vertex, StartTangentOtherEdge);
+
+					double CosAngle = StartTangentEdge.ComputeCosinus(StartTangentOtherEdge);
+					if (CosAngle < 0.9)
+					{
+						continue;
+					}
+
+					if (Edge->GetFace() != OtherEdge->GetFace())
+					{
+						Edge->Link(OtherEdge);
+						Edge->GetLinkActiveEdge()->SetMarker1();
+						ProcessedEdges.Add(Edge->GetLinkActiveEdge());
+					}
+				}
+			}
+		}
+	}
+
+	for (TSharedPtr<FTopologicalEdge> Edge : ProcessedEdges)
+	{
+		Edge->ResetMarker1();
+	}
+#endif
+}
diff --git a/Engine/Source/Runtime/Datasmith/CADKernel/Private/CADKernel/Topo/Joiner.h b/Engine/Source/Runtime/Datasmith/CADKernel/Private/CADKernel/Topo/Joiner.h
new file mode 100644
--- /dev/null
+++ b/Engine/Source/Runtime/Datasmith/CADKernel/Private/CADKernel/Topo/Joiner.h
@@ -0,0 +1,44 @@
+// Copyright Epic Games, Inc. All Rights Reserved.
+#pragma once
+
+#include "CADKernel/Core/Types.h"
+
+namespace CADKernel
+{
+	class FTopologicalFace;
+	class FTopologicalVertex;
+
+	class FJoiner
+	{
+		friend class FThinZone;
+		friend class FThinZoneFinder;
+
+	protected:
+
+		TArray<TSharedPtr<FTopologicalFace>>& Surfaces;
+		TArray<TSharedPtr<FTopologicalVertex>> Vertices;
+
+		double JoiningTolerance;
+		double JoiningToleranceSqr;
+
+	public:
+
+		FJoiner(TArray<TSharedPtr<FTopologicalFace>>& surfaces, double Tolerance);
+
+		void JoinSurfaces(bool bProcessOnlyBorderEdges = false, bool bProcessOnlyNonManifoldEdges = false);
+
+		void FixCollapsedEdges();
+
+	private:
+
+		void GetVertices();
+
+		void JoinVertices();
+		void JoinBorderVertices(TArray<TSharedPtr<FTopologicalVertex>>& MergedVertices);
+		void CheckSelfConnectedEdge();
+
+		void JoinEdges(TArray<TSharedPtr<FTopologicalVertex>>& Vertices);
+		void MergeUnconnectedAdjacentEdges(TArray<TSharedPtr<FTopologicalVertex>>& Vertices, bool bOnlyIfSameCurve);
+	};
+
+} // namespace CADKernel
diff --git a/Engine/Source/Runtime/Datasmith/CADKernel/Private/CADKernel/Topo/Model.cpp b/Engine/Source/Runtime/Datasmith/CADKernel/Private/CADKernel/Topo/Model.cpp
new file mode 100644
--- /dev/null
+++ b/Engine/Source/Runtime/Datasmith/CADKernel/Private/CADKernel/Topo/Model.cpp
@@ -0,0 +1,141 @@
+// Copyright Epic Games, Inc. All Rights Reserved.
+#include "CADKernel/Topo/Model.h"
+
+#include "CADKernel/Core/EntityGeom.h"
+#include "CADKernel/Topo/Body.h"
+#include "CADKernel/Topo/TopologicalFace.h"
+#include "CADKernel/Topo/TopologicalEdge.h"
+#include "CADKernel/Topo/TopologicalVertex.h"
+#include "CADKernel/Topo/TopologicalFace.h"
+
+
+using namespace CADKernel;
+
+
+void FModel::AddEntity(TSharedRef<FTopologicalEntity> Entity)
+{
+	switch (Entity->GetEntityType())
+	{
+	case EEntity::Body:
+		Add(StaticCastSharedRef<FBody>(Entity));
+		break;
+	case EEntity::TopologicalFace:
+		Add(StaticCastSharedRef<FTopologicalFace>(Entity));
+		break;
+	default:
+		break;
+	}
+}
+
+bool FModel::Contains(TSharedPtr<FTopologicalEntity> Entity)
+{
+	switch(Entity->GetEntityType())
+	{
+	case EEntity::Body:
+		return Bodies.Find(StaticCastSharedPtr<FBody>(Entity)) != INDEX_NONE;
+	case EEntity::TopologicalFace:
+		return Faces.Find(StaticCastSharedPtr<FTopologicalFace>(Entity)) != INDEX_NONE;
+	default:
+		return false;
+	}
+	return false;
+}
+
+void FModel::RemoveEntity(TSharedPtr<FTopologicalEntity> Entity)
+{
+	switch (Entity->GetEntityType())
+	{
+	case EEntity::Body:
+		RemoveBody(StaticCastSharedPtr<FBody>(Entity));
+		break;
+	case EEntity::TopologicalFace:
+		RemoveDomain(StaticCastSharedPtr<FTopologicalFace>(Entity));
+		break;
+	default:
+		break;
+	}
+}
+
+void FModel::RemoveDomain(TSharedPtr<FTopologicalFace> Domain)
+{
+	Faces.Remove(Domain);
+}
+
+void FModel::RemoveBody(TSharedPtr<FBody> Body)
+{
+	Bodies.Remove(Body);
+}
+
+TSharedPtr<FEntityGeom> FModel::ApplyMatrix(const FMatrixH& InMatrix) const
+{
+	TSharedPtr<FModel> Model = FEntity::MakeShared<FModel>();
+
+	for (TSharedPtr<FBody> Body : Bodies)
+	{
+		TSharedPtr<FBody> TransformedBody = StaticCastSharedPtr<FBody>(Body->ApplyMatrix(InMatrix));
+		Model->Add(TransformedBody);
+	}
+
+	for (TSharedPtr<FTopologicalFace> Domain : Faces)
+	{
+		TSharedPtr<FTopologicalFace> TransformedSurface = StaticCastSharedPtr<FTopologicalFace>(Domain->ApplyMatrix(InMatrix));
+		Model->Add(TransformedSurface);
+	}
+
+	return Model;
+}
+
+int32 FModel::FaceCount() const
+{
+	int32 FaceCount = 0;
+	for (const TSharedPtr<FBody>& Body : Bodies)
+	{
+		FaceCount += Body->FaceCount();
+	}
+	FaceCount += Faces.Num();
+	return FaceCount;
+}
+
+void FModel::GetFaces(TArray<TSharedPtr<FTopologicalFace>>& OutFaces) 
+{
+	for (const TSharedPtr<FBody>& Body : Bodies)
+	{
+		Body->GetFaces(OutFaces);
+	}
+
+	for (TSharedPtr<FTopologicalFace>& Face : Faces)
+	{
+		if (!Face->HasMarker1())
+		{
+			OutFaces.Emplace(Face);
+			Face->SetMarker1();
+		}
+	}
+}
+
+void FModel::SpreadBodyOrientation()
+{
+	for (TSharedPtr<FBody>& Body : Bodies)
+	{
+		Body->SpreadBodyOrientation();
+	}
+}
+
+#ifdef CADKERNEL_DEV
+FInfoEntity& FModel::GetInfo(FInfoEntity& Info) const
+{
+	return FEntityGeom::GetInfo(Info)
+		.Add(TEXT("Bodies"), Bodies)
+		.Add(TEXT("Domains"), Faces);
+}
+#endif
+
+void FModel::Add(TSharedPtr<FTopologicalFace> Face)
+{
+	Faces.Add(Face);
+}
+
+void FModel::Add(TSharedPtr<FBody> Body)
+{
+	Bodies.Add(Body);
+}
diff --git a/Engine/Source/Runtime/Datasmith/CADKernel/Private/CADKernel/Topo/Shell.cpp b/Engine/Source/Runtime/Datasmith/CADKernel/Private/CADKernel/Topo/Shell.cpp
new file mode 100644
--- /dev/null
+++ b/Engine/Source/Runtime/Datasmith/CADKernel/Private/CADKernel/Topo/Shell.cpp
@@ -0,0 +1,100 @@
+// Copyright Epic Games, Inc. All Rights Reserved.
+#include "CADKernel/Topo/Shell.h"
+
+#include "CADKernel/Geo/GeoEnum.h"
+#include "CADKernel/Topo/Body.h"
+#include "CADKernel/Topo/TopologicalFace.h"
+
+using namespace CADKernel;
+
+FShell::FShell(const TArray<TSharedPtr<FTopologicalFace>>& InTopologicalFaces, bool bIsInnerShell)
+	: FTopologicalEntity()
+	, TopologicalFaces()
+{
+	TopologicalFaces.Reserve(InTopologicalFaces.Num());
+
+	TArray<EOrientation> Orientations;
+	Orientations.Reserve(InTopologicalFaces.Num());
+
+	for (TSharedPtr<FTopologicalFace> Face : InTopologicalFaces)
+	{
+		TopologicalFaces.Emplace(Face, EOrientation::Front);
+	}
+
+	if (bIsInnerShell)
+	{
+		SetInner();
+	}
+}
+
+
+FShell::FShell(const TArray<TSharedPtr<FTopologicalFace>>& InTopologicalFaces, const TArray<EOrientation>& InOrientations, bool bIsInnerShell)
+	: FTopologicalEntity()
+{
+	ensureCADKernel(InTopologicalFaces.Num() == InOrientations.Num());
+	for (int32 Index = 0; Index < InTopologicalFaces.Num(); ++Index)
+	{
+		TSharedPtr<FTopologicalFace> Face = InTopologicalFaces[Index];
+		EOrientation Orientation = InOrientations[Index];
+		TopologicalFaces.Emplace(Face, Orientation);
+	}
+
+	if (bIsInnerShell)
+	{
+		SetInner();
+	}
+}
+
+TSharedPtr<FEntityGeom> FShell::ApplyMatrix(const FMatrixH& InMatrix) const
+{
+	ensureCADKernel(false);
+	return TSharedPtr<FEntityGeom>();
+}
+
+void FShell::Add(TSharedRef<FTopologicalFace> InTopologicalFace, EOrientation Orientation)
+{
+	TSharedPtr<FTopologicalFace> Face = InTopologicalFace;
+	TopologicalFaces.Emplace(Face, Orientation);
+}
+
+#ifdef CADKERNEL_DEV
+FInfoEntity& FShell::GetInfo(FInfoEntity& Info) const
+{
+	return FEntity::GetInfo(Info)
+		.Add(TEXT("TopologicalFaces"), (TArray<TOrientedEntity<FEntity>>&) TopologicalFaces);
+}
+#endif
+
+void FShell::GetFaces(TArray<TSharedPtr<FTopologicalFace>>& Faces)
+{
+	for (FOrientedFace& Face : TopologicalFaces)
+	{
+		if (Face.Entity->HasMarker1())
+		{
+			continue;
+		}
+
+		Faces.Emplace(Face.Entity);
+		Face.Entity->SetMarker1();
+	}
+}
+
+void FShell::SpreadBodyOrientation()
+{
+	bool bIsOutter = IsOutter();
+	for (FOrientedFace& Face : TopologicalFaces)
+	{
+		if (Face.Entity->HasMarker2())
+		{
+			// the face is already processed, this should not append...
+			continue;
+		}
+
+		if (bIsOutter != (Face.Direction == EOrientation::Front))
+		{
+			Face.Entity->SetBackOriented();
+		}
+		Face.Entity->SetMarker2();
+	}
+}
+
diff --git a/Engine/Source/Runtime/Datasmith/CADKernel/Private/CADKernel/Topo/TopologicalEdge.cpp b/Engine/Source/Runtime/Datasmith/CADKernel/Private/CADKernel/Topo/TopologicalEdge.cpp
new file mode 100644
--- /dev/null
+++ b/Engine/Source/Runtime/Datasmith/CADKernel/Private/CADKernel/Topo/TopologicalEdge.cpp
@@ -0,0 +1,590 @@
+// Copyright Epic Games, Inc. All Rights Reserved.
+
+#include "CADKernel/Topo/TopologicalEdge.h"
+
+#include "CADKernel/Core/KernelParameters.h"
+#include "CADKernel/Geo/Sampling/PolylineTools.h"
+#include "CADKernel/Geo/Sampler/SamplerOnParam.h"
+#include "CADKernel/Math/Point.h"
+#include "CADKernel/Math/SlopeUtils.h"
+#include "CADKernel/Mesh/Structure/EdgeMesh.h"
+#include "CADKernel/Topo/TopologicalLoop.h"
+#include "CADKernel/Topo/TopologicalLink.h"
+#include "CADKernel/Topo/TopologicalFace.h"
+#include "CADKernel/Topo/TopologicalVertex.h"
+
+using namespace CADKernel;
+
+FTopologicalEdge::FTopologicalEdge(const TSharedRef<FRestrictionCurve>& InCurve, const TSharedRef<FTopologicalVertex>& InVertex1, const TSharedRef<FTopologicalVertex>& InVertex2, const FLinearBoundary& InBoundary)
+	: TLinkable<FTopologicalEdge, FEdgeLink>()
+	, StartVertex(InVertex1)
+	, EndVertex(InVertex2)
+	, Boundary(InBoundary)
+	, Curve(InCurve)
+	, Length3D(-1)
+	, Loop(TSharedPtr<FTopologicalLoop>())
+	, Mesh(TSharedPtr<FEdgeMesh>())
+{
+	ensureCADKernel(Boundary.IsValid());
+}
+
+FTopologicalEdge::FTopologicalEdge(const TSharedRef<FRestrictionCurve>& InCurve, const TSharedRef<FTopologicalVertex>& InVertex1, const TSharedRef<FTopologicalVertex>& InVertex2)
+	: TLinkable<FTopologicalEdge, FEdgeLink>()
+	, StartVertex(InVertex1)
+	, EndVertex(InVertex2)
+	, Curve(InCurve)
+	, Length3D(-1)
+	, Loop(TSharedPtr<FTopologicalLoop>())
+	, Mesh(TSharedPtr<FEdgeMesh>())
+{
+	Boundary = Curve->GetBoundary();
+	ensureCADKernel(Boundary.IsValid());
+}
+
+FTopologicalEdge::FTopologicalEdge(const TSharedRef<FRestrictionCurve>& InCurve, const FLinearBoundary& InBoundary)
+	: TLinkable<FTopologicalEdge, FEdgeLink>()
+	, Boundary(InBoundary)
+	, Curve(InCurve)
+	, Length3D(-1)
+	, Loop(TSharedPtr<FTopologicalLoop>())
+	, Mesh(TSharedPtr<FEdgeMesh>())
+{
+	TArray<double> Coordinates = { Boundary.Min, Boundary.Max };
+	TArray<FCurvePoint> Points;
+	Curve->EvaluatePoints(Coordinates, Points);
+
+	StartVertex = FEntity::MakeShared<FTopologicalVertex>(Points[0].Point);
+	EndVertex = FEntity::MakeShared<FTopologicalVertex>(Points[1].Point);
+}
+
+FTopologicalEdge::FTopologicalEdge(const TSharedRef<FRestrictionCurve>& InCurve)
+	: FTopologicalEdge(InCurve, InCurve->GetBoundary())
+{
+}
+
+void FTopologicalEdge::LinkVertex()
+{
+	StartVertex->AddConnectedEdge(StaticCastSharedRef<FTopologicalEdge>(AsShared()));
+	EndVertex->AddConnectedEdge(StaticCastSharedRef<FTopologicalEdge>(AsShared()));
+
+	if (IsDegenerated())
+	{
+		StartVertex->Link(EndVertex.ToSharedRef());
+	}
+}
+
+bool FTopologicalEdge::CheckVertices()
+{
+	TArray<double> Coordinates = { Boundary.Min, Boundary.Max };
+	TArray<FPoint> Points;
+	Curve->Approximate3DPoints(Coordinates, Points);
+
+	double ToleranceGeo = GetTolerance3D();
+
+	TFunction<bool(TSharedPtr<FTopologicalVertex>, FPoint)> CheckExtremityGap = [&](TSharedPtr<FTopologicalVertex> Vertex, FPoint Point)
+	{
+		double GapToVertex = Vertex->GetCoordinates().Distance(Point);
+		return (GapToVertex < ToleranceGeo);
+	};
+
+	if (!CheckExtremityGap(StartVertex, Points[0]))
+	{
+		if (CheckExtremityGap(StartVertex, Points[1]) && CheckExtremityGap(EndVertex, Points[0]))
+		{
+			Swap(StartVertex, EndVertex);
+			return true;
+		}
+		return false;
+	}
+	return CheckExtremityGap(EndVertex, Points[1]);
+}
+
+double FTopologicalEdge::GetTolerance3D()
+{
+	return GetLoop()->GetFace()->GetCarrierSurface()->Get3DTolerance();
+}
+
+bool FTopologicalEdge::CheckIfDegenerated()
+{
+	bool bDegeneration2D = false;
+	bool bDegeneration3D = false;
+
+	Curve->CheckIfDegenerated(Boundary, bDegeneration2D, bDegeneration3D, Length3D);
+
+	if (bDegeneration3D)
+	{
+		SetAsDegenerated();
+	}
+
+	return bDegeneration2D;
+}
+
+TSharedPtr<FTopologicalEdge> FTopologicalEdge::Make(const TSharedRef<FRestrictionCurve>& InCurve, const TSharedRef<FTopologicalVertex>& InVertex1, const TSharedRef<FTopologicalVertex>& InVertex2, const FLinearBoundary& InBoundary)
+{
+	TSharedRef<FTopologicalEdge> Edge = FEntity::MakeShared<FTopologicalEdge>(InCurve, InVertex1, InVertex2, InBoundary);
+	if (Edge->CheckIfDegenerated())
+	{
+		return TSharedPtr<FTopologicalEdge>();
+	}
+
+	if (!Edge->CheckVertices())
+	{
+		return TSharedPtr<FTopologicalEdge>();
+	}
+
+	Edge->LinkVertex();
+	return Edge;
+}
+
+TSharedPtr<FTopologicalEdge> FTopologicalEdge::Make(const TSharedRef<FRestrictionCurve>& InCurve, const TSharedRef<FTopologicalVertex>& InVertex1, const TSharedRef<FTopologicalVertex>& InVertex2)
+{
+	TSharedRef<FTopologicalEdge> Edge = FEntity::MakeShared<FTopologicalEdge>(InCurve, InVertex1, InVertex2);
+	if (Edge->CheckIfDegenerated())
+	{
+		return TSharedPtr<FTopologicalEdge>();
+	}
+
+	if (!Edge->CheckVertices())
+	{
+		return TSharedPtr<FTopologicalEdge>();
+	}
+
+	Edge->LinkVertex();
+	return Edge;
+}
+
+TSharedPtr<FTopologicalEdge> FTopologicalEdge::Make(const TSharedRef<FRestrictionCurve>& InCurve, const FLinearBoundary& InBoundary)
+{
+	TSharedRef<FTopologicalEdge> Edge = FEntity::MakeShared<FTopologicalEdge>(InCurve, InBoundary);
+	if (Edge->CheckIfDegenerated())
+	{
+		return TSharedPtr<FTopologicalEdge>();
+	}
+
+	Edge->LinkVertex();
+	return Edge;
+}
+
+TSharedPtr<FTopologicalEdge> FTopologicalEdge::Make(const TSharedRef<FRestrictionCurve>& InCurve)
+{
+	TSharedRef<FTopologicalEdge> Edge = FEntity::MakeShared<FTopologicalEdge>(InCurve);
+	if (Edge->CheckIfDegenerated())
+	{
+		return TSharedPtr<FTopologicalEdge>();
+	}
+
+	Edge->LinkVertex();
+	return Edge;
+}
+
+void FTopologicalEdge::Link(const TSharedRef<FTopologicalEdge>& Twin)
+{
+	// Degenerated twin edges are not linked
+	if (IsDegenerated() && Twin->IsDegenerated())
+	{
+		return;
+	}
+
+	ensureCADKernel(!IsDegenerated() || !Twin->IsDegenerated());
+
+	FPoint Edge1Vertex1 = GetStartVertex()->GetBarycenter();
+	FPoint Edge1Vertex2 = GetEndVertex()->GetBarycenter();
+	FPoint Edge2Vertex1 = Twin->GetStartVertex()->GetBarycenter();
+	FPoint Edge2Vertex2 = Twin->GetEndVertex()->GetBarycenter();
+
+	// Define the orientation
+	double DistanceSameOrientation = Edge1Vertex1.Distance(Edge2Vertex1) + Edge1Vertex2.Distance(Edge2Vertex2);
+	double DistanceReverseOrientation = Edge1Vertex1.Distance(Edge2Vertex2) + Edge1Vertex2.Distance(Edge2Vertex1);
+	if (DistanceSameOrientation < DistanceReverseOrientation)
+	{
+		if (DistanceSameOrientation > 0.1)
+		{
+			FMessage::Printf(Log, TEXT("Edge %d and Edge %d are to far to be connected"), Id, Twin->GetId());
+			return;
+		}
+		GetStartVertex()->Link(Twin->GetStartVertex());
+		GetEndVertex()->Link(Twin->GetEndVertex());
+	}
+	else
+	{
+		if (DistanceReverseOrientation > 0.1)
+		{
+			FMessage::Printf(Log, TEXT("Edge %d and Edge %d are to far to be connected"), Id, Twin->GetId());
+			return;
+		}
+		GetStartVertex()->Link(Twin->GetEndVertex());
+		GetEndVertex()->Link(Twin->GetStartVertex());
+	}
+	MakeLink(Twin);
+}
+
+void FTopologicalEdge::Delete()
+{
+	StartVertex->RemoveConnectedEdge(StaticCastSharedRef<FTopologicalEdge>(AsShared()));
+	EndVertex->RemoveConnectedEdge(StaticCastSharedRef<FTopologicalEdge>(AsShared()));
+
+	StartVertex.Reset();
+	EndVertex.Reset();
+
+    Curve.Reset();
+	Loop.Reset();
+	Mesh.Reset();
+	SetDeleted();
+}
+
+TSharedRef<FTopologicalFace> FTopologicalEdge::GetFace() const
+{
+	ensureCADKernel(Loop.IsValid());
+	return Loop.Pin()->GetFace();
+}
+
+void FTopologicalEdge::ComputeCrossingPointCoordinates()
+{
+	TSharedRef<FTopologicalEdge> ActiveEdge = GetLinkActiveEdge();
+	ensureCADKernel(this == &ActiveEdge.Get());
+
+	double Tolerance = GetTolerance3D();
+
+	FSurfacicPolyline Presampling;
+	FSurfacicCurveSamplerOnParam Sampler(*Curve.ToSharedRef(), Boundary, Tolerance * 10., Tolerance, Presampling);
+	Sampler.Sample();
+
+	TArray<double>& TabCrossingPointU = ActiveEdge->GetCrossingPointUs();
+	Presampling.SwapCoordinates(TabCrossingPointU);
+}
+
+void FTopologicalEdge::SetStartVertex(const double NewCoordinate)
+{
+	ensureCADKernel(Curve->GetUMax() > NewCoordinate);
+	Boundary.SetMin(NewCoordinate);
+	FCurvePoint OutPoint;
+	Curve->EvaluatePoint(NewCoordinate, OutPoint);
+	StartVertex->SetCoordinates(OutPoint.Point);
+}
+
+void FTopologicalEdge::SetEndVertex(const double NewCoordinate)
+{
+	ensureCADKernel(Curve->GetUMin() < NewCoordinate);
+	Boundary.SetMax(NewCoordinate);
+	FCurvePoint OutPoint;
+	Curve->EvaluatePoint(NewCoordinate, OutPoint);
+	EndVertex->SetCoordinates(OutPoint.Point);
+}
+
+void FTopologicalEdge::SetStartVertex(const double NewCoordinate, const FPoint& NewPoint3D)
+{
+	ensureCADKernel(Curve->GetUMin() < NewCoordinate);
+	Boundary.SetMin(NewCoordinate);
+	StartVertex->SetCoordinates(NewPoint3D);
+}
+
+void FTopologicalEdge::SetEndVertex(const double NewCoordinate, const FPoint& NewPoint3D)
+{
+	ensureCADKernel(Curve->GetUMax() > NewCoordinate);
+	Boundary.SetMax(NewCoordinate);
+	EndVertex->SetCoordinates(NewPoint3D);
+}
+
+double FTopologicalEdge::Length() const
+{
+	if (Length3D < 0)
+	{
+		Length3D = Curve->ComputeLength(Boundary);
+	}
+	return Length3D;
+}
+
+void FTopologicalEdge::GetTangentsAtExtremities(FPoint& StartTangent, FPoint& EndTangent, bool bForward) const
+{
+	ensureCADKernel(Curve->Polyline.Size());
+
+	FDichotomyFinder Finder(Curve->Polyline.GetCoordinates());
+	int32 StartIndex = Finder.Find(Boundary.Min);
+	int32 EndIndex = Finder.Find(Boundary.Max);
+
+	const TArray<FPoint>& Polyline3D = Curve->Polyline.GetPoints();
+	if(bForward)
+	{
+		StartTangent = Polyline3D[StartIndex + 1] - Polyline3D[StartIndex];
+		EndTangent = Polyline3D[EndIndex] - Polyline3D[EndIndex + 1];
+	}
+	else
+	{
+		EndTangent = Polyline3D[StartIndex + 1] - Polyline3D[StartIndex];
+		StartTangent = Polyline3D[EndIndex] - Polyline3D[EndIndex + 1];
+	}
+}
+
+
+void FTopologicalEdge::Sample(const double DesiredSegmentLength, TArray<double>& OutCoordinates) const
+{
+	Curve->Sample(Boundary, DesiredSegmentLength, OutCoordinates);
+}
+
+int32 FTopologicalEdge::EvaluateCuttingPointNum()
+{
+	double Num = 0;
+	for (int32 Index = 0; Index < CrossingPointUs.Num() - 1; Index++)
+	{
+		Num += ((CrossingPointUs[Index + 1] - CrossingPointUs[Index]) / CrossingPointDeltaUMaxs[Index]);
+	}
+	Num *= 1.5;
+	return (int32) Num;
+}
+
+double FTopologicalEdge::TransformLocalCoordinateToActiveEdgeCoordinate(const double InLocalCoordinate)
+{
+	if (IsActiveEntity())
+	{
+		return InLocalCoordinate;
+	}
+
+	TSharedPtr<FTopologicalEdge> ActiveEdge = GetLinkActiveEntity();
+	FPoint PointOnEdge = Curve->Approximate3DPoint(InLocalCoordinate);
+	FPoint ProjectedPoint;
+	return ActiveEdge->GetCurve()->GetCoordinateOfProjectedPoint(Boundary, PointOnEdge, ProjectedPoint);
+}
+
+double FTopologicalEdge::TransformActiveEdgeCoordinateToLocalCoordinate(const double InActiveEdgeCoordinate)
+{
+	if (IsActiveEntity())
+	{
+		return InActiveEdgeCoordinate;
+	}
+
+ 	TSharedPtr<FTopologicalEdge> ActiveEdge = GetLinkActiveEntity();
+	FPoint PointOnEdge = ActiveEdge->GetCurve()->Approximate3DPoint(InActiveEdgeCoordinate);
+	FPoint ProjectedPoint;
+	return Curve->GetCoordinateOfProjectedPoint(Boundary, PointOnEdge, ProjectedPoint);
+}
+
+void FTopologicalEdge::TransformLocalCoordinatesToActiveEdgeCoordinates(const TArray<double>& InLocalCoordinate, TArray<double>& OutActiveEdgeCoordinates)
+{
+	if (IsActiveEntity())
+	{
+		OutActiveEdgeCoordinates = InLocalCoordinate;
+	}
+	TSharedPtr<FTopologicalEdge> ActiveEdge = GetLinkActiveEntity();
+	TArray<FPoint> EdgePoints;
+	Curve->Approximate3DPoints(InLocalCoordinate, EdgePoints);
+	TArray<FPoint> ProjectedPoints;
+	ActiveEdge->GetCurve()->ProjectPoints(Boundary, EdgePoints, OutActiveEdgeCoordinates, ProjectedPoints);
+}
+
+void FTopologicalEdge::TransformActiveEdgeCoordinatesToLocalCoordinates(const TArray<double>& InActiveEdgeCoordinate, TArray<double>& OutLocalCoordinates)
+{
+	if (IsActiveEntity())
+	{
+		OutLocalCoordinates = InActiveEdgeCoordinate;
+	}
+ 
+	TSharedPtr<FTopologicalEdge> ActiveEdge = GetLinkActiveEntity();
+	TArray<FPoint> ActiveEdgePoint;
+	ActiveEdge->GetCurve()->Approximate3DPoints(InActiveEdgeCoordinate, ActiveEdgePoint);
+	TArray<FPoint> ProjectedPoints;
+	Curve->ProjectPoints(Boundary, ActiveEdgePoint, OutLocalCoordinates, ProjectedPoints);
+}
+
+void FTopologicalEdge::AddImposedCuttingPointU(const double ImposedCuttingPointU, int32 OppositeNodeIndex)
+{
+	if (!IsActiveEntity())
+	{
+		ensureCADKernel(false);
+		FPoint Point = Curve->Approximate3DPoint(ImposedCuttingPointU);
+		FPoint ProjectedPoint;
+		double ActiveEdgeParamU = GetLinkActiveEntity()->ProjectPoint(Point, ProjectedPoint);
+		return GetLinkActiveEntity()->AddImposedCuttingPointU(ActiveEdgeParamU, OppositeNodeIndex);
+	}
+
+	ImposedCuttingPointUs.Emplace(ImposedCuttingPointU, OppositeNodeIndex);
+}
+
+void FTopologicalEdge::ProjectTwinEdgePointsOn2DCurve(const TSharedRef<FTopologicalEdge>& InTwinEdge, const TArray<double>& InTwinEdgePointCoords, TArray<FPoint2D>& OutPoints2D)
+{
+	if (&InTwinEdge.Get() == this)
+	{
+		Curve->Approximate2DPoints(InTwinEdgePointCoords, OutPoints2D);
+	}
+	else
+	{
+		TArray<FPoint> Points3D;
+		InTwinEdge->ApproximatePoints(InTwinEdgePointCoords, Points3D);
+
+		bool bSameDirection = IsSameDirection(InTwinEdge);
+		TArray<double> Coordinates;
+		Curve->ProjectTwinCurvePoints(Points3D, bSameDirection, Coordinates);
+		Curve->Approximate2DPoints(Coordinates, OutPoints2D);
+	}
+}
+
+bool FTopologicalEdge::IsSameDirection(const TSharedPtr<FTopologicalEdge>& Edge) const
+{
+	if (TopologicalLink != Edge->GetLink())
+	{
+		return true;
+	}
+
+	if (Edge == AsShared())
+	{
+		return true;
+	}
+
+	TSharedPtr<const FVertexLink> vertex1Edge = GetStartVertex()->GetLink();
+	TSharedPtr<const FVertexLink> vertex2Edge = GetEndVertex()->GetLink();
+
+	if (vertex1Edge == vertex2Edge)
+	{
+		if (Edge->IsDegenerated())
+		{
+			return true;
+		}
+
+		// TODO
+		ensureCADKernel(false);
+	}
+
+	return vertex1Edge == Edge->GetStartVertex()->GetLink();
+}
+
+TSharedPtr<FEntityGeom> FTopologicalEdge::ApplyMatrix(const FMatrixH& InMatrix) const
+{
+	TSharedPtr<FTopologicalVertex> v1Transformed = StaticCastSharedPtr<FTopologicalVertex>(StartVertex->ApplyMatrix(InMatrix));
+	if (!v1Transformed.IsValid()) 
+	{
+		return TSharedPtr<FEntityGeom>();
+	}
+
+	TSharedPtr<FTopologicalVertex> v2Transformed = StaticCastSharedPtr<FTopologicalVertex>(EndVertex->ApplyMatrix(InMatrix));
+	if (!v2Transformed.IsValid())
+	{
+		return TSharedPtr<FEntityGeom>();
+	}
+
+	TSharedPtr<FRestrictionCurve> TransformedCurve = StaticCastSharedPtr<FRestrictionCurve>(Curve->ApplyMatrix(InMatrix));
+	if (!TransformedCurve.IsValid())
+	{
+		return TSharedPtr<FEntityGeom>();
+	}
+
+	return FTopologicalEdge::Make(TransformedCurve.ToSharedRef(), v1Transformed.ToSharedRef(), v2Transformed.ToSharedRef(), Boundary);
+}
+
+#ifdef CADKERNEL_DEV
+FInfoEntity& FTopologicalEdge::GetInfo(FInfoEntity& Info) const
+{
+	return FTopologicalEntity::GetInfo(Info)
+		.Add(TEXT("Link"), TopologicalLink)
+		.Add(TEXT("Curve"), Curve)
+		.Add(TEXT("Vertex1"), StartVertex)
+		.Add(TEXT("Vertex2"), EndVertex)
+		.Add(TEXT("Boundary"), Boundary)
+		.Add(TEXT("Loop"), Loop)
+		.Add(TEXT("Length"), Length())
+		.Add(TEXT("Mesh"), Mesh);
+}
+#endif
+
+TSharedRef<FEdgeMesh> FTopologicalEdge::GetOrCreateMesh(const TSharedRef<FModelMesh>& ShellMesh)
+{
+	if (!IsActiveEntity())
+	{
+		return GetLinkActiveEdge()->GetOrCreateMesh(ShellMesh);
+	}
+
+	if (!Mesh)
+	{
+		Mesh = FEntity::MakeShared<FEdgeMesh>(ShellMesh, StaticCastSharedRef<FTopologicalEdge>(AsShared()));
+	}
+	return Mesh.ToSharedRef();
+}
+
+void FTopologicalEdge::ChooseFinalDeltaUs()
+{
+	for (int32 Index = 0; Index < CrossingPointDeltaUMins.Num(); ++Index)
+	{
+		if (CrossingPointDeltaUMins[Index] > CrossingPointDeltaUMaxs[Index])
+		{
+			CrossingPointDeltaUMaxs[Index] = CrossingPointDeltaUMins[Index];
+		}
+	}
+}
+
+TSharedPtr<FTopologicalEdge> FTopologicalEdge::CreateEdgeToMerge2Edges(const TArray<TWeakPtr<FTopologicalEdge>>& Edges, const TArray<EOrientation>& EdgeDirections)
+{
+	ensureCADKernel(false);
+	TSharedPtr<FTopologicalEdge> Edge0 = Edges[0].Pin();
+	TSharedPtr<FTopologicalEdge> Edge1 = Edges[1].Pin();
+
+	TSharedPtr<FTopologicalVertex> Vertex1 = EdgeDirections[0] == EOrientation::Front ? Edge0->GetStartVertex() : Edge0->GetEndVertex();
+	double Coordinate1 = EdgeDirections[0] == EOrientation::Front ? Edge0->GetStartCurvilinearCoordinates() : Edge0->GetEndCurvilinearCoordinates();
+	TSharedPtr<FTopologicalVertex> Vertex2 = EdgeDirections[1] == EOrientation::Front ? Edge1->GetEndVertex() : Edge1->GetStartVertex();
+	double Coordinate2 = EdgeDirections[1] == EOrientation::Front ? Edge1->GetEndCurvilinearCoordinates() : Edge1->GetStartCurvilinearCoordinates();
+	TSharedPtr<FTopologicalVertex> DeletedVertex = EdgeDirections[1] == EOrientation::Front ? Edge1->GetStartVertex() : Edge1->GetEndVertex();
+	if (Coordinate1 > Coordinate2)
+	{
+		Swap(Vertex1, Vertex2);
+		Swap(Coordinate1, Coordinate2);
+	}
+
+	Edge0->SetStartVertex(Coordinate1, Vertex1->GetCoordinates());
+	Edge0->SetEndVertex(Coordinate2, Vertex2->GetCoordinates());
+
+	DeletedVertex->RemoveConnectedEdge(Edge0.ToSharedRef());
+	DeletedVertex->RemoveConnectedEdge(Edge1.ToSharedRef());
+
+	Edge0->GetLoop()->RemoveEdge(Edge1);
+
+	return Edge0;
+}
+
+void FTopologicalEdge::ComputeEdge2DProperties(FEdge2DProperties& EdgeCharacteristics)
+{
+	const TArray<FPoint2D>& Polyline2D = Curve->Polyline.Get2DPoints();
+	const TArray<FPoint>& Polyline3D = Curve->Polyline.GetPoints();
+	const TArray<double>& Parameters = Curve->Polyline.GetCoordinates();
+
+	FDichotomyFinder Finder(Curve->Polyline.GetCoordinates());
+	int32 StartIndex = Finder.Find(Boundary.Min);
+	int32 EndIndex = Finder.Find(Boundary.Max);
+
+	for (int32 Index = StartIndex; Index <= EndIndex; ++Index)
+	{
+		double Slop = ComputeUnorientedSlope(Polyline2D[Index], Polyline2D[Index + 1], 0);
+		if (Slop > 2.)
+		{
+			Slop = 4. - Slop;
+		}
+		EdgeCharacteristics.Add(Slop, Polyline3D[Index].Distance(Polyline3D[Index + 1]));
+	}
+}
+
+void FTopologicalEdge::SpawnIdent(FDatabase& Database)
+{
+	if (!FEntity::SetId(Database))
+	{
+		return;
+	}
+
+	StartVertex->SpawnIdent(Database);
+	EndVertex->SpawnIdent(Database);
+	Curve->SpawnIdent(Database);
+
+	if(TopologicalLink.IsValid())
+	{
+		TopologicalLink->SpawnIdent(Database);
+	}
+	if (Mesh.IsValid())
+	{
+		Mesh->SpawnIdent(Database);
+	}
+}
+
+#ifdef CADKERNEL_DEV
+FInfoEntity& FEdgeLink::GetInfo(FInfoEntity& Info) const
+{
+	return FEntity::GetInfo(Info)
+		.Add(TEXT("active Entity"), ActiveEntity)
+		.Add(TEXT("twin Entities"), TwinsEntities);
+}
+#endif
+
diff --git a/Engine/Source/Runtime/Datasmith/CADKernel/Private/CADKernel/Topo/TopologicalEntity.cpp b/Engine/Source/Runtime/Datasmith/CADKernel/Private/CADKernel/Topo/TopologicalEntity.cpp
new file mode 100644
--- /dev/null
+++ b/Engine/Source/Runtime/Datasmith/CADKernel/Private/CADKernel/Topo/TopologicalEntity.cpp
@@ -0,0 +1,15 @@
+// Copyright Epic Games, Inc. All Rights Reserved.
+#include "CADKernel/Topo/TopologicalEntity.h"
+
+
+using namespace CADKernel;
+
+#ifdef CADKERNEL_DEV
+FInfoEntity& FTopologicalEntity::GetInfo(FInfoEntity& Info) const
+{
+	return FEntityGeom::GetInfo(Info)
+		.Add(TEXT("IsMesh"), IsMeshed());
+}
+#endif
+
+
diff --git a/Engine/Source/Runtime/Datasmith/CADKernel/Private/CADKernel/Topo/TopologicalFace.cpp b/Engine/Source/Runtime/Datasmith/CADKernel/Private/CADKernel/Topo/TopologicalFace.cpp
new file mode 100644
--- /dev/null
+++ b/Engine/Source/Runtime/Datasmith/CADKernel/Private/CADKernel/Topo/TopologicalFace.cpp
@@ -0,0 +1,493 @@
+// Copyright Epic Games, Inc. All Rights Reserved.
+
+#include "CADKernel/Topo/TopologicalFace.h"
+
+#include "CADKernel/Core/KernelParameters.h"
+#include "CADKernel/Core/System.h"
+#include "CADKernel/Geo/Curves/RestrictionCurve.h"
+#include "CADKernel/Geo/Curves/SegmentCurve.h"
+#include "CADKernel/Geo/GeoPoint.h"
+#include "CADKernel/Geo/Surfaces/Surface.h"
+#include "CADKernel/Mesh/Structure/FaceMesh.h"
+#include "CADKernel/Mesh/Structure/Grid.h"
+#include "CADKernel/Topo/TopologicalEdge.h"
+
+using namespace CADKernel;
+
+
+void FTopologicalFace::ComputeBoundary() const
+{
+	Boundary->Init();
+	TArray<TArray<FPoint2D>> TmpLoops;
+	Get2DLoopSampling(TmpLoops);
+
+	for (const TArray<FPoint2D>& Loop : TmpLoops)
+	{
+		for (const FPoint2D& Point : Loop)
+		{
+			Boundary->ExtendTo(Point);
+		}
+	}
+
+	// Check with the carrier surface bounds
+	CarrierSurface->ExtendBoundaryTo(Boundary);
+
+	Boundary->WidenIfDegenerated();
+	Boundary.SetReady();
+}
+
+void FTopologicalFace::Presample()
+{
+	const FSurfacicBoundary& FaceBoundaries = GetBoundary();
+	CarrierSurface->Presample(FaceBoundaries, CrossingCoordinates);
+}
+
+void FTopologicalFace::ApplyNaturalLoops()
+{
+	ensureCADKernel(Loops.Num() == 0);
+
+	TArray<TSharedPtr<FTopologicalEdge>> Edges;
+	Edges.Reserve(4);
+
+	const FSurfacicBoundary& Bounds = CarrierSurface->GetBoundary();
+
+	TFunction<void(const FPoint&, const FPoint&)> BuildEdge = [&](const FPoint& StartPoint, const FPoint& EndPoint)
+	{
+		double Tolerance2D = CarrierSurface->Get2DTolerance();
+		double Tolerance3D = CarrierSurface->Get3DTolerance();
+		TSharedRef<FCurve> Curve2D = FEntity::MakeShared<FSegmentCurve>(Tolerance2D, StartPoint, EndPoint, 2);
+		TSharedRef<FRestrictionCurve> Curve3D = FEntity::MakeShared<FRestrictionCurve>(Tolerance3D, CarrierSurface.ToSharedRef(), Curve2D);
+		TSharedPtr<FTopologicalEdge> Edge = FTopologicalEdge::Make(Curve3D);
+		if (Edge.IsValid())
+		{
+			return;
+		}
+		Edges.Add(Edge);
+	};
+
+	FPoint StartPoint;
+	FPoint EndPoint;
+
+	// Build 4 bounding edges of the surface
+	StartPoint.Set(Bounds.UVBoundaries[EIso::IsoU].Min, Bounds.UVBoundaries[EIso::IsoV].Min);
+	EndPoint.Set(Bounds.UVBoundaries[EIso::IsoU].Min, Bounds.UVBoundaries[EIso::IsoV].Max);
+	BuildEdge(StartPoint, EndPoint);
+
+	StartPoint.Set(Bounds.UVBoundaries[EIso::IsoU].Min, Bounds.UVBoundaries[EIso::IsoV].Max);
+	EndPoint.Set(Bounds.UVBoundaries[EIso::IsoU].Max, Bounds.UVBoundaries[EIso::IsoV].Max);
+	BuildEdge(StartPoint, EndPoint);
+
+	StartPoint.Set(Bounds.UVBoundaries[EIso::IsoU].Max, Bounds.UVBoundaries[EIso::IsoV].Max);
+	EndPoint.Set(Bounds.UVBoundaries[EIso::IsoU].Max, Bounds.UVBoundaries[EIso::IsoV].Min);
+	BuildEdge(StartPoint, EndPoint);
+
+	StartPoint.Set(Bounds.UVBoundaries[EIso::IsoU].Max, Bounds.UVBoundaries[EIso::IsoV].Min);
+	EndPoint.Set(Bounds.UVBoundaries[EIso::IsoU].Min, Bounds.UVBoundaries[EIso::IsoV].Min);
+	BuildEdge(StartPoint, EndPoint);
+
+	TSharedPtr<FTopologicalEdge> PreviousEdge = Edges.Last();
+	for (TSharedPtr<FTopologicalEdge>& Edge : Edges)
+	{
+		PreviousEdge->GetEndVertex()->Link(Edge->GetStartVertex());
+		PreviousEdge = Edge;
+	}
+
+	TArray<EOrientation> Orientations;
+	Orientations.Init(EOrientation::Front, Edges.Num());
+
+	TSharedPtr<FTopologicalLoop> Loop = FTopologicalLoop::Make(Edges, Orientations);
+	AddLoop(Loop);
+}
+
+void FTopologicalFace::AddLoops(const TArray<TSharedPtr<FTopologicalLoop>>& InLoops)
+{
+	for (TSharedPtr<FTopologicalLoop> Loop : InLoops)
+	{
+		AddLoop(Loop);
+	}
+
+	for (TSharedPtr<FTopologicalLoop> Loop : InLoops)
+	{
+		Loop->Orient();
+	}
+}
+
+void FTopologicalFace::AddLoop(const TSharedPtr<FTopologicalLoop>& InLoop)
+{
+	TSharedRef<FTopologicalFace> Face = StaticCastSharedRef<FTopologicalFace>(AsShared());
+	InLoop->SetSurface(Face);
+	if (Loops.Num() > 0)
+	{
+		InLoop->SetAsInnerBoundary();
+	}
+	Loops.Add(InLoop);
+}
+
+void FTopologicalFace::RemoveLoop(const TSharedPtr<FTopologicalLoop>& Loop)
+{
+	int32 Index = Loops.Find(Loop);
+	if (Index != INDEX_NONE)
+	{
+		Loop->ResetSurface();
+		Loops.RemoveAt(Index);
+	}
+
+	if (Loops.Num() == 0)
+	{
+		SetDeleted();
+	}
+}
+
+void FTopologicalFace::RemoveLinksWithNeighbours()
+{
+	for (const TSharedPtr<FTopologicalLoop>& Loop : GetLoops())
+	{
+		for (const FOrientedEdge& Edge : Loop->GetEdges())
+		{
+			Edge.Entity->RemoveFromLink();
+		}
+	}
+}
+
+bool FTopologicalFace::HasSameBoundariesAs(const TSharedPtr<FTopologicalFace>& OtherFace) const
+{
+	int32 EdgeCount = 0;
+	for (const TSharedPtr<FTopologicalLoop>& Loop : GetLoops())
+	{
+		for (const FOrientedEdge& Edge : Loop->GetEdges())
+		{
+			if (Edge.Entity->IsDegenerated())
+			{
+				continue;
+			}
+			Edge.Entity->GetLinkActiveEntity()->SetMarker1();
+			++EdgeCount;
+		}
+	}
+
+	bool bSameBoundary = true;
+	int32 OtherFaceEdgeCount = 0;
+	for (const TSharedPtr<FTopologicalLoop>& Loop : OtherFace->GetLoops())
+	{
+		OtherFaceEdgeCount += Loop->EdgeCount();
+		for (const FOrientedEdge& Edge : Loop->GetEdges())
+		{
+			if (Edge.Entity->IsDegenerated())
+			{
+				continue;
+			}
+			if (!Edge.Entity->GetLinkActiveEntity()->HasMarker1())
+			{
+				bSameBoundary = false;
+				break;
+			}
+			++OtherFaceEdgeCount;
+		}
+	}
+
+	for (const TSharedPtr<FTopologicalLoop>& Loop : GetLoops())
+	{
+		EdgeCount += Loop->EdgeCount();
+		for (const FOrientedEdge& Edge : Loop->GetEdges())
+		{
+			Edge.Entity->GetLinkActiveEntity()->ResetMarkers();
+		}
+	}
+
+	if (EdgeCount != OtherFaceEdgeCount)
+	{
+		bSameBoundary = false;
+	}
+
+	return bSameBoundary;
+}
+
+TSharedPtr<FTopologicalEdge> FTopologicalFace::GetLinkedEdge(const TSharedPtr<FTopologicalEdge>& LinkedEdge) const
+{
+	for (TWeakPtr<FTopologicalEdge> TwinEdge : LinkedEdge->GetTwinsEntities())
+	{
+		if (TwinEdge.Pin()->GetLoop()->GetFace() == AsShared())
+		{
+			return TwinEdge.Pin();
+		}
+	}
+
+	return TSharedPtr<FTopologicalEdge>();
+}
+
+void FTopologicalFace::GetEdgeIndex(const TSharedPtr<FTopologicalEdge>& Edge, int32& OutBoundaryIndex, int32& OutEdgeIndex) const
+{
+	OutEdgeIndex = INDEX_NONE;
+	for (OutBoundaryIndex = 0; OutBoundaryIndex < Loops.Num(); ++OutBoundaryIndex)
+	{
+		TSharedPtr<FTopologicalLoop> Loop = Loops[OutBoundaryIndex];
+		if ((OutEdgeIndex = Loop->GetEdgeIndex(Edge)) >= 0)
+		{
+			return;
+		}
+	}
+	OutBoundaryIndex = INDEX_NONE;
+}
+
+void FTopologicalFace::EvaluateGrid(FGrid& Grid) const
+{
+	CarrierSurface->EvaluateGrid(Grid);
+}
+
+const void FTopologicalFace::Get2DLoopSampling(TArray<TArray<FPoint2D>>& LoopSamplings) const
+{
+	LoopSamplings.Empty(GetLoops().Num());
+
+	for (const TSharedPtr<FTopologicalLoop> Loop : GetLoops())
+	{
+		TArray<FPoint2D>& LoopSampling2D = LoopSamplings.Emplace_GetRef();
+		Loop->Get2DSampling(LoopSampling2D);
+	}
+}
+
+void FTopologicalFace::SpawnIdent(FDatabase& Database)
+{
+	if (!FEntity::SetId(Database))
+	{
+		return;
+	}
+
+	SpawnIdentOnEntities(Loops, Database);
+	CarrierSurface->SpawnIdent(Database);
+	if (Mesh.IsValid())
+	{
+		Mesh->SpawnIdent(Database);
+	}
+}
+
+#ifdef CADKERNEL_DEV
+FInfoEntity& FTopologicalFace::GetInfo(FInfoEntity& Info) const
+{
+	return FTopologicalEntity::GetInfo(Info)
+		.Add(TEXT("Carrier Surface"), CarrierSurface)
+		.Add(TEXT("Boundary"), (FSurfacicBoundary&) Boundary)
+		.Add(TEXT("Loops"), Loops)
+		.Add(TEXT("QuadCriteria"), QuadCriteria)
+		.Add(TEXT("mesh"), Mesh);
+}
+#endif
+
+TSharedRef<FFaceMesh> FTopologicalFace::GetOrCreateMesh(const TSharedRef<FModelMesh>& MeshModel)
+{
+	if (!Mesh.IsValid())
+	{
+		Mesh = FEntity::MakeShared<FFaceMesh>(MeshModel, StaticCastSharedRef<FTopologicalFace>(AsShared()));
+	}
+	return Mesh.ToSharedRef();
+}
+
+void FTopologicalFace::InitDeltaUs()
+{
+	CrossingPointDeltaMins[EIso::IsoU].Init(SMALL_NUMBER, CrossingCoordinates[EIso::IsoU].Num() - 1);
+	CrossingPointDeltaMaxs[EIso::IsoU].Init(HUGE_VALUE, CrossingCoordinates[EIso::IsoU].Num() - 1);
+
+	CrossingPointDeltaMins[EIso::IsoV].Init(SMALL_NUMBER, CrossingCoordinates[EIso::IsoV].Num() - 1);
+	CrossingPointDeltaMaxs[EIso::IsoV].Init(HUGE_VALUE, CrossingCoordinates[EIso::IsoV].Num() - 1);
+}
+
+void FTopologicalFace::ChooseFinalDeltaUs()
+{
+	TFunction<void(const TArray<double>&, TArray<double>&)> ChooseFinalDeltas = [](const TArray<double>& DeltaUMins, TArray<double>& DeltaUMaxs)
+	{
+		int32 Index = 0;
+		for (; Index < DeltaUMins.Num(); ++Index)
+		{
+			if (DeltaUMins[Index] > DeltaUMaxs[Index])
+			{
+				DeltaUMaxs[Index] = DeltaUMins[Index];
+			}
+		}
+	};
+
+	ChooseFinalDeltas(CrossingPointDeltaMins[EIso::IsoU], CrossingPointDeltaMaxs[EIso::IsoU]);
+	ChooseFinalDeltas(CrossingPointDeltaMins[EIso::IsoV], CrossingPointDeltaMaxs[EIso::IsoV]);
+}
+
+
+// =========================================================================================================================================================================================================
+// =========================================================================================================================================================================================================
+// =========================================================================================================================================================================================================
+//
+//
+//                                                                            NOT YET REVIEWED
+//
+//
+// =========================================================================================================================================================================================================
+// =========================================================================================================================================================================================================
+// =========================================================================================================================================================================================================
+
+TSharedPtr<FEntityGeom> FTopologicalFace::ApplyMatrix(const FMatrixH& InMatrix) const
+{
+	NOT_IMPLEMENTED;
+	return TSharedPtr<FEntityGeom>();
+}
+
+// Quad ==============================================================================================================================================================================================================================
+
+double FTopologicalFace::GetQuadCriteria()
+{
+	if (GetQuadType() == EQuadType::Unset)
+	{
+		return 0;
+	}
+	return QuadCriteria;
+}
+
+void FTopologicalFace::ComputeQuadCriteria()
+{
+	if (GetQuadType() != EQuadType::Unset)
+	{
+		QuadCriteria = FMath::Max(Curvatures[EIso::IsoU].Max, Curvatures[EIso::IsoU].Max);
+	}
+}
+
+void FTopologicalFace::ComputeSurfaceSideProperties()
+{
+	TFunction<double(const int32)> GetSideLength = [&](const int32 SideIndex)
+	{
+		TSharedPtr<FTopologicalLoop> Loop = GetLoops()[0];
+
+		double Length = 0;
+		int32 NextSideIndex = SideIndex + 1;
+		if (NextSideIndex == GetStartSideIndices().Num())
+		{
+			NextSideIndex = 0;
+		}
+		int32 EndIndex = GetStartSideIndices()[NextSideIndex];
+		for (int32 Index = GetStartSideIndices()[SideIndex]; Index != EndIndex;)
+		{
+			Length += Loop->GetEdge(Index)->Length();
+			if (++Index == Loop->EdgeCount())
+			{
+				Index = 0;
+			}
+		}
+		return Length;
+	};
+
+	Loops[0]->FindSurfaceCorners(SurfaceCorners, StartSideIndices);
+	Loops[0]->ComputeBoundaryProperties(StartSideIndices, SideProperties);
+
+	LoopLength = 0;
+	for (int32 Index = 0; Index < SurfaceCorners.Num(); ++Index)
+	{
+		SideProperties[Index].Length3D = GetSideLength(Index);
+		LoopLength += SideProperties[Index].Length3D;
+	}
+}
+
+void FTopologicalFace::DefineSurfaceType()
+{
+	const double Tolerance3D = CarrierSurface->Get3DTolerance();
+	const double GeometricTolerance = 20.0 * Tolerance3D;
+
+	switch (SurfaceCorners.Num())
+	{
+	case 3:
+		QuadType = EQuadType::Triangular;
+		break;
+
+	case 4:
+		QuadType = EQuadType::Other;
+		for (int32 Index = 0; Index < 4; ++Index)
+		{
+			// If the type is not ISO, the neighbor surface is checked, if it's quad so it's ok...
+			if (SideProperties[Index].IsoType == EIso::UndefinedIso)
+			{
+				TSharedPtr<FTopologicalEdge> Edge = Loops[0]->GetEdge(StartSideIndices[Index]);
+				int32 NeighborsNum = Edge->GetTwinsEntityCount();
+				// If non manifold Edge => Stop
+				if (NeighborsNum != 2)
+				{
+					return;
+				}
+
+				{
+					int32 OppositIndex = (Index + 2) % 4;
+					SideProperties[Index].IsoType = SideProperties[OppositIndex].IsoType;
+					if (SideProperties[Index].IsoType == EIso::UndefinedIso)
+					{
+						int32 AdjacentIndex = (Index + 1) % 4;
+						if (SideProperties[AdjacentIndex].IsoType != EIso::UndefinedIso)
+						{
+							SideProperties[Index].IsoType = (SideProperties[AdjacentIndex].IsoType == EIso::IsoU) ? EIso::IsoV : EIso::IsoU;
+						}
+					}
+				}
+
+				TSharedPtr<FTopologicalFace> Neighbor;
+				{
+					for(const TWeakPtr<FTopologicalEdge>& NeighborEdge : Edge->GetTwinsEntities() )
+					{
+						if (NeighborEdge.HasSameObject(Edge.Get()))
+						{
+							continue;
+						}
+						Neighbor = Edge->GetLoop()->GetFace();
+					}
+				}
+
+				ensure(Neighbor.IsValid());
+
+				// it's not a quad surface
+				if (Neighbor->SurfaceCorners.Num() == 0)
+				{
+					return;
+				}
+
+				TSharedPtr<FTopologicalEdge> TwinEdge = Edge->GetFirstTwinEdge();
+				int32 SideIndex = Neighbor->GetSideIndex(TwinEdge);
+				if (SideIndex < 0)
+				{
+					return;
+				}
+
+				const FEdge2DProperties& Property = Neighbor->GetSideProperty(SideIndex);
+				if (Property.IsoType == EIso::UndefinedIso)
+				{
+					return;
+				}
+
+				double SideLength = SideProperties[Index].Length3D;
+				double OtherSideLength = Property.Length3D;
+
+				if (FMath::Abs(SideLength - OtherSideLength) < GeometricTolerance)
+				{
+					int32 OppositIndex = (Index + 2) % 4;
+					if (SideProperties[OppositIndex].IsoType == EIso::UndefinedIso)
+					{
+						if (Index < 2)
+						{
+							if (SideProperties[!Index].IsoType == EIso::IsoU)
+							{
+								SideProperties[Index].IsoType = EIso::IsoV;
+							}
+							else
+							{
+								SideProperties[Index].IsoType = EIso::IsoU;
+							}
+						}
+						return;
+					}
+					SideProperties[Index].IsoType = SideProperties[OppositIndex].IsoType;
+				}
+			}
+		}
+
+		if ((SideProperties[0].IsoType != EIso::UndefinedIso) && (SideProperties[1].IsoType != EIso::UndefinedIso) && (SideProperties[0].IsoType == SideProperties[2].IsoType) && (SideProperties[1].IsoType == SideProperties[3].IsoType))
+		{
+			QuadType = EQuadType::Quadrangular;
+		}
+		break;
+
+	default:
+		QuadType = EQuadType::Other;
+		break;
+	}
+}
+
