CADKernel Source Code v0.01

diff --git a/Engine/Source/Runtime/Datasmith/CADKernel/Private/CADKernel/Topo/TopologicalLoop.cpp b/Engine/Source/Runtime/Datasmith/CADKernel/Private/CADKernel/Topo/TopologicalLoop.cpp
new file mode 100644
--- /dev/null
+++ b/Engine/Source/Runtime/Datasmith/CADKernel/Private/CADKernel/Topo/TopologicalLoop.cpp
@@ -0,0 +1,458 @@
+// Copyright Epic Games, Inc. All Rights Reserved.
+
+#include "CADKernel/Topo/TopologicalLoop.h"
+
+#include "CADKernel/Core/KernelParameters.h"
+#include "CADKernel/Geo/Curves/RestrictionCurve.h"
+#include "CADKernel/Math/SlopeUtils.h"
+#include "CADKernel/Topo/TopologicalEdge.h"
+#include "CADKernel/Topo/TopologicalLink.h"
+#include "CADKernel/Topo/TopologicalFace.h"
+#include "CADKernel/UI/Display.h"
+#include "CADKernel/UI/Message.h"
+#include "CADKernel/Utils/Util.h"
+
+
+using namespace CADKernel;
+
+TSharedPtr<FTopologicalLoop> FTopologicalLoop::Make(const TArray<TSharedPtr<FTopologicalEdge>>& InEdges, const TArray<EOrientation>& InEdgeDirections)
+{
+	TSharedRef<FTopologicalLoop> Loop = FEntity::MakeShared<FTopologicalLoop>(InEdges, InEdgeDirections);
+	for (FOrientedEdge& OrientedEdge : Loop->GetEdges())
+	{
+		OrientedEdge.Entity->SetBoundary(Loop);
+	}
+	return Loop;
+}
+
+FTopologicalLoop::FTopologicalLoop(const TArray<TSharedPtr<FTopologicalEdge>>& InEdges, const TArray<EOrientation>& InEdgeDirections)
+	: FTopologicalEntity()
+	, Surface(TSharedPtr<FTopologicalFace>())
+	, bExternalLoop(true)
+{
+	Edges.Reserve(InEdges.Num());
+	for(int32 Index = 0; Index < InEdges.Num(); ++Index)
+	{
+		TSharedPtr<FTopologicalEdge> Edge = InEdges[Index];
+		EOrientation Orientation = InEdgeDirections[Index];
+		Edges.Emplace(Edge, Orientation);
+	}
+	EnsureLogicalClosing();
+}
+
+void FTopologicalLoop::RemoveEdge(TSharedPtr<FTopologicalEdge>& EdgeToRemove)
+{
+	for (int32 IEdge = 0; IEdge < Edges.Num(); IEdge++)
+	{
+		if (Edges[IEdge].Entity == EdgeToRemove)
+		{
+			EdgeToRemove->RemoveBoundary();
+			Edges.RemoveAt(IEdge);
+			return;
+		}
+	}
+	ensureCADKernel(false);
+}
+
+EOrientation FTopologicalLoop::GetDirection(TSharedPtr<FTopologicalEdge>& InEdge, bool bAllowLinkedEdge) const
+{
+	ensureCADKernel(InEdge.IsValid());
+
+	for (const FOrientedEdge& BoundaryEdge : Edges)
+	{
+		if (BoundaryEdge.Entity == InEdge)
+		{
+			return BoundaryEdge.Direction;
+		}
+		else if (bAllowLinkedEdge && BoundaryEdge.Entity->IsLinkedTo(InEdge.ToSharedRef()))
+		{
+			return BoundaryEdge.Direction;
+		}
+	}
+
+	FMessage::Printf(Debug, TEXT("Edge %d is not in boundary %d Edges\n"), InEdge->GetId(), GetId());
+	ensureCADKernel(false);
+	return EOrientation::Front;
+}
+
+
+void FTopologicalLoop::Get2DSampling(TArray<FPoint2D>& LoopSampling)
+{
+	int32 PointNum = 0;
+	for (const FOrientedEdge& Edge : Edges)
+	{
+		PointNum += Edge.Entity->GetCurve()->GetPolylineSize();
+	}
+
+	LoopSampling.Empty(PointNum);
+
+	for (const FOrientedEdge& Edge : Edges)
+	{
+		Edge.Entity->GetDiscretization2DPoints(Edge.Direction, LoopSampling);
+		LoopSampling.Pop();
+	}
+	LoopSampling.Emplace(LoopSampling[0]);
+}
+
+void FTopologicalLoop::Orient()
+{
+	ensureCADKernel(Edges.Num() > 0);
+
+	TArray<FPoint2D> LoopSampling;
+	Get2DSampling(LoopSampling);
+
+	double UMin = HUGE_VAL;
+	double UMax = -HUGE_VAL;
+	double VMin = HUGE_VAL;
+	double VMax = -HUGE_VAL;
+	int32 IndexUMin = 0;
+	int32 IndexUMax = 0;
+	int32 IndexVMin = 0;
+	int32 IndexVMax = 0;
+
+	int32 PointCount = LoopSampling.Num();
+
+	for (int32 Index = 0; Index < PointCount; ++Index)
+	{
+		if (LoopSampling[Index].U > UMax)
+		{
+			UMax = LoopSampling[Index].U;
+			IndexUMax = Index;
+		}
+		if (LoopSampling[Index].U < UMin)
+		{
+			UMin = LoopSampling[Index].U;
+			IndexUMin = Index;
+		}
+
+		if (LoopSampling[Index].V > VMax)
+		{
+			VMax = LoopSampling[Index].V;
+			IndexVMax = Index;
+		}
+		if (LoopSampling[Index].V < VMin)
+		{
+			VMin = LoopSampling[Index].V;
+			IndexVMin = Index;
+		}
+	}
+
+	int32 WrongOrientationNum = 0;
+	TFunction<void(int32&, double)> CompareOrientation = [&](int32& Index, double ReferenceSlop)
+	{
+		// if the slop of the selected segment is not close to the BBox side (closed of 0 of 4), so the angle between the neighboring segments of the local extrema is not closed to 4 and allows to defined the orientation
+		// Pic case: the slop is compute between previous and nex segment of the extrema 
+		int32 Index2 = Index == 0 ? PointCount - 1 : Index - 1;
+		double Slop = ComputePositiveSlope(LoopSampling[Index], LoopSampling[Index + 1], LoopSampling[Index2]);
+		if (Slop > 4.2)
+		{
+			WrongOrientationNum++;
+		}
+		else if (Slop > 3.8)
+		{
+			// the angle between the neighboring segments of the local extrema is too closed to 4 (PI) so the slop of the segment Index, Index+1 is compared to the AABB side to define the orientation
+			Slop = ComputeUnorientedSlope(LoopSampling[Index], LoopSampling[Index + 1], ReferenceSlop);
+			// slop should be closed to [0, 0.2] or [3.8, 4]
+			if (Slop > 3)
+			{
+				WrongOrientationNum++;
+			}
+		}
+
+#ifdef DEBUG_ORIENT
+		if (false /*Surface->GetId() == 826*/)
+		{
+			F3DDebugSession _(*FString::Printf(TEXT("Pic case")));
+			{
+				F3DDebugSession _(*FString::Printf(TEXT("Loop")));
+				DisplayPolyline(LoopSampling);
+			}
+			{
+				F3DDebugSession _(*FString::Printf(TEXT("Next")));
+				::DisplaySegment(LoopSampling[Index + 1], LoopSampling[Index], EVisuProperty::NonManifoldEdge);
+			}
+			{
+				F3DDebugSession _(*FString::Printf(TEXT("Previous")));
+				::DisplaySegment(LoopSampling[Index2], LoopSampling[Index], EVisuProperty::NonManifoldEdge);
+			}
+			{
+				F3DDebugSession _(*FString::Printf(TEXT("Node %f"), Slop));
+				::Display(LoopSampling[Index], EVisuProperty::RedPoint, Index);
+			}
+			Wait();
+		}
+#endif
+	};
+
+	CompareOrientation(IndexUMin, 6);
+	CompareOrientation(IndexUMax, 2);
+	CompareOrientation(IndexVMin, 0);
+	CompareOrientation(IndexVMax, 4);
+
+#ifdef DEBUG_ORIENT
+	if (WrongOrientationNum != 0 && WrongOrientationNum != 4)
+	{
+		F3DDebugSession GraphicSession(TEXT("Points of evaluation"));
+		{
+			F3DDebugSession G(*FString::Printf(TEXT("Loop Discretization %d"), Surface.Pin()->GetId()));
+			DisplayPolyline(LoopSampling);
+			for (int32 Index = 0; Index < LoopSampling.Num(); ++Index)
+			{
+				::Display(LoopSampling[Index], Index);
+			}
+		}
+		{
+			F3DDebugSession G(*FString::Printf(TEXT("Seg UMin Loop Discretization %d"), Surface.Pin()->GetId()));
+			::DisplaySegment(LoopSampling[IndexUMin + 1], LoopSampling[IndexUMin]);
+			::Display(LoopSampling[IndexUMin], EVisuProperty::RedPoint, IndexUMin);
+		}
+		{
+			F3DDebugSession G(*FString::Printf(TEXT("Seg UMax Loop Discretization %d"), Surface.Pin()->GetId()));
+			::DisplaySegment(LoopSampling[IndexUMax + 1], LoopSampling[IndexUMax]);
+			::Display(LoopSampling[IndexUMax], EVisuProperty::RedPoint, IndexUMax);
+		}
+		{
+			F3DDebugSession G(*FString::Printf(TEXT("Seg VMin Loop Discretization %d"), Surface.Pin()->GetId()));
+			::DisplaySegment(LoopSampling[IndexVMin + 1], LoopSampling[IndexVMin]);
+			::Display(LoopSampling[IndexVMin], EVisuProperty::RedPoint, IndexVMin);
+		}
+		{
+			F3DDebugSession G(*FString::Printf(TEXT("Seg VMax Loop Discretization %d"), Surface.Pin()->GetId()));
+			::DisplaySegment(LoopSampling[IndexVMax + 1], LoopSampling[IndexVMax]);
+			::Display(LoopSampling[IndexVMax], EVisuProperty::RedPoint, IndexVMax);
+		}
+		//Wait();
+		FMessage::Printf(Log, TEXT("WARNING: Loop Orientation of surface %d could failed\n"), Surface.Pin()->GetId());
+	}
+#endif
+
+	if ((WrongOrientationNum > 2) == bExternalLoop)
+	{
+		SwapOrientation();
+	}
+}
+
+void FTopologicalLoop::SwapOrientation()
+{
+	TArray<FOrientedEdge> TmpEdges;
+	TmpEdges.Reserve(Edges.Num());
+	for (int32 Index = Edges.Num() - 1; Index >= 0; Index--)
+	{
+		TmpEdges.Emplace(Edges[Index].Entity, GetReverseOrientation(Edges[Index].Direction));
+	}
+	Swap(TmpEdges, Edges);
+}
+
+void FTopologicalLoop::ReplaceEdge(TSharedPtr<FTopologicalEdge>& OldEdge, TSharedPtr<FTopologicalEdge>& NewEdge)
+{
+	for (int32 IEdge = 0; IEdge < (int32)Edges.Num(); IEdge++)
+	{
+		if (Edges[IEdge].Entity == OldEdge)
+		{
+			Edges[IEdge].Entity = NewEdge;
+			OldEdge->RemoveBoundary();
+			NewEdge->SetBoundary(StaticCastSharedRef<FTopologicalLoop>(AsShared()));
+			return;
+		}
+	}
+	ensureCADKernel(false);
+}
+
+void FTopologicalLoop::ReplaceEdge(TSharedPtr<FTopologicalEdge>& Edge, TArray<TSharedPtr<FTopologicalEdge>>& NewEdges)
+{
+	TArray<FOrientedEdge> TmpEdges;
+	int32 NewEdgeNum = Edges.Num() + NewEdges.Num();
+	TmpEdges.Reserve(NewEdgeNum);
+
+	Edge->RemoveBoundary();
+	TSharedRef<FTopologicalLoop> Loop = StaticCastSharedRef<FTopologicalLoop>(AsShared());
+	for (TSharedPtr<FTopologicalEdge>& NewEdge : NewEdges)
+	{
+		NewEdge->SetBoundary(Loop);
+	}
+
+	for (int32 IEdge = 0; IEdge < Edges.Num(); IEdge++)
+	{
+		if (Edges[IEdge].Entity == Edge)
+		{
+			EOrientation OldEdgeDirection = Edges[IEdge].Direction;
+			if (OldEdgeDirection == EOrientation::Front)
+			{
+				for (int32 INewEdge = 0; INewEdge < NewEdges.Num(); INewEdge++)
+				{
+					TmpEdges.Emplace(NewEdges[INewEdge], OldEdgeDirection);
+				}
+			}
+			else
+			{
+				for (int32 INewEdge = NewEdges.Num() - 1; INewEdge >= 0; --INewEdge)
+				{
+					TmpEdges.Emplace(NewEdges[INewEdge], OldEdgeDirection);
+				}
+			}
+		}
+		else 
+		{
+			TmpEdges.Emplace(Edges[IEdge]);
+		}
+	}
+
+	Swap(TmpEdges, Edges);
+}
+
+void FTopologicalLoop::ReplaceEdgesWithMergedEdge(TArray<TSharedPtr<FTopologicalEdge>>& OldEdges, TSharedPtr<FTopologicalVertex>& MiddleVertex, TSharedPtr<FTopologicalEdge>& NewEdge)
+{
+	ensureCADKernel(OldEdges.Num() == 2);
+	ensureCADKernel(EdgeCount() != 2);
+
+	NewEdge->SetBoundary(StaticCastSharedRef<FTopologicalLoop>(AsShared()));
+	for (TSharedPtr<FTopologicalEdge> Edge : OldEdges)
+	{
+		Edge->RemoveBoundary();
+		Edge->GetStartVertex()->RemoveConnectedEdge(Edge.ToSharedRef());
+		Edge->GetEndVertex()->RemoveConnectedEdge(Edge.ToSharedRef());
+	}
+
+	int32 IndexFirstOldEdge = GetEdgeIndex(OldEdges[0]);
+	int32 IndexLastOldEdge = GetEdgeIndex(OldEdges[1]);
+
+	ensureCADKernel(IndexFirstOldEdge != INDEX_NONE && IndexLastOldEdge != INDEX_NONE);
+
+	Sort(IndexFirstOldEdge, IndexLastOldEdge);
+	if (IndexFirstOldEdge == 0 && IndexLastOldEdge != 1)
+	{
+		IndexLastOldEdge = 0;
+		IndexFirstOldEdge = EdgeCount() - 1;
+	}
+	int32 PreviousEdgeIndex = IndexFirstOldEdge ? IndexFirstOldEdge - 1 : EdgeCount() - 1;
+	TSharedPtr<FTopologicalVertex> PreviousVertex = Edges[PreviousEdgeIndex].Direction == EOrientation::Front ? Edges[PreviousEdgeIndex].Entity->GetEndVertex() : Edges[PreviousEdgeIndex].Entity->GetStartVertex();
+	EOrientation NewEdgeOrientation = NewEdge->GetStartVertex() == PreviousVertex ? EOrientation::Front : EOrientation::Back;
+
+	if (IndexLastOldEdge)
+	{
+		Edges.RemoveAt(IndexFirstOldEdge);
+		Edges.RemoveAt(IndexFirstOldEdge + 1);
+		Edges.EmplaceAt(IndexFirstOldEdge, NewEdge, NewEdgeOrientation);
+	}
+	else
+	{
+		Edges.RemoveAt(0);
+		Edges.Pop();
+		Edges.Emplace(NewEdge, NewEdgeOrientation);
+	}
+}
+
+
+
+void FTopologicalLoop::FindSurfaceCorners(TArray<TSharedPtr<FTopologicalVertex>>& OutCorners, TArray<int32>& OutStartSideIndex) const
+{
+	TArray<double> BreakValues;
+	FindBreaks(OutCorners, OutStartSideIndex, BreakValues);
+}
+
+void FTopologicalLoop::ComputeBoundaryProperties(const TArray<int32>& StartSideIndex, TArray<FEdge2DProperties>& OutSideProperties) const
+{
+	if (StartSideIndex.Num() == 0)
+	{
+		return;
+	}
+
+	OutSideProperties.Reserve(StartSideIndex.Num());
+
+	int32 EdgeIndex = StartSideIndex[0];
+	for (int32 SideIndex = 0; SideIndex < StartSideIndex.Num(); ++SideIndex)
+	{
+		int32 LastEdgeIndex = SideIndex + 1;
+		LastEdgeIndex = (LastEdgeIndex == StartSideIndex.Num()) ? StartSideIndex[0] : StartSideIndex[LastEdgeIndex];
+
+		FEdge2DProperties& SideProperty = OutSideProperties.Emplace_GetRef();
+		do
+		{
+			Edges[EdgeIndex].Entity->ComputeEdge2DProperties(SideProperty);
+			if (++EdgeIndex == Edges.Num())
+			{
+				EdgeIndex = 0;
+			}
+		} while (EdgeIndex != LastEdgeIndex);
+		SideProperty.Finalize();
+	}
+}
+
+void FTopologicalLoop::EnsureLogicalClosing()
+{
+	TSharedPtr<FTopologicalEdge> Edge = Edges.Last().Entity;
+
+	TSharedPtr<FTopologicalVertex> EdgeStartVertex = Edge->GetStartVertex(Edges[Edges.Num() - 1].Direction);
+
+	for (int32 Index = 0; Index < Edges.Num(); Index++)
+	{
+		const TSharedPtr<FTopologicalVertex> PreviousEdgeEndVertex = Edge->GetOtherVertex(EdgeStartVertex.ToSharedRef());
+
+		Edge = Edges[Index].Entity;
+		EdgeStartVertex = Edge->GetStartVertex(Edges[Index].Direction);
+
+		PreviousEdgeEndVertex->Link(EdgeStartVertex.ToSharedRef());
+	}
+}
+
+#ifdef CADKERNEL_DEV
+FInfoEntity& FTopologicalLoop::GetInfo(FInfoEntity& Info) const
+{
+	return FEntity::GetInfo(Info)
+		.Add(TEXT("Edges"), (TArray<TOrientedEntity<FEntity>>&) Edges)
+		.Add(TEXT("domain"), Surface);
+}
+#endif
+
+void FTopologicalLoop::FindBreaks(TArray<TSharedPtr<FTopologicalVertex>>& OutBreaks, TArray<int32>& OutStartSideIndex, TArray<double>& OutBreakValues) const
+{
+	const double MinCosAngleOfBreak = -0.7;  // 135 deg
+	if (Edges.Num() == 0)
+	{
+		return;
+	}
+
+	int32 EdgeNum = (int32)Edges.Num();
+	OutBreaks.Empty(EdgeNum);
+	OutBreakValues.Empty(EdgeNum);
+	OutStartSideIndex.Empty(EdgeNum);
+
+	FPoint StartTangentEdge;
+	FPoint EndTangentPreviousEdge;
+	Edges[EdgeNum - 1].Entity->GetTangentsAtExtremities(StartTangentEdge, EndTangentPreviousEdge, Edges[EdgeNum - 1].Direction == EOrientation::Front);
+	bool bPreviousIsSurface = (Edges[EdgeNum - 1].Entity->GetTwinsEntityCount() > 1);
+
+	for (int32 Index = 0; Index < EdgeNum; Index++)
+	{
+		FPoint EndTangentEdge;
+		Edges[Index].Entity->GetTangentsAtExtremities(StartTangentEdge, EndTangentEdge, Edges[Index].Direction == EOrientation::Front);
+		bool bIsSurface = (Edges[Index].Entity->GetTwinsEntityCount() > 1);
+
+		// if both edge are border, the rupture is not evaluate. 
+		if (bIsSurface || bPreviousIsSurface)
+		{
+			double CosAngle = StartTangentEdge.ComputeCosinus(EndTangentPreviousEdge);
+
+#ifdef FIND_BREAKS
+			{
+				FPoint& Start = BoundaryEdges[Index]->GetStartVertex(BoundaryEdgeDirections[Index])->GetCoordinates();
+				Open3DDebugSession(TEXT("Cos Angle " + Utils::ToString(CosAngle));
+				DisplayPoint(Start, (CosAngle > MinCosAngleOfBreak) ? EVisuProperty::RedPoint : EVisuProperty::BluePoint);
+				DisplaySegment(Start, Start + StartTangentEdge);
+				DisplaySegment(Start, Start + EndTangentPreviousEdge);
+				Close3DDebugSession();
+			}
+#endif
+
+			if (CosAngle > MinCosAngleOfBreak)
+			{
+				OutBreaks.Add(Edges[Index].Direction == EOrientation::Front ? Edges[Index].Entity->GetStartVertex() : Edges[Index].Entity->GetEndVertex());
+				OutBreakValues.Add(CosAngle);
+				OutStartSideIndex.Add(Index);
+			}
+		}
+
+		EndTangentPreviousEdge = EndTangentEdge;
+		bPreviousIsSurface = bIsSurface;
+	}
+}
diff --git a/Engine/Source/Runtime/Datasmith/CADKernel/Private/CADKernel/Topo/TopologicalVertex.cpp b/Engine/Source/Runtime/Datasmith/CADKernel/Private/CADKernel/Topo/TopologicalVertex.cpp
new file mode 100644
--- /dev/null
+++ b/Engine/Source/Runtime/Datasmith/CADKernel/Private/CADKernel/Topo/TopologicalVertex.cpp
@@ -0,0 +1,182 @@
+// Copyright Epic Games, Inc. All Rights Reserved.
+
+#include "CADKernel/Topo/TopologicalVertex.h"
+
+#include "CADKernel/Mesh/Structure/ModelMesh.h"
+#include "CADKernel/Mesh/Structure/VertexMesh.h"
+#include "CADKernel/Topo/TopologicalEdge.h"
+
+using namespace CADKernel;
+
+
+#ifdef CADKERNEL_DEV
+FInfoEntity& FTopologicalVertex::GetInfo(FInfoEntity& Info) const
+{
+	return FTopologicalEntity::GetInfo(Info)
+		.Add(TEXT("Link"), TopologicalLink)
+		.Add(TEXT("Position"), Coordinates)
+		.Add(TEXT("ConnectedEdges"), ConnectedEdges)
+		.Add(TEXT("mesh"), Mesh);
+}
+#endif
+
+void FTopologicalVertex::AddConnectedEdge(TSharedRef<FTopologicalEdge> Edge)
+{
+	ConnectedEdges.Add(Edge);
+}
+
+void FTopologicalVertex::RemoveConnectedEdge(TSharedRef<FTopologicalEdge> Edge)
+{
+	for (int32 EdgeIndex = 0; EdgeIndex < ConnectedEdges.Num(); EdgeIndex++)
+	{
+		if(ConnectedEdges[EdgeIndex]==Edge) 
+		{
+			ConnectedEdges.RemoveAt(EdgeIndex);
+			return;
+		}
+	}
+	ensureCADKernel(false);
+}
+
+bool FTopologicalVertex::IsBorderVertex()
+{
+	for (TWeakPtr<FTopologicalEdge> Edge : GetConnectedEdges())
+	{
+		if (Edge.Pin()->GetTwinsEntityCount() == 1)
+		{
+			return true;
+		}
+	}
+	return false;
+}
+
+TSharedPtr<FEntityGeom> FTopologicalVertex::ApplyMatrix(const FMatrixH& InMatrix) const
+{
+	FPoint transformedPoint = InMatrix.Multiply(Coordinates);
+	return FEntity::MakeShared<FTopologicalVertex>(transformedPoint);
+}
+
+void FTopologicalVertex::GetConnectedEdges(TSharedPtr<FTopologicalVertex> OtherVertex, TArray<TSharedPtr<FTopologicalEdge>>& Edges) const
+{
+	Edges.Reserve(GetTwinsEntityCount());
+
+	TSharedPtr<TTopologicalLink<FTopologicalVertex>> OtherVertexLink = OtherVertex->GetLink();
+	for (TWeakPtr<FTopologicalVertex> Vertex : GetTwinsEntities())
+	{
+		ensureCADKernel(Vertex.IsValid());
+		for (TWeakPtr<FTopologicalEdge> Edge : Vertex.Pin()->GetConnectedEdges())
+		{
+			ensureCADKernel(Edge.IsValid());
+			if (Edge.Pin()->GetOtherVertex(Vertex.Pin().ToSharedRef())->GetLink() == OtherVertexLink)
+			{
+				Edges.Add(Edge.Pin());
+			}
+		}
+	}
+}
+
+void FTopologicalVertex::Link(TSharedRef<FTopologicalVertex> Twin)
+{
+	// The active vertex is always the closest of the Barycenter
+	if (TopologicalLink.IsValid() && Twin->TopologicalLink.IsValid())
+	{
+		if (TopologicalLink == Twin->TopologicalLink)
+		{
+			return;
+		}
+	}
+
+	FPoint Barycenter = GetBarycenter() * (double) GetTwinsEntityCount()
+		+ Twin->GetBarycenter() * (double)Twin->GetTwinsEntityCount();
+
+	MakeLink(Twin);
+
+	Barycenter /= (double) GetTwinsEntityCount();
+	GetLink()->SetBarycenter(Barycenter);
+
+	// Find the closest vertex of the Barycenter
+	GetLink()->DefineActiveEntity();
+}
+
+void FVertexLink::ComputeBarycenter()
+{
+	Barycenter = FPoint::ZeroPoint;
+	for (TWeakPtr<FTopologicalVertex> Vertex : TwinsEntities)
+	{
+		Barycenter += Vertex.Pin()->Coordinates;
+	}
+	Barycenter /= TwinsEntities.Num();
+}
+
+void FVertexLink::DefineActiveEntity()
+{
+	double DistanceSquare = HUGE_VALUE;
+	TWeakPtr<FTopologicalVertex> ClosedVertex = TwinsEntities.HeapTop();
+	for (const TWeakPtr<FTopologicalVertex>& Vertex : TwinsEntities)
+	{
+		double Square = Vertex.Pin()->SquareDistance(Barycenter);
+		if (Square < DistanceSquare)
+		{
+			DistanceSquare = Square;
+			ClosedVertex = Vertex;
+			if (FMath::IsNearlyZero(Square))
+			{
+				break;
+			}
+		}
+	}
+	ActiveEntity = ClosedVertex;
+}
+
+TSharedRef<FVertexMesh> FTopologicalVertex::GetOrCreateMesh(TSharedRef<FModelMesh>& MeshModel)
+{
+	if (!IsActiveEntity())
+	{
+		return GetLinkActiveEntity()->GetOrCreateMesh(MeshModel);
+	}
+
+	if (!Mesh.IsValid())
+	{
+		Mesh = FEntity::MakeShared<FVertexMesh>(MeshModel, StaticCastSharedRef<FTopologicalVertex>(AsShared()));
+
+		Mesh->GetNodeCoordinates().Emplace(GetBarycenter());
+		Mesh->RegisterCoordinates();
+		MeshModel->AddMesh(Mesh.ToSharedRef());
+		SetMeshed();
+	}
+	return Mesh.ToSharedRef();
+}
+
+void FTopologicalVertex::SpawnIdent(FDatabase& Database)
+{
+	if (!FEntity::SetId(Database))
+	{
+		return;
+	}
+
+	if (TopologicalLink.IsValid())
+	{
+		TopologicalLink->SpawnIdent(Database);
+	}
+
+	if(Mesh.IsValid())
+	{
+		Mesh->SpawnIdent(Database);
+	}
+}
+
+
+#ifdef CADKERNEL_DEV
+FInfoEntity& TTopologicalLink<FTopologicalVertex>::GetInfo(FInfoEntity& Info) const
+{
+	return FEntity::GetInfo(Info)
+		.Add(TEXT("active Entity"), ActiveEntity)
+		.Add(TEXT("twin Entities"), TwinsEntities);
+}
+
+FInfoEntity& FVertexLink::GetInfo(FInfoEntity& Info) const
+{
+	return TTopologicalLink<FTopologicalVertex>::GetInfo(Info)
+		.Add(TEXT("barycenter"), Barycenter);
+}
+#endif
diff --git a/Engine/Source/Runtime/Datasmith/CADKernel/Private/CADKernel/UI/Display.cpp b/Engine/Source/Runtime/Datasmith/CADKernel/Private/CADKernel/UI/Display.cpp
new file mode 100644
--- /dev/null
+++ b/Engine/Source/Runtime/Datasmith/CADKernel/Private/CADKernel/UI/Display.cpp
@@ -0,0 +1,1044 @@
+// Copyright Epic Games, Inc. All Rights Reserved.
+
+#include "CADKernel/UI/Display.h"
+
+#include "CADKernel/Core/Chrono.h"
+#include "CADKernel/Core/Group.h"
+#include "CADKernel/Core/System.h"
+#include "CADKernel/Geo/Curves/RestrictionCurve.h"
+#include "CADKernel/Geo/Curves/BezierCurve.h"
+#include "CADKernel/Geo/Curves/NURBSCurve.h"
+#include "CADKernel/Geo/Curves/RestrictionCurve.h"
+#include "CADKernel/Geo/Sampler/SamplerOnChord.h"
+#include "CADKernel/Geo/Sampler/SamplerOnParam.h"
+#include "CADKernel/Geo/Surfaces/BezierSurface.h"
+#include "CADKernel/Geo/Surfaces/NURBSSurface.h"
+#include "CADKernel/Geo/Surfaces/Surface.h"
+#include "CADKernel/Math/AABB.h"
+#include "CADKernel/Math/Boundary.h"
+#include "CADKernel/Math/Plane.h"
+#include "CADKernel/Mesh/Structure/EdgeMesh.h"
+#include "CADKernel/Mesh/Structure/ModelMesh.h"
+#include "CADKernel/Mesh/Structure/FaceMesh.h"
+#include "CADKernel/Mesh/Structure/VertexMesh.h"
+#include "CADKernel/Topo/Body.h"
+#include "CADKernel/Topo/TopologicalEdge.h"
+#include "CADKernel/Topo/Model.h"
+#include "CADKernel/Topo/Shell.h"
+#include "CADKernel/Topo/TopologicalFace.h"
+
+
+namespace CADKernel
+{
+	void Open3DDebugSession(FString name, const TArray<FIdent>& IdArray)
+	{
+		FSystem::Get().GetVisu()->Open3DDebugSession(*name, IdArray);
+	}
+
+	void Close3DDebugSession()
+	{
+		FSystem::Get().GetVisu()->Close3DDebugSession();
+	}
+
+	void Wait(bool bMakeWait)
+	{
+		if (bMakeWait)
+		{
+#ifdef CADKERNEL_DEV
+			system("PAUSE");
+			printf("Hello\n");
+#endif
+		}
+	}
+
+	void Open3DDebugSegment(FIdent Ident)
+	{
+		FSystem::Get().GetVisu()->Open3DDebugSegment(Ident);
+	}
+
+	void Close3DDebugSegment()
+	{
+		FSystem::Get().GetVisu()->Close3DDebugSegment();
+	}
+
+	void FlushVisu()
+	{
+		FSystem::Get().GetVisu()->UpdateViewer();
+	}
+
+	void DrawElement(int32 Dimension, TArray<FPoint>& Points, EVisuProperty Property)
+	{
+		FSystem::Get().GetVisu()->DrawElement(Dimension, Points, Property);
+	}
+
+	void DrawMesh(const TSharedPtr<FMesh>& Mesh)
+	{
+		FSystem::Get().GetVisu()->DrawMesh(Mesh->GetId());
+	}
+
+	void Draw(const FLinearBoundary& Boundary, const TSharedPtr<FRestrictionCurve>& Curve, EVisuProperty Property)
+	{
+		if (Boundary.IsDegenerated())
+		{
+			return;
+		}
+
+		bool bShowOrientation = FSystem::Get().GetVisu()->GetParameters()->bDisplayCADOrient;
+
+		TArray<FPoint> Polyline;
+		Curve->GetDiscretizationPoints<FPoint>(Curve->GetBoundary(), EOrientation::Front, Polyline);
+		Draw(Polyline, Property);
+
+		if (bShowOrientation)
+		{
+			double Length = 0;
+			for (int32 Index = 1; Index < Polyline.Num(); ++Index)
+			{
+				Length += Polyline[Index].Distance(Polyline[Index - 1]);
+			}
+
+			double Coordinate = (Boundary.Max + Boundary.Min) / 2.;
+			FCurvePoint Point;
+			Curve->EvaluatePoint(Coordinate, Point, 1);
+
+			double Height = Length / 20.0;
+			double Base = Height / 2;
+
+			DrawQuadripode(Height, Base, Point.Point, Point.Gradient);
+		}
+	}
+
+	void Draw2D(const TSharedPtr<FCurve>& Curve, const FLinearBoundary& Boundary, EVisuProperty Property)
+	{
+		if (Boundary.IsDegenerated())
+		{
+			return;
+		}
+
+		double DiscVisu = ((double)(FSystem::Get().GetVisu()->GetParameters()->ChordError)) / 10.;
+		bool bVoirOrientation = FSystem::Get().GetVisu()->GetParameters()->bDisplayCADOrient;
+
+		FPolyline2D Polyline;
+		FCurve2DSamplerOnChord Sampler(*Curve, Boundary, DiscVisu, Polyline);
+		Sampler.Sample();
+
+		Draw(Polyline.GetPoints(), Property);
+
+		if (bVoirOrientation)
+		{
+			double Length = Polyline.GetLength(Boundary);
+
+			double Coordinate = (Boundary.Max + Boundary.Min) / 2.;
+			FCurvePoint Point;
+			Curve->EvaluatePoint(Coordinate, Point, 1);
+
+			double Height = Length / 20.0;
+			double Base = Height / 2;
+
+			DrawQuadripode(Height, Base, Point.Point, Point.Gradient);
+		}
+	}
+
+	void Draw3D(const TSharedPtr<FCurve>& Curve, const FLinearBoundary& Boundary, EVisuProperty Property)
+	{
+		if (Boundary.IsDegenerated())
+		{
+			return;
+		}
+
+		double DiscVisu = FSystem::Get().GetVisu()->GetParameters()->ChordError;
+		bool bVoirOrientation = FSystem::Get().GetVisu()->GetParameters()->bDisplayCADOrient;
+
+		FPolyline3D Polyline;
+		FCurveSamplerOnChord Sampler(*Curve, Boundary, DiscVisu, Polyline);
+		Sampler.Sample();
+
+		Draw(Polyline.GetPoints(), Property);
+
+		if (bVoirOrientation)
+		{
+			double Length = Polyline.GetLength(Boundary);
+
+			double Coordinate = (Boundary.Max + Boundary.Min) / 2.;
+			FCurvePoint Point;
+			Curve->EvaluatePoint(Coordinate, Point, 1);
+
+			double Height = Length / 20.0;
+			double Base = Height / 2;
+
+			DrawQuadripode(Height, Base, Point.Point, Point.Gradient);
+		}
+	}
+
+	void Draw(const TSharedPtr<FCurve>& Curve, const FLinearBoundary& Boundary, EVisuProperty Property)
+	{
+		if (Curve->GetDimension() == 3)
+		{
+			Draw3D(Curve, Boundary, Property);
+		}
+		else
+		{
+			Draw2D(Curve, Boundary, Property);
+		}
+	}
+
+
+	void Draw(const TSharedPtr<FCurve>& Curve, EVisuProperty Property)
+	{
+		const FLinearBoundary& Bounds = Curve->GetBoundary();
+		Draw(Curve, Bounds, Property);
+	}
+
+	void DrawQuadripode(double Height, double Base, FPoint& Center, FPoint& InDirection)
+	{
+		FPoint Direction = InDirection;
+		Direction.Normalize();
+
+		FPoint Normal(-Direction[1], Direction[0], 0.0);
+
+		double UVNorm = Normal.Length();
+		Normal /= UVNorm;
+
+		FPoint BiNormal(-Direction[2] * Normal[1], Direction[2] * Normal[0], UVNorm);
+
+		FPoint Point0 = Center;
+		FPoint PointBase = Point0 - Direction * Height;
+
+		Normal *= Base;
+		BiNormal *= Base;
+
+		FPoint Point1 = PointBase + Normal;
+		FPoint Point2 = PointBase + BiNormal;
+		FPoint Point3 = PointBase - Normal;
+		FPoint Point4 = PointBase - BiNormal;
+
+		TArray<FPoint> Polygone;
+		Polygone.Reserve(7);
+
+		Polygone.Add(Point1);
+		Polygone.Add(Point0);
+		Polygone.Add(Point2);
+		Polygone.Add(Point0);
+		Polygone.Add(Point3);
+		Polygone.Add(Point0);
+		Polygone.Add(Point4);
+
+		Draw(Polygone);
+	}
+
+	void DisplayEntity(const TSharedPtr<FEntity>& Entity)
+	{
+		if (!Entity.IsValid())
+		{
+			FMessage::Printf(Log, TEXT("Entity not found"));
+			return;
+		}
+
+		FTimePoint StartTime = FChrono::Now();
+
+		FProgress Progress;
+
+		F3DDebugSession GraphicSession(FString::Printf(TEXT("%s %d"), Entity->GetTypeName(), Entity->GetId()), { Entity->GetId() });
+
+		switch (Entity->GetEntityType())
+		{
+		case EEntity::TopologicalVertex:
+			Display((const TSharedPtr<FTopologicalVertex>) StaticCastSharedPtr<FTopologicalVertex>(Entity));
+			break;
+		case EEntity::Curve:
+			Display(StaticCastSharedPtr<FCurve>(Entity));
+			break;
+		case EEntity::Surface:
+			Display(StaticCastSharedPtr<FSurface>(Entity));
+			break;
+		case EEntity::TopologicalFace:
+			Display(StaticCastSharedPtr<FTopologicalFace>(Entity));
+			break;
+		case EEntity::TopologicalLoop:
+			Display(StaticCastSharedPtr<FTopologicalLoop>(Entity));
+			break;
+		case EEntity::Shell:
+			Display(StaticCastSharedPtr<FShell>(Entity));
+			break;
+		case EEntity::Body:
+			Display(StaticCastSharedPtr<FBody>(Entity));
+			break;
+		case EEntity::TopologicalEdge:
+			Display(StaticCastSharedPtr<FTopologicalEdge>(Entity));
+			break;
+		case EEntity::Model:
+			Display(StaticCastSharedPtr<FModel>(Entity));
+			break;
+		case EEntity::Group:
+			Display(StaticCastSharedPtr<FGroup>(Entity));
+			break;
+		case EEntity::Mesh:
+			DisplayMesh(StaticCastSharedPtr<FMesh>(Entity));
+			break;
+		case EEntity::MeshModel:
+			Display(StaticCastSharedPtr<FModelMesh>(Entity));
+			break;
+		default:
+			FMessage::Printf(Log, TEXT("Unable to display Entity of type %s"), FEntity::GetTypeName(Entity->GetEntityType()));
+		}
+
+		FDuration DisplayDuration = FChrono::Elapse(StartTime);
+		FChrono::PrintClockElapse(Log, TEXT("  "), TEXT("Display "), DisplayDuration);
+
+	}
+
+	void DisplayEntity2D(const TSharedPtr<FEntity>& Entity)
+	{
+		if (!Entity.IsValid())
+		{
+			FMessage::Printf(Log, TEXT("Entity not found"));
+			return;
+		}
+
+		FTimePoint StartTime = FChrono::Now();
+
+		FProgress Progress;
+
+		F3DDebugSession GraphicSession(FString::Printf(TEXT("%s %d"), Entity->GetTypeName(), Entity->GetId()), { Entity->GetId() });
+
+		switch (Entity->GetEntityType())
+		{
+		case EEntity::Surface:
+			Display2D(StaticCastSharedPtr<FSurface>(Entity));
+			break;
+		case EEntity::TopologicalFace:
+			Display2D(StaticCastSharedPtr<FTopologicalFace>(Entity));
+			break;
+		case EEntity::TopologicalLoop:
+			Display2D(StaticCastSharedPtr<FTopologicalLoop>(Entity));
+			break;
+		case EEntity::TopologicalEdge:
+			Display2D(StaticCastSharedPtr<FTopologicalEdge>(Entity));
+			break;
+		default:
+			FMessage::Printf(Log, TEXT("Unable to display Entity of type %s"), FEntity::GetTypeName(Entity->GetEntityType()));
+		}
+
+		FDuration DisplayDuration = FChrono::Elapse(StartTime);
+		FChrono::PrintClockElapse(Log, TEXT("  "), TEXT("Display "), DisplayDuration);
+
+	}
+
+	void Display(const FPlane& Plane, FIdent Ident)
+	{
+		FPoint Normal = Plane.GetNormal();
+		if (Normal.Length() < SMALL_NUMBER)
+		{
+			return;
+		}
+
+		FPoint UAxis;
+		// Find the first frame axis non collinear to the plane normal
+		for (int32 Index = 0; Index < 3; ++Index)
+		{
+			FPoint Axis = FPoint::ZeroPoint;
+			Axis[Index] = 1;
+			UAxis = Normal ^ Axis;
+			if (UAxis.Length() > SMALL_NUMBER)
+			{
+				break;
+			}
+		}
+		FPoint VAxis = UAxis ^ Normal;
+
+		UAxis.Normalize();
+		UAxis *= 10;
+		VAxis.Normalize();
+		VAxis *= 10;
+
+		TArray<FPoint> Points;
+		Points.Init(FPoint(), 5);
+		const FPoint& Point = Plane.GetPoint();
+		Points[0] = Point + UAxis + VAxis;
+		Points[1] = Point + UAxis - VAxis;
+		Points[2] = Point - UAxis - VAxis;
+		Points[3] = Point - UAxis + VAxis;
+		Points[4] = Point + UAxis + VAxis;
+
+		F3DDebugSegment G(Ident);
+		DrawElement(2, Points);
+	}
+
+	void DisplayEdgeCriteriaGrid(int32 EdgeId, const TArray<FPoint>& Points3D)
+	{
+		FString Name = FString::Printf(TEXT("Edge Grid %d"), EdgeId);
+		F3DDebugSession G(Name);
+		{
+			F3DDebugSession _(TEXT("CriteriaGrid Point 3d"));
+			for (int32 Index = 0; Index < Points3D.Num(); Index += 2)
+			{
+				DisplayPoint(Points3D[Index]);
+			}
+		}
+		{
+			F3DDebugSession _(TEXT("CriteriaGrid IntermediateU"));
+			for (int32 Index = 1; Index < Points3D.Num(); Index += 2)
+			{
+				DisplayPoint(Points3D[Index], EVisuProperty::ControlPoint);
+			}
+		}
+	}
+
+	void DisplayAABB(const FAABB& Aabb, FIdent Ident)
+	{
+		TFunction<void(int32, FPoint&, FPoint&)> GetWire = [&](int32 Segment, FPoint& A, FPoint& B)
+		{
+			int32 Corner1;
+			int32 Corner2;
+
+			if (Segment < 4)
+			{
+				Corner1 = Segment * 2;
+				Corner2 = Corner1 + 1;
+			}
+			else if (Segment < 8)
+			{
+				if (Segment < 6)
+				{
+					Corner1 = Segment - 4;
+				}
+				else
+				{
+					Corner1 = Segment - 2;
+				}
+				Corner2 = Corner1 + 2;
+			}
+			else if (Segment < 12)
+			{
+				Corner1 = Segment - 8;
+				Corner2 = Corner1 + 4;
+			}
+			else
+			{
+				Corner1 = 0;
+				Corner2 = 1;
+				ERROR_NOT_EXPECTED;
+			}
+			A = Aabb.GetCorner(Corner1);
+			B = Aabb.GetCorner(Corner2);
+		};
+
+		F3DDebugSegment G(Ident);
+		for (int32 SegmentIndex = 0; SegmentIndex < 12; SegmentIndex++)
+		{
+			TArray<FPoint> Points;
+			Points.SetNum(2);
+			GetWire(SegmentIndex, Points[0], Points[1]);
+			Draw(Points);
+		}
+	}
+
+	void DisplayAABB2D(const FAABB2D& aabb, FIdent Ident)
+	{
+		FPoint A(aabb.GetMin().U, aabb.GetMin().V, 0.);
+		FPoint B(aabb.GetMin().U, aabb.GetMax().V, 0.);
+		FPoint C(aabb.GetMax().U, aabb.GetMax().V, 0.);
+		FPoint D(aabb.GetMax().U, aabb.GetMin().V, 0.);
+
+		TArray<FPoint> Points;
+		Points.Add(A);
+		Points.Add(B);
+		Points.Add(C);
+		Points.Add(D);
+		Points.Add(A);
+		Draw(Points);
+	}
+
+	void Display(const TSharedPtr<FTopologicalVertex>& Vertex)
+	{
+		F3DDebugSegment G(Vertex->GetId());
+		DrawPoint(Vertex->GetCoordinates());
+	}
+
+	void Display(const TSharedPtr<FCurve>& Curve)
+	{
+		F3DDebugSegment G(Curve->GetId());
+		Draw(Curve, Curve->GetBoundary());
+	}
+
+	void Display(const TSharedPtr<FShell>& Shell)
+	{
+		Draw(Shell);
+	}
+
+	void Draw(const TSharedPtr<FShell>& Shell)
+	{
+		FProgress Progress(Shell->GetFaces().Num(), TEXT("Display Shell"));
+		for (const FOrientedFace& Face : Shell->GetFaces())
+		{
+			Display(Face.Entity);
+		}
+	}
+
+	void Display(const TSharedPtr<FBody>& Body)
+	{
+		FProgress Progress(Body->GetShells().Num(), TEXT("Display Body"));
+		for (TSharedPtr<FShell> Shell : Body->GetShells())
+		{
+			if (!Shell.IsValid())
+			{
+				continue;
+			}
+			Draw(Shell);
+		}
+	}
+
+	void Display(const TSharedPtr<FSurface>& Surface)
+	{
+		FProgress Progress(TEXT("Display Surface"));
+
+		int32 IsoUNum = FSystem::Get().GetVisu()->GetParameters()->IsoUNumber;
+		int32 IsoVNum = FSystem::Get().GetVisu()->GetParameters()->IsoVNumber;
+
+		const double VisuSag = FSystem::Get().GetVisu()->GetParameters()->ChordError;
+
+		FPolyline3D Polyline;
+		FIsoCurve3DSamplerOnChord Sampler(*Surface.Get(), VisuSag, Polyline);
+
+		TFunction <void(const int32, const EIso)> DrawIsos = [&](int32 IsoCount, const EIso IsoType)
+		{
+			const FLinearBoundary& CurveBounds = Surface->GetBoundary().Get(IsoType == EIso::IsoU ? EIso::IsoV : EIso::IsoU);
+			const FLinearBoundary& Bounds = Surface->GetBoundary().Get(IsoType);
+
+			double Coordinate = Bounds.Min;
+			IsoCount++;
+			const double Step = (Bounds.Max - Bounds.Min) / IsoCount;
+
+			for (int32 Index = 0; Index <= IsoCount; Index++)
+			{
+				Polyline.Empty();
+				Sampler.Set(IsoType, Coordinate, CurveBounds);
+				Sampler.Sample();
+
+				Draw(Polyline.GetPoints(), EVisuProperty::Iso);
+				Coordinate += Step;
+			}
+		};
+
+		DrawIsos(IsoVNum, EIso::IsoV);
+		DrawIsos(IsoUNum, EIso::IsoU);
+	}
+
+	void Display2D(const TSharedPtr<FSurface>& Surface)
+	{
+		double StepU, StepV;
+
+		F3DDebugSegment G(Surface->GetId());
+
+		int32 IsoUCount = FSystem::Get().GetVisu()->GetParameters()->IsoUNumber;
+		int32 IsoVCount = FSystem::Get().GetVisu()->GetParameters()->IsoVNumber;
+		FSurfacicBoundary Bounds = Surface->GetBoundary();
+
+		StepU = (Bounds.UVBoundaries[EIso::IsoU].Max - Bounds.UVBoundaries[EIso::IsoU].Min) / (IsoUCount + 1);
+		StepV = (Bounds.UVBoundaries[EIso::IsoV].Max - Bounds.UVBoundaries[EIso::IsoV].Min) / (IsoVCount + 1);
+
+		for (int32 iIso = 0; iIso <= (IsoUCount + 1); iIso++)
+		{
+			FPoint2D StartPoint(Bounds.UVBoundaries[EIso::IsoU].Min + (iIso)*StepU, Bounds.UVBoundaries[EIso::IsoV].Min);
+			FPoint2D EndPoint(Bounds.UVBoundaries[EIso::IsoU].Min + (iIso)*StepU, Bounds.UVBoundaries[EIso::IsoV].Max);
+			DrawSegment(StartPoint, EndPoint, EVisuProperty::Iso);
+		}
+
+		for (int32 iIso = 0; iIso <= (IsoVCount + 1); iIso++)
+		{
+			FPoint2D StartPoint(Bounds.UVBoundaries[EIso::IsoU].Min, Bounds.UVBoundaries[EIso::IsoV].Min + (iIso)*StepV);
+			FPoint2D EndPoint(Bounds.UVBoundaries[EIso::IsoU].Max, Bounds.UVBoundaries[EIso::IsoV].Min + (iIso)*StepV);
+			DrawSegment(StartPoint, EndPoint, EVisuProperty::Iso);
+		}
+	}
+
+	void DisplayIsoCurve(const TSharedPtr<FSurface>& Surface, double Coordinate, EIso IsoType)
+	{
+		const double VisuSag = FSystem::Get().GetVisu()->GetParameters()->ChordError;
+
+		F3DDebugSegment G(Surface->GetId());
+
+		const FLinearBoundary& CurveBounds = Surface->GetBoundary().Get(IsoType == EIso::IsoU ? EIso::IsoV : EIso::IsoU);
+
+		FPolyline3D Polyline;
+		FIsoCurve3DSamplerOnChord Sampler(*Surface.Get(), VisuSag, Polyline);
+		Sampler.Set(IsoType, Coordinate, CurveBounds);
+		Draw(Polyline.GetPoints(), EVisuProperty::Iso);
+	}
+
+
+
+	void DisplayControlPolygon(const TSharedPtr<FCurve>& Entity)
+	{
+		TFunction<void(const TArray<FPoint>&)> DisplayHull = [](const TArray<FPoint>& Poles)
+		{
+			for (const FPoint& Pole : Poles)
+			{
+				Display(Pole, EVisuProperty::GreenPoint);
+			}
+
+			for (int32 Index = 1; Index < Poles.Num(); Index++)
+			{
+				FPoint Segment = Poles[Index] - Poles[Index - 1];
+				DisplaySegment(Poles[Index - 1] + Segment * 0.1, Poles[Index] - Segment * 0.1, 0, EVisuProperty::GreenCurve);
+			}
+		};
+
+		F3DDebugSegment GraphicSegment(Entity->GetId());
+		if (Entity->GetCurveType() == ECurve::Bezier)
+		{
+			const TSharedPtr<FBezierCurve>& Bezier = StaticCastSharedPtr<FBezierCurve>(Entity);
+			DisplayHull(Bezier->GetPoles());
+			return;
+		}
+		if (Entity->GetCurveType() == ECurve::Nurbs)
+		{
+			const TSharedPtr<FNURBSCurve>& Nurbs = StaticCastSharedPtr<FNURBSCurve>(Entity);
+			const TArray<FPoint>& Poles = Nurbs->GetPoles();
+			DisplayHull(Nurbs->GetPoles());
+			return;
+		}
+	}
+
+	void DisplayControlPolygon(const TSharedPtr<FSurface>& Entity)
+	{
+		TFunction<void(const TArray<FPoint>&, int32, int32)> DisplayHull = [](const TArray<FPoint>& Poles, int32 PoleUNum, int32 PoleVNum)
+		{
+			for (int32 Index = 0; Index < Poles.Num(); Index++)
+			{
+				Display(Poles[Index], EVisuProperty::GreenPoint);
+			}
+
+			for (int32 IndexV = 0, Index = 0; IndexV < PoleVNum; IndexV++)
+			{
+				Index = Index + 1;
+				for (int32 IndexU = 1; IndexU < PoleUNum; IndexU++, Index++)
+				{
+					FPoint Segment = Poles[Index] - Poles[Index - 1];
+					DisplaySegment(Poles[Index - 1] + Segment * 0.1, Poles[Index] - Segment * 0.1, 0, EVisuProperty::GreenCurve);
+				}
+			}
+
+			for (int32 IndexU = 0; IndexU < PoleUNum; IndexU++)
+			{
+				for (int32 IndexV = 1, Index = PoleUNum + IndexU; IndexV < PoleVNum; IndexV++, Index += PoleUNum)
+				{
+					FPoint Segment = Poles[Index] - Poles[Index - PoleUNum];
+					DisplaySegment(Poles[Index - PoleUNum] + Segment * 0.1, Poles[Index] - Segment * 0.1, 0, EVisuProperty::GreenCurve);
+				}
+			}
+		};
+
+		F3DDebugSegment GraphicSegment(Entity->GetId());
+		if (Entity->GetSurfaceType() == ESurface::Bezier)
+		{
+			const TSharedPtr<FBezierSurface>& Bezier = StaticCastSharedPtr<FBezierSurface>(Entity);
+			DisplayHull(Bezier->GetPoles(), Bezier->GetUDegre() + 1, Bezier->GetVDegre() + 1);
+			return;
+		}
+		if (Entity->GetSurfaceType() == ESurface::Nurbs)
+		{
+			const TSharedPtr<FNURBSSurface>& Nurbs = StaticCastSharedPtr<FNURBSSurface>(Entity);
+			DisplayHull(Nurbs->GetPoles(), Nurbs->GetPoleCount(EIso::IsoU), Nurbs->GetPoleCount(EIso::IsoV));
+			return;
+		}
+	}
+
+	void Display(const TSharedPtr<FTopologicalFace>& Face)
+	{
+		F3DDebugSegment GraphicSegment(Face->GetId());
+		Draw(Face);
+	}
+
+	void Display2D(const TSharedPtr<FTopologicalFace>& Face)
+	{
+		F3DDebugSegment GraphicSegment(Face->GetId());
+		Draw2D(Face);
+	}
+
+	void Draw(const TSharedPtr<FTopologicalFace>& Face)
+	{
+		{
+			for (const TSharedPtr<FTopologicalLoop>& Loop : Face->GetLoops())
+			{
+				for (const FOrientedEdge& Edge : Loop->GetEdges())
+				{
+					EVisuProperty Property = EVisuProperty::BlueCurve;
+					switch (Edge.Entity->GetTwinsEntityCount())
+					{
+					case 1:
+						Property = EVisuProperty::BorderEdge;
+						break;
+					case 2:
+						Property = EVisuProperty::BlueCurve;
+						break;
+					default:
+						Property = EVisuProperty::NonManifoldEdge;
+					}
+					Draw(Edge.Entity, Property);
+				}
+			}
+		}
+
+		DrawIsoCurves(Face);
+	}
+
+	void Draw2D(const TSharedPtr<FTopologicalFace>& Face)
+	{
+		TArray<TArray<FPoint2D>> BoundaryApproximation;
+		Face->Get2DLoopSampling(BoundaryApproximation);
+
+		for (const TArray<FPoint2D>& Boudary : BoundaryApproximation)
+		{
+			Draw(Boudary, EVisuProperty::BlueCurve);
+		}
+
+		TFunction <void(const int32, const EIso)> DrawIsos = [&](int32 IsoCount, const EIso IsoType)
+		{
+			EIso OtherIso = IsoType == EIso::IsoU ? EIso::IsoV : EIso::IsoU;
+			const FLinearBoundary& Bounds = Face->GetBoundary().Get(IsoType);
+			const FLinearBoundary& CurveBounds = Face->GetBoundary().Get(OtherIso);
+
+
+			double Coordinate = Bounds.Min;
+			IsoCount++;
+			const double Step = (Bounds.Max - Bounds.Min) / IsoCount;
+
+			for (int32 iIso = 1; iIso < IsoCount; iIso++)
+			{
+				Coordinate += Step;
+
+				TArray<double> Intersections;
+				FindLoopIntersectionsWithIso(IsoType, Coordinate, BoundaryApproximation, Intersections);
+
+				FPoint2D Start;
+				FPoint2D End;
+
+				Start[IsoType] = Coordinate;
+				End[IsoType] = Coordinate;
+
+				if (Intersections.Num() % 2 != 0)
+				{
+					Start[OtherIso] = Intersections[Intersections.Num() - 1];
+					End[OtherIso] = CurveBounds.GetMax();
+					DrawSegment(Start, End, EVisuProperty::YellowCurve);
+					Intersections.Pop();
+				}
+
+				for (int32 ISection = 0; ISection < Intersections.Num(); ISection += 2)
+				{
+					Start[OtherIso] = Intersections[ISection];
+					End[OtherIso] = Intersections[ISection + 1];
+					DrawSegment(Start, End, EVisuProperty::Iso);
+				}
+			}
+		};
+
+		int32 IsoUCount = FSystem::Get().GetVisu()->GetParameters()->IsoUNumber;
+		int32 IsoVCount = FSystem::Get().GetVisu()->GetParameters()->IsoVNumber;
+
+		DrawIsos(IsoUCount, EIso::IsoU);
+		DrawIsos(IsoVCount, EIso::IsoV);
+	}
+
+
+	void DrawIsoCurves(const TSharedPtr<FTopologicalFace>& Face)
+	{
+		TArray<TArray<FPoint2D>> BoundaryApproximation;
+		Face->Get2DLoopSampling(BoundaryApproximation);
+
+		const double VisuSag = FSystem::Get().GetVisu()->GetParameters()->ChordError;
+
+		FPolyline3D Polyline;
+		const FSurface& Surface = Face->GetCarrierSurface().Get();
+		FIsoCurve3DSamplerOnChord Sampler(Surface, VisuSag, Polyline);
+
+		TFunction <void(const int32, const EIso)> DrawIsos = [&](int32 IsoCount, const EIso IsoType)
+		{
+			const FLinearBoundary& Bounds = Face->GetBoundary().Get(IsoType);
+
+			EIso Other = IsoType == EIso::IsoU ? EIso::IsoV : EIso::IsoU;
+
+			double Coordinate = Bounds.Min;
+			IsoCount++;
+			const double Step = (Bounds.Max - Bounds.Min) / IsoCount;
+
+			for (int32 iIso = 1; iIso < IsoCount; iIso++)
+			{
+				Coordinate += Step;
+
+				TArray<double> Intersections;
+				FindLoopIntersectionsWithIso(IsoType, Coordinate, BoundaryApproximation, Intersections);
+				int32 IntersectionCount = Intersections.Num();
+				if (IntersectionCount == 0)
+				{
+					return;
+				}
+
+				FLinearBoundary CurveBounds(Intersections[0], Intersections.Last());
+
+				Polyline.Empty();
+				Sampler.Set(IsoType, Coordinate, CurveBounds);
+				Sampler.Sample();
+
+				if (IntersectionCount % 2 != 0)
+				{
+					TArray<FPoint> SubPolyline;
+					FLinearBoundary Boundary(Intersections[IntersectionCount - 1], CurveBounds.GetMax());
+					Polyline.GetSubPolyline(Boundary, EOrientation::Front, SubPolyline);
+					Draw(SubPolyline, EVisuProperty::YellowCurve);
+					Intersections.Pop();
+					IntersectionCount--;
+				}
+
+				for (int32 ISection = 0; ISection < IntersectionCount; ISection += 2)
+				{
+					TArray<FPoint> SubPolyline;
+					FLinearBoundary Boundary(Intersections[ISection], Intersections[ISection + 1]);
+					Polyline.GetSubPolyline(Boundary, EOrientation::Front, SubPolyline);
+					Draw(SubPolyline, EVisuProperty::Iso);
+				}
+			}
+		};
+
+		int32 IsoUCount = FSystem::Get().GetVisu()->GetParameters()->IsoUNumber;
+		int32 IsoVCount = FSystem::Get().GetVisu()->GetParameters()->IsoVNumber;
+
+		DrawIsos(IsoUCount, EIso::IsoU);
+		DrawIsos(IsoVCount, EIso::IsoV);
+	}
+
+	void Display(const TSharedPtr<FTopologicalEdge>& Edge, EVisuProperty Property)
+	{
+		F3DDebugSegment GraphicSegment(Edge->GetId());
+		Draw(Edge, Property);
+	}
+
+	void Display2D(const TSharedPtr<FTopologicalEdge>& Edge, EVisuProperty Property)
+	{
+		F3DDebugSegment GraphicSegment(Edge->GetId());
+		TArray<FPoint2D> Polyline;
+		Edge->GetCurve()->GetDiscretizationPoints(Edge->GetBoundary(), EOrientation::Front, Polyline);
+		DisplayPolyline(Polyline, Property);
+	}
+
+	void Draw(const TSharedPtr<FTopologicalEdge>& Edge, EVisuProperty Property)
+	{
+		const FLinearBoundary& Boundary = Edge->GetBoundary();
+		Draw(Boundary, Edge->GetCurve(), Property);
+	}
+
+	void Display(const TSharedPtr<FModel>& Model)
+	{
+		FProgress MainProgress(2, TEXT("Display model"));
+		{
+			const TArray<TSharedPtr<FBody>>& Bodies = Model->GetBodyList();
+			FProgress BodyProgress((int32)Bodies.Num(), TEXT("Bodies"));
+			for (const TSharedPtr<FBody>& Body : Bodies)
+			{
+				Display(Body);
+			}
+		}
+
+		{
+			const TArray<TSharedPtr<FTopologicalFace>>& Surfaces = Model->GetSurfaceList();
+			FProgress SurfaceProgress((int32)Surfaces.Num(), TEXT("surfaces"));
+			for (const TSharedPtr<FTopologicalFace>& Surface : Surfaces)
+			{
+				Display(Surface);
+			}
+		}
+	}
+
+	void Display(const TSharedPtr<FGroup>& Group)
+	{
+		TArray<TSharedPtr<FEntity>> Entities;
+
+		Group->GetValidEntities(Entities);
+
+		FProgress Progress(Entities.Num());
+		for (const TSharedPtr<FEntity>& Entity : Entities)
+		{
+			DisplayEntity(Entity);
+		}
+	}
+
+	void DisplayMesh(const TSharedPtr<FMesh>& Mesh)
+	{
+		ensureCADKernel(Mesh.IsValid());
+		switch (Mesh->GetGeometricEntity()->GetEntityType())
+		{
+		case EEntity::TopologicalFace:
+			DisplayMesh(StaticCastSharedRef<FFaceMesh>(Mesh.ToSharedRef()));
+			break;
+		case EEntity::TopologicalEdge:
+			DisplayMesh(StaticCastSharedRef<FEdgeMesh>(Mesh.ToSharedRef()));
+			break;
+		case EEntity::TopologicalVertex:
+			DisplayMesh(StaticCastSharedRef<FVertexMesh>(Mesh.ToSharedRef()));
+			break;
+		default:
+			break;
+		}
+	}
+
+	void DisplayMesh(const TSharedRef<FFaceMesh>& Mesh)
+	{
+		TMap<int32, const FPoint*> NodeIdToCoordinates;
+		Mesh->GetNodeIdToCoordinates(NodeIdToCoordinates);
+
+		TArray<FPoint> Points;
+		Points.SetNum(3);
+		double EdgeMeanLength = 0;
+		const TArray<int32>& TriangleIndices = Mesh->TrianglesVerticesIndex;
+		const TArray<int32>& VertexIndices = Mesh->VerticesGlobalIndex;
+		for (int32 Index = 0; Index < TriangleIndices.Num();)
+		{
+			Points[0] = *NodeIdToCoordinates[VertexIndices[TriangleIndices[Index++]]];
+			Points[1] = *NodeIdToCoordinates[VertexIndices[TriangleIndices[Index++]]];
+			Points[2] = *NodeIdToCoordinates[VertexIndices[TriangleIndices[Index++]]];
+			DrawElement(2, Points, EVisuProperty::Element);
+			DrawSegment(Points[0], Points[1], EVisuProperty::EdgeMesh);
+			DrawSegment(Points[1], Points[2], EVisuProperty::EdgeMesh);
+			DrawSegment(Points[2], Points[0], EVisuProperty::EdgeMesh);
+			EdgeMeanLength += Points[0].Distance(Points[1]);
+			EdgeMeanLength += Points[1].Distance(Points[2]);
+			EdgeMeanLength += Points[2].Distance(Points[0]);
+		}
+		EdgeMeanLength /= 10 * TriangleIndices.Num();
+
+		for (const int32& Index : VertexIndices)
+		{
+			F3DDebugSegment GraphicSegment(Index);
+			DrawPoint(*NodeIdToCoordinates[Index], EVisuProperty::NodeMesh);
+		}
+
+		bool test = FSystem::Get().GetVisu()->GetParameters()->bDisplayNormals;
+		if (FSystem::Get().GetVisu()->GetParameters()->bDisplayNormals)
+		{
+			double NormalLength = FSystem::Get().GetVisu()->GetParameters()->NormalLength;
+			const TArray<FPoint>& Normals = Mesh->Normals;
+			for (int32 Index = 0; Index < VertexIndices.Num(); ++Index)
+			{
+				F3DDebugSegment GraphicSegment(Index);
+				FPoint Normal = Normals[Index];
+				Normal.Normalize();
+				Normal *= NormalLength;
+				DrawSegment(*NodeIdToCoordinates[VertexIndices[Index]], *NodeIdToCoordinates[VertexIndices[Index]] + Normal, EVisuProperty::EdgeMesh);
+			}
+		}
+	}
+
+	void DisplayMesh(const TSharedRef<FEdgeMesh>& Mesh)
+	{
+		const TSharedRef<FModelMesh> MeshModel = Mesh->GetMeshModel();
+
+		const TArray<int32>& NodeIds = Mesh->EdgeVerticesIndex;
+		const TArray<FPoint>& NodeCoordinates = Mesh->GetNodeCoordinates();
+
+		int32 StartNodeId = NodeIds[0];
+		int32 LastNodeId = NodeIds.Last();
+
+		{
+			F3DDebugSegment GraphicSegment(Mesh->GetGeometricEntity()->GetId());
+			if (NodeCoordinates.Num() == 0)
+			{
+				DrawSegment(MeshModel->GetMeshOfVertexNodeId(StartNodeId)->GetNodeCoordinates()[0], MeshModel->GetMeshOfVertexNodeId(LastNodeId)->GetNodeCoordinates()[0], EVisuProperty::EdgeMesh);
+			}
+			else
+			{
+				DrawSegment(MeshModel->GetMeshOfVertexNodeId(StartNodeId)->GetNodeCoordinates()[0], NodeCoordinates[0], EVisuProperty::EdgeMesh);
+				if (NodeCoordinates.Num() > 1)
+				{
+					for (int32 Index = 0; Index < NodeCoordinates.Num() - 1; Index++)
+					{
+						DrawSegment(NodeCoordinates[Index], NodeCoordinates[Index + 1], EVisuProperty::EdgeMesh);
+					}
+				}
+				DrawSegment(NodeCoordinates.Last(), MeshModel->GetMeshOfVertexNodeId(LastNodeId)->GetNodeCoordinates()[0], EVisuProperty::EdgeMesh);
+			}
+		}
+
+		{
+			F3DDebugSegment GraphicSegment(StartNodeId);
+			DrawPoint(MeshModel->GetMeshOfVertexNodeId(StartNodeId)->GetNodeCoordinates()[0], EVisuProperty::NodeMesh);
+		}
+
+		{
+			F3DDebugSegment GraphicSegment(LastNodeId);
+			DrawPoint(MeshModel->GetMeshOfVertexNodeId(LastNodeId)->GetNodeCoordinates()[0], EVisuProperty::NodeMesh);
+		}
+
+		if (NodeCoordinates.Num() > 1)
+		{
+			for (int32 Index = 0; Index < NodeCoordinates.Num(); Index++)
+			{
+				F3DDebugSegment GraphicSegment(NodeIds[Index]);
+				DrawPoint(NodeCoordinates[Index], EVisuProperty::NodeMesh);
+			}
+		}
+	}
+
+	void DisplayMesh(const TSharedRef<FVertexMesh>& Mesh)
+	{
+		DisplayPoint(Mesh->GetNodeCoordinates()[0], EVisuProperty::NodeMesh, Mesh->GetId());
+	}
+
+	void Display(const TSharedPtr<FModelMesh>& MeshModel)
+	{
+		for (TSharedPtr<FMesh> Mesh : MeshModel->GetMeshes())
+		{
+			if (Mesh.IsValid())
+			{
+				DisplayMesh(Mesh.ToSharedRef());
+			}
+		}
+	}
+
+	void DisplaySegment(const FPoint& Point1, const FPoint& Point2, FIdent Ident, EVisuProperty Property)
+	{
+		F3DDebugSegment G(Ident);
+		DrawSegment(Point1, Point2, Property);
+	}
+
+	void DisplaySegment(const FPoint2D& Point1, const FPoint2D& Point2, FIdent Ident, EVisuProperty Property, bool bWithOrientation)
+	{
+		F3DDebugSegment G(Ident);
+		if (bWithOrientation)
+		{
+			DrawSegmentOrientation(Point1, Point2, Property);
+		}
+		DrawSegment(Point1, Point2, Property);
+	}
+
+	void DisplaySegment(const FPointH& Point1, const FPointH& Point2, FIdent Ident, EVisuProperty Property)
+	{
+		F3DDebugSegment G(Ident);
+		DrawSegment(Point1, Point2, Property);
+	}
+
+	void DisplayLoop(const TSharedPtr<FTopologicalFace>& Surface)
+	{
+		for (const TSharedPtr<FTopologicalLoop> Loop : Surface->GetLoops())
+		{
+			Display(Loop);
+		}
+	}
+
+
+	void Display(const TSharedPtr<FTopologicalLoop>& Loop)
+	{
+		for (const FOrientedEdge& Edge : Loop->GetEdges())
+		{
+			Display(Edge.Entity);
+		}
+	}
+
+	void Display2D(const TSharedPtr<FTopologicalLoop>& Loop)
+	{
+		for (const FOrientedEdge& Edge : Loop->GetEdges())
+		{
+			Display2D(Edge.Entity);
+		}
+	}
+} // namespace CADKernel
diff --git a/Engine/Source/Runtime/Datasmith/CADKernel/Private/CADKernel/UI/Progress.cpp b/Engine/Source/Runtime/Datasmith/CADKernel/Private/CADKernel/UI/Progress.cpp
new file mode 100644
--- /dev/null
+++ b/Engine/Source/Runtime/Datasmith/CADKernel/Private/CADKernel/UI/Progress.cpp
@@ -0,0 +1,81 @@
+// Copyright Epic Games, Inc. All Rights Reserved.
+
+#include "CADKernel/UI/Progress.h"
+#include "CADKernel/UI/Message.h"
+#include "CADKernel/Utils/Util.h"
+#include "CADKernel/Core/System.h"
+
+using namespace CADKernel;
+
+void FProgressManager::SetCurrentProgress(FProgress* InNewCurrent)
+{
+	if (!RootProgress)
+	{
+		RootProgress = InNewCurrent;
+		CurrentProgress = RootProgress;
+	}
+	else if (InNewCurrent == nullptr)
+	{
+		RootProgress = nullptr;
+		CurrentProgress = nullptr;
+	}
+
+	CurrentProgress = InNewCurrent;
+	Update();
+}
+
+double FProgressManager::GetProgression() const
+{
+	if (!RootProgress) 
+	{
+		return 0.;
+	}
+	return RootProgress->GetProgression();
+}
+
+FString FProgressManager::GetCurrentStep() const
+{
+	FString CurrentStepStr;
+
+	for (FProgress* Progress = GetCurrentProgress(); Progress; Progress = Progress->GetParent())
+	{
+		if (!Progress->GetName().IsEmpty())
+		{
+			CurrentStepStr = Progress->GetName() + TEXT(" -> ") + CurrentStepStr;
+		}
+	}
+	CurrentStepStr.RemoveFromEnd(TEXT(" -> "));
+	return CurrentStepStr;
+}
+
+
+FProgress::FProgress(int32 InStepCount, const FString& InStepName)
+: Name(InStepName)
+, StepCount(InStepCount)
+{
+	Parent = FSystem::Get().GetProgressManager().GetCurrentProgress();
+
+	if (Parent)
+	{
+		Parent->AddUnderlying(this);
+	}
+
+	FSystem::Get().GetProgressManager().SetCurrentProgress(this);
+}
+
+FProgress::~FProgress()
+{
+	if (Parent)
+	{
+		Parent->UnderlyingFinished(this);
+	}
+
+	FSystem::Get().GetProgressManager().SetCurrentProgress(Parent);
+	FSystem::Get().GetProgressManager().Update();
+}
+
+void FProgress::Increase(int32 StepSize)
+{
+	Progression += StepSize;
+	FSystem::Get().GetProgressManager().Update();
+}
diff --git a/Engine/Source/Runtime/Datasmith/CADKernel/Private/CADKernel/UI/Visu.cpp b/Engine/Source/Runtime/Datasmith/CADKernel/Private/CADKernel/UI/Visu.cpp
new file mode 100644
--- /dev/null
+++ b/Engine/Source/Runtime/Datasmith/CADKernel/Private/CADKernel/UI/Visu.cpp
@@ -0,0 +1,36 @@
+// Copyright Epic Games, Inc. All Rights Reserved.
+
+#include "CADKernel/UI/Visu.h"
+
+namespace CADKernel
+{
+	const TCHAR* VisuPropertyNames[] = 
+	{
+		TEXT("Bounding Box"),              
+		TEXT("Points"),
+		TEXT("Elements"),
+		TEXT("Curves"),
+		TEXT("Iso"),
+		TEXT("Control Lines"),
+		TEXT("Edges"),
+		TEXT("Nodes"),
+		TEXT("Border Edges"),
+		TEXT("Multiple Edges"),
+		TEXT("Geometrical Borders"),
+		TEXT("Control Points"),
+
+		TEXT("YellowPoint"),
+		TEXT("YellowCurve"),
+		TEXT("BluePoint"),
+		TEXT("BlueCurve"),
+		TEXT("RedPoint"),
+		TEXT("RedCurve"),
+		TEXT("PurplePoint"),
+		TEXT("PurpleCurve"),
+		TEXT("GreenPoint"),
+		TEXT("GreenCurve"),
+		TEXT("PinkPoint"),
+		TEXT("PinkCurve"),
+		nullptr
+	};
+}
diff --git a/Engine/Source/Runtime/Datasmith/CADKernel/Private/CADKernel/Utils/Util.cpp b/Engine/Source/Runtime/Datasmith/CADKernel/Private/CADKernel/Utils/Util.cpp
new file mode 100644
--- /dev/null
+++ b/Engine/Source/Runtime/Datasmith/CADKernel/Private/CADKernel/Utils/Util.cpp
@@ -0,0 +1,8 @@
+// Copyright Epic Games, Inc. All Rights Reserved.
+
+#include "CADKernel/Utils/Util.h"
+
+namespace CADKernel
+{
+}
+
diff --git a/Engine/Source/Runtime/Datasmith/CADKernel/Public/CADKernel/Core/Entity.h b/Engine/Source/Runtime/Datasmith/CADKernel/Public/CADKernel/Core/Entity.h
new file mode 100644
--- /dev/null
+++ b/Engine/Source/Runtime/Datasmith/CADKernel/Public/CADKernel/Core/Entity.h
@@ -0,0 +1,239 @@
+// Copyright Epic Games, Inc. All Rights Reserved.
+#pragma once
+
+#include "CADKernel/Core/CADKernelArchive.h"
+#include "CADKernel/Core/HaveStates.h"
+#include "CADKernel/Core/OrientedEntity.h"
+#include "CADKernel/Core/Types.h"
+#include "CADKernel/Geo/GeoEnum.h"
+
+#include "Serialization/Archive.h"
+
+#ifdef CADKERNEL_DEV
+#include "Toolkit/Core/InfoEntity.h"
+#endif
+
+namespace CADKernel
+{
+	enum class EEntity : uint8
+	{
+		NullEntity = 0,
+
+		Curve,
+		Surface,
+
+		EdgeLink,
+		VertexLink,
+		TopologicalEdge,
+		TopologicalFace,
+		TopologicalLink,
+		TopologicalLoop,
+		TopologicalVertex,
+		Shell,
+		Body,
+		Model,
+
+		MeshModel,
+		Mesh,
+
+		Group,
+		Criterion,
+		Property,
+		
+		EntityTypeEnd
+	};
+
+	class FEntity;
+	class FSystem;
+	class FDatabase;
+	class FSession;
+	class FInfoEntity;
+
+	class CADKERNEL_API FEntity : public TSharedFromThis<FEntity>, public FHaveStates
+	{
+		friend class FDatabase;
+
+	protected:
+		static const TCHAR* TypesNames[];
+		FIdent Id = 0;
+
+	public:
+		virtual ~FEntity();
+
+		template<typename OtherEntity, typename... InArgTypes>
+		static TSharedRef<OtherEntity> MakeShared(InArgTypes&&... Args)
+		{
+			OtherEntity* Entity = new OtherEntity(Forward<InArgTypes>(Args)...);
+			TSharedRef<OtherEntity> NewShared = MakeShareable<OtherEntity>(Entity);
+			return NewShared;
+		}
+
+		template<typename OtherEntity>
+		static TSharedRef<OtherEntity> MakeShared(FCADKernelArchive& Archiver)
+		{
+			OtherEntity* Entity = new OtherEntity(Archiver);
+			TSharedRef<OtherEntity> NewShared = MakeShareable<OtherEntity>(Entity);
+			Archiver.AddEntityFromArchive(NewShared);
+			return NewShared;
+		}
+
+		/**
+		 * To Serialize a TSharedPtr<FEntity>, the Id of the FEntity is saved in the archive (ArchiveId). 
+		 * At the deserialization of the archive, a map is build linking ArchiveId to the TSharedPtr of the deserialized entities.
+		 * At the deserialization of a TSharedPtr<FEntity>, the ArchiveId of the TSharedPtr<FEntity> is get. The map return the associated TSharedPtr.
+		 * If the Entity is not yet deserialized, the TSharedPtr is add to a waiting list to be set as soon as the associated Entity is deserialized
+		 * See FDatabase for more details 
+		 */
+		static void SerializeIdent(FCADKernelArchive& Ar, TSharedPtr<FEntity>& Entity, bool bSaveSelection = true);
+		static void SerializeIdent(FCADKernelArchive& Ar, TWeakPtr<FEntity>& Entity, bool bSaveSelection = true);
+
+		template<typename EntityType>
+		static void SerializeIdent(FCADKernelArchive& Ar, TSharedPtr<EntityType>& Entity, bool bSaveSelection = true)
+		{
+			SerializeIdent(Ar, (TSharedPtr<FEntity>&) Entity, bSaveSelection);
+		}
+
+		template<typename EntityType>
+		static void SerializeIdent(FCADKernelArchive& Ar, TWeakPtr<EntityType>& Entity, bool bSaveSelection = true)
+		{
+			SerializeIdent(Ar, (TWeakPtr<FEntity>&) Entity, bSaveSelection);
+		}
+
+		/**
+		 * SerializeIdent of each TSharedPtr<FEntity> of the array
+		 */
+		static void SerializeIdents(FCADKernelArchive& Ar, TArray<TWeakPtr<FEntity>>& Array, bool bSaveSelection = true);
+		static void SerializeIdents(FCADKernelArchive& Ar, TArray<TSharedPtr<FEntity>>& Array, bool bSaveSelection = true);
+		static void SerializeIdents(FCADKernelArchive& Ar, TArray<TOrientedEntity<FEntity>>& Array);
+
+
+		template<typename EntityType>
+		static void SerializeIdents(FCADKernelArchive& Ar, TArray<TWeakPtr<EntityType>>& Array, bool bSaveSelection = true)
+		{
+			SerializeIdents(Ar, (TArray<TWeakPtr<FEntity>>&) Array, bSaveSelection);
+		}
+
+		template<typename EntityType>
+		static void SerializeIdents(FCADKernelArchive& Ar, TArray<TSharedPtr<EntityType>>& Array, bool bSaveSelection = true)
+		{
+			SerializeIdents(Ar, (TArray<TSharedPtr<FEntity>>&) Array, bSaveSelection);
+		}
+
+		/**
+		 * Spawn Ident on each TSharedPtr<FEntity> of the array
+		 */
+		static void SpawnIdentOnEntities(TArray<TSharedPtr<FEntity>>& Array, FDatabase& Database)
+		{
+			for (TSharedPtr<FEntity>& Entity : Array)
+			{
+				Entity->SpawnIdent(Database);
+			}
+		}
+
+		static void SpawnIdentOnEntities(TArray<TOrientedEntity<FEntity>>& Array, FDatabase& Database)
+		{
+			for (TOrientedEntity<FEntity>& Entity : Array)
+			{
+				Entity.Entity->SpawnIdent(Database);
+			}
+		}
+			
+		template<typename EntityType>
+		static void SpawnIdentOnEntities(TArray<TSharedPtr<EntityType>>& Array, FDatabase& Database)
+		{
+			SpawnIdentOnEntities((TArray<TSharedPtr<FEntity>>&) Array, Database);
+		}
+
+		/**
+		 * Reset Processed flags Recursively on each TSharedPtr<FEntity> of the array
+		 */
+		static void ResetMarkersRecursivelyOnEntities(TArray<TWeakPtr<FEntity>>& Array)
+		{
+			for (TWeakPtr<FEntity>& Entity : Array)
+			{
+				Entity.Pin()->ResetMarkersRecursively();
+			}
+		}
+
+		static void ResetMarkersRecursivelyOnEntities(TArray<TSharedPtr<FEntity>>& Array)
+		{
+			for (TSharedPtr<FEntity>& Entity : Array)
+			{
+				Entity->ResetMarkersRecursively();
+			}
+		}
+
+		static void ResetMarkersRecursivelyOnEntities(TArray<TOrientedEntity<FEntity>>& Array)
+		{
+			for (TOrientedEntity<FEntity>& OrientedEntity : Array)
+			{
+				OrientedEntity.Entity->ResetMarkersRecursively();
+			}
+		}
+
+		template<typename EntityType>
+		static void ResetMarkersRecursivelyOnEntities(TArray<TSharedPtr<EntityType>>& Array)
+		{
+			ResetMarkersRecursivelyOnEntities((TArray<TSharedPtr<FEntity>>&) Array);
+		}
+
+		/**
+		 * Serialization of a FEntity. Each class derived from FEntity has to override this method (and call the direct base class override method first)
+		 * E.g.
+		 * 	class FEntityXXX : public FEntityXX
+		 *  {
+		 *		virtual void Serialize(FCADKernelArchive& Ar) override
+		 * 		{
+		 *			FEntityXX::Serialize(Ar);
+		 *			...
+		 * 		}
+		 *	}
+		 */
+		virtual void Serialize(FCADKernelArchive& Ar)
+		{
+			Ar << Id;
+			FHaveStates::Serialize(Ar);
+		}
+
+		/**
+		 * Deserialize the next entity in the archive.
+		 */
+		static TSharedPtr<FEntity> Deserialize(FCADKernelArchive& Ar);
+
+		static const TCHAR* GetTypeName(EEntity Type);
+
+		virtual EEntity GetEntityType() const = 0;
+		const TCHAR* GetTypeName() const
+		{
+			return GetTypeName(GetEntityType());
+		}
+
+		const FIdent& GetId() const
+		{
+			return Id;
+		}
+
+#ifdef CADKERNEL_DEV
+		void InfoEntity() const;
+		virtual FInfoEntity& GetInfo(FInfoEntity& EntityInfo) const;
+#endif
+
+		virtual void SpawnIdent(FDatabase& Database)
+		{
+			if (!FEntity::SetId(Database))
+			{
+				return;
+			}
+		}
+
+		virtual void ResetMarkersRecursively()
+		{
+			ResetMarkers();
+		}
+
+	protected:
+
+		bool SetId(FDatabase& Database);
+	};
+}
+
diff --git a/Engine/Source/Runtime/Datasmith/CADKernel/Public/CADKernel/Core/EntityGeom.h b/Engine/Source/Runtime/Datasmith/CADKernel/Public/CADKernel/Core/EntityGeom.h
new file mode 100644
--- /dev/null
+++ b/Engine/Source/Runtime/Datasmith/CADKernel/Public/CADKernel/Core/EntityGeom.h
@@ -0,0 +1,51 @@
+// Copyright Epic Games, Inc. All Rights Reserved.
+#pragma once
+
+#include "CADKernel/Core/CADKernelArchive.h"
+#include "CADKernel/Core/Entity.h"
+
+#include "CADKernel/Math/MatrixH.h"
+
+class FString;
+
+namespace CADKernel
+{
+	class CADKERNEL_API FEntityGeom : public FEntity
+	{
+	friend class FCoreTechBridge;
+
+	protected:
+#ifdef CORETECHBRIDGE_DEBUG
+		FIdent CtKioId = 0;
+#endif
+	public:
+		FEntityGeom()
+			: FEntity()
+		{
+		}
+
+
+		virtual TSharedPtr<FEntityGeom> ApplyMatrix(const FMatrixH& InMatrix) const
+		{
+			return TSharedPtr<FEntityGeom>();
+		}
+
+		virtual void Display(const FString& Name) const
+		{
+		}
+
+		virtual void Serialize(FCADKernelArchive& Ar) override
+		{
+			FEntity::Serialize(Ar);
+#ifndef CORETECHBRIDGE_DEBUG
+			FIdent CtKioId = 0;
+#endif
+			Ar << CtKioId;
+		}
+
+#ifdef CADKERNEL_DEV
+		virtual FInfoEntity& GetInfo(FInfoEntity& Info) const override;
+#endif
+	};
+}
+
diff --git a/Engine/Source/Runtime/Datasmith/CADKernel/Public/CADKernel/Core/HaveStates.h b/Engine/Source/Runtime/Datasmith/CADKernel/Public/CADKernel/Core/HaveStates.h
new file mode 100644
--- /dev/null
+++ b/Engine/Source/Runtime/Datasmith/CADKernel/Public/CADKernel/Core/HaveStates.h
@@ -0,0 +1,145 @@
+// Copyright Epic Games, Inc. All Rights Reserved.
+#pragma once
+
+#include "CADKernel/Core/Types.h"
+#include "CADKernel/Core/CADKernelArchive.h"
+
+enum class EHaveStates : uint16
+{
+	None                = 0x0000u,  // No flags.
+								    
+	IsApplyCriteria     = 0x0001u,  // used for FTopologicalFace, FEdge
+	IsRemoved           = 0x0001u,  // used for FThinZone2D
+								    
+	IsMeshed            = 0x0002u,  // used for FTopologicalFace, FEdge
+	IsInner             = 0x0002u,  // used for FEdgeSegment, FShell, 
+	IsFirstSide         = 0x0002u,  // used for FThinZoneSide
+	FirstSideClosed     = 0x0002u,  // used for FThinZone2D
+								    
+	IsBackOriented      = 0x0004u,  // used for FTopologicalFace
+	ThinPeak            = 0x0004u,  // used for FEdge 
+	SecondSideClosed    = 0x0004u,  // used for FThinZone2D
+								    
+	ThinZone            = 0x0008u,  // used for FTopologicalFace, FEdge, FEdgeSegment  
+								    
+	Degenerated         = 0x0010u,  // used for FEdge, FGrid
+								    
+	IsDeleted           = 0x0020u,  // used for all class
+								    
+	HasMarker1          = 0x1000u,  // used for all class
+	HasMarker2          = 0x2000u,  // used for all class
+	HasMarker3          = 0x4000u,  // used for all class
+	IsProcess4			= 0x8000u,  // used for all class
+	AllMarkers          = 0xF000u,  // used for all class
+
+	All                 = 0xFFu
+};
+
+ENUM_CLASS_FLAGS(EHaveStates);
+
+namespace CADKernel
+{
+	class CADKERNEL_API FHaveStates
+	{
+	protected:
+		mutable EHaveStates States;
+
+	public:
+
+		FHaveStates()
+			: States(EHaveStates::None)
+		{};
+
+		void Serialize(FCADKernelArchive& Ar)
+		{
+			Ar << States;
+		}
+
+		void ResetElementStatus()
+		{
+			States = EHaveStates::None;
+		}
+
+		bool HasMarker1() const
+		{
+			return ((States & EHaveStates::HasMarker1) == EHaveStates::HasMarker1);
+		}
+
+		bool HasMarker2() const
+		{
+			return ((States & EHaveStates::HasMarker2) == EHaveStates::HasMarker2);
+		}
+
+		bool HasMarker3() const
+		{
+			return ((States & EHaveStates::HasMarker3) == EHaveStates::HasMarker3);
+		}
+
+		void SetMarker1() const
+		{
+			States |= EHaveStates::HasMarker1;
+		}
+
+		void SetMarker2() const
+		{
+			States |= EHaveStates::HasMarker2;
+		}
+
+		void SetMarker3() const
+		{
+			States |= EHaveStates::HasMarker3;
+		}
+
+		void ResetMarker1() const
+		{
+			States &= ~EHaveStates::HasMarker1;
+		}
+
+		void ResetMarker2() const
+		{
+			States &= ~EHaveStates::HasMarker2;
+		}
+
+		void ResetMarker3() const
+		{
+			States &= ~EHaveStates::HasMarker3;
+		}
+
+		void ResetMarkers() const
+		{
+			States &= ~EHaveStates::AllMarkers;
+		}
+
+		bool IsDeleted() const
+		{
+			return ((States & EHaveStates::IsDeleted) == EHaveStates::IsDeleted);
+		}
+
+		void SetDeleted() const
+		{
+			States |= EHaveStates::IsDeleted;
+		}
+
+		void ResetDeleted() const
+		{
+			States &= ~EHaveStates::IsDeleted;
+		}
+
+		virtual bool IsDegenerated() const
+		{
+			return ((States & EHaveStates::Degenerated) == EHaveStates::Degenerated);
+		}
+
+		virtual void SetAsDegenerated() const
+		{
+			States |= EHaveStates::Degenerated;
+		}
+
+		virtual void ResetDegenerated() const
+		{
+			States &= ~EHaveStates::Degenerated;
+		}
+
+	};
+}
+
diff --git a/Engine/Source/Runtime/Datasmith/CADKernel/Public/CADKernel/Core/Metadata.h b/Engine/Source/Runtime/Datasmith/CADKernel/Public/CADKernel/Core/Metadata.h
new file mode 100644
--- /dev/null
+++ b/Engine/Source/Runtime/Datasmith/CADKernel/Public/CADKernel/Core/Metadata.h
@@ -0,0 +1,143 @@
+// Copyright Epic Games, Inc. All Rights Reserved.
+#pragma once
+
+#include "CADKernel/Core/Types.h"
+#include "CADKernel/Math/MatrixH.h"
+
+namespace CADKernel
+{
+	class FLayerData;
+	class FRGBColorData;
+	class FRGBAColorData;
+
+	template<typename ValueType>
+	class TMetadata;
+
+	class CADKERNEL_API FMetadata
+	{
+	public:
+		FMetadata() = default;
+
+		template<typename ValueType>
+		static TSharedRef<FMetadata> Create(const ValueType& Value)
+		{
+			return MakeShared<TMetadata<ValueType>>(Value);
+		}
+
+		static TSharedRef<FMetadata> CreateLayer(const int32 LayerId, const FString& LayerName, const int32 LayerFlag)
+		{
+			TSharedRef<FLayerData> Layer = MakeShared<FLayerData>(LayerId, LayerName, LayerFlag);
+			return StaticCastSharedRef<FMetadata>(Layer);
+		}
+
+		static TSharedRef<FMetadata> CreateRGBColor(const double RedValue, const double GreenValue, const double BlueValue)
+		{
+			TSharedRef<FRGBColorData> Layer = MakeShared<FRGBColorData>(RedValue, GreenValue, BlueValue);
+			return StaticCastSharedRef<FMetadata>(Layer);
+		}
+
+		static TSharedRef<FMetadata> CreateRGBAColor(const double RedValue, const double GreenValue, const double BlueValue, const double AlphaValue)
+		{
+			TSharedRef<FRGBAColorData> Layer = MakeShared<FRGBAColorData>(RedValue, GreenValue, BlueValue, AlphaValue);
+			return StaticCastSharedRef<FMetadata>(Layer);
+		}
+	};
+
+	template<typename ValueType>
+	class TMetadata : public FMetadata
+	{
+	protected:
+		ValueType Value;
+
+	public:
+		TMetadata(const ValueType& InValue)
+			: Value(InValue)
+		{
+		}
+		~TMetadata() = default;
+
+		const ValueType& GetValue()
+		{
+			return Value;
+		}
+	};
+
+	class CADKERNEL_API FRGBColorData : public FMetadata
+	{
+	protected:
+		uint8 Red;
+		uint8 Green;
+		uint8 Blue;
+
+	public:
+
+		FRGBColorData(const double RedValue, const double GreenValue, const double BlueValue)
+			: Red(ToUInt8(RedValue))
+			, Green(ToUInt8(GreenValue))
+			, Blue(ToUInt8(BlueValue))
+		{
+		}
+
+		FRGBColorData(const uint8 RedValue, const uint8 GreenValue, const uint8 BlueValue)
+			: Red(RedValue)
+			, Green(GreenValue)
+			, Blue(BlueValue)
+		{
+		}
+	};
+
+	class CADKERNEL_API FRGBAColorData : public FRGBColorData
+	{
+	protected:
+		uint8 Alpha;
+
+	public:
+
+		FRGBAColorData(const double RedValue, const double GreenValue, const double BlueValue, const double AlphaValue)
+			: FRGBColorData(RedValue, GreenValue, BlueValue)
+			, Alpha(ToUInt8(AlphaValue))
+		{
+		}
+
+		FRGBAColorData(const uint8 RedValue, const uint8 GreenValue, const uint8 BlueValue, const uint8 AlphaValue)
+			: FRGBColorData(RedValue, GreenValue, BlueValue)
+			, Alpha(AlphaValue)
+		{
+		}
+	};
+
+	class CADKERNEL_API FLayerData : public FMetadata
+	{
+	protected:
+		int32 Id;
+		FString Name;
+		int32 Flag;
+
+	public:
+
+		FLayerData(const int32 InLayerId, const FString& InLayerName, const int32 InLayerFlag)
+			: Id(InLayerId)
+			, Name(InLayerName)
+			, Flag(InLayerFlag)
+		{
+		}
+
+		const int32 GetId()
+		{
+			return Id;
+		}
+
+		const FString& GetName()
+		{
+			return Name;
+		}
+
+		const int32 GetFlag()
+		{
+			return Flag;
+		}
+
+	};
+}
+
+
diff --git a/Engine/Source/Runtime/Datasmith/CADKernel/Public/CADKernel/Core/MetadataDictionary.h b/Engine/Source/Runtime/Datasmith/CADKernel/Public/CADKernel/Core/MetadataDictionary.h
new file mode 100644
--- /dev/null
+++ b/Engine/Source/Runtime/Datasmith/CADKernel/Public/CADKernel/Core/MetadataDictionary.h
@@ -0,0 +1,74 @@
+// Copyright Epic Games, Inc. All Rights Reserved.
+#pragma once
+
+#include "CADKernel/Core/Types.h"
+
+#include "CADKernel/Core/Metadata.h"
+
+namespace CADKernel
+{
+	class CADKERNEL_API FMetadataDictionary
+	{
+
+	private:
+		TMap<FString, TSharedPtr<FMetadata>> MetadataMap;
+
+		void AddValue(const FString& Name, TSharedPtr<FMetadata>& Attribute)
+		{
+			MetadataMap.Emplace(Name, Attribute);
+		}
+
+	public:
+
+		FMetadataDictionary() = default;
+
+		template<typename ValueType>
+		void Add(const FString& Name, const ValueType& Value)
+		{
+			TSharedRef<FMetadata> NewAttribute = FMetadata::Create(Value);
+			SetOrAdd(Name, NewAttribute);
+		}
+
+		void AddLayer(const FString& Name, const int32 LayerId, const FString& LayerName, const int32 LayerFlag)
+		{
+			TSharedRef<FMetadata> NewAttribute = FMetadata::CreateLayer(LayerId, LayerName, LayerFlag);
+			SetOrAdd(Name, NewAttribute);
+		}
+
+		void AddRGBColor(const FString& Name, const double RedValue, const double GreenValue, const double BlueValue)
+		{
+			TSharedRef<FMetadata> NewAttribute = FMetadata::CreateRGBColor(RedValue, GreenValue, BlueValue);
+			SetOrAdd(Name, NewAttribute);
+		}
+
+		const TSharedPtr<FMetadata> Get(const FString& Name)
+		{
+			TSharedPtr<FMetadata>* Attribute = MetadataMap.Find(Name);
+			return Attribute == nullptr ? TSharedPtr<FMetadata>() : *Attribute;
+		}
+
+		int32 Count() const
+		{
+			return MetadataMap.Num();
+		}
+
+	private:
+
+		void SetOrAdd(const FString Name, const TSharedPtr<FMetadata> NewAttribute)
+		{
+			TSharedPtr<FMetadata>* OldAttribute = MetadataMap.Find(Name);
+			if (OldAttribute != nullptr)
+			{
+				(*OldAttribute).Reset();
+				*OldAttribute = NewAttribute;
+			}
+			else
+			{
+				MetadataMap.Add(Name, NewAttribute);
+			}
+		}
+
+	};
+}
+
+
diff --git a/Engine/Source/Runtime/Datasmith/CADKernel/Public/CADKernel/Core/OrientedEntity.h b/Engine/Source/Runtime/Datasmith/CADKernel/Public/CADKernel/Core/OrientedEntity.h
new file mode 100644
--- /dev/null
+++ b/Engine/Source/Runtime/Datasmith/CADKernel/Public/CADKernel/Core/OrientedEntity.h
@@ -0,0 +1,31 @@
+// Copyright Epic Games, Inc. All Rights Reserved.
+#pragma once
+
+#include "CADKernel/Core/Types.h"
+#include "CADKernel/Geo/GeoEnum.h"
+
+namespace CADKernel
+{
+	template <typename KeyType>
+	class CADKERNEL_API TOrientedEntity
+	{
+	public:
+		TSharedPtr<KeyType> Entity;
+		EOrientation Direction;
+
+		TOrientedEntity(TSharedPtr<KeyType> InEntity, EOrientation InDirection)
+			: Entity(InEntity)
+			, Direction(InDirection)
+		{
+		}
+
+		TOrientedEntity(const TOrientedEntity<KeyType>& OrientiredEntity)
+			: Entity(OrientiredEntity.Entity)
+			, Direction(OrientiredEntity.Direction)
+		{
+		}
+
+		TOrientedEntity() = default;
+	};
+}
+
diff --git a/Engine/Source/Runtime/Datasmith/CADKernel/Public/CADKernel/Core/Parameter.h b/Engine/Source/Runtime/Datasmith/CADKernel/Public/CADKernel/Core/Parameter.h
new file mode 100644
--- /dev/null
+++ b/Engine/Source/Runtime/Datasmith/CADKernel/Public/CADKernel/Core/Parameter.h
@@ -0,0 +1,307 @@
+// Copyright Epic Games, Inc. All Rights Reserved.
+#pragma once
+
+#include "CADKernel/Core/Types.h"
+
+#include "Serialization/Archive.h"
+
+namespace CADKernel
+{
+	class CADKERNEL_API FParameters;
+
+	class CADKERNEL_API FParameterValue
+	{
+	public:
+		union
+		{
+			double Real;
+			int32 Integer;
+			bool Boolean;
+		} Value;
+
+		FString StringValue;
+
+	protected:
+		EValue Type;
+
+	public:
+		FParameterValue(int32 InValue)
+			: Type(EValue::Integer)
+		{
+			*this = InValue;
+		}
+
+		FParameterValue(double InValue)
+			: Type(EValue::Double)
+		{
+			*this = InValue;
+		}
+
+		FParameterValue(const FString& InValue)
+			: Type(EValue::String)
+		{
+			*this = InValue;
+		}
+
+		FParameterValue(bool InValue)
+			: Type(EValue::Boolean)
+		{
+			*this = InValue;
+		}
+
+		FParameterValue(EValue InType = EValue::Integer)
+			: Type(InType)
+		{
+			*this = 0;
+		}
+
+		FParameterValue(const FParameterValue& Other)
+		{
+			Type = Other.Type;
+			Value = Other.Value;
+			StringValue = Other.StringValue;
+		}
+
+		FString ToString() const;
+
+		inline EValue GetType() const
+		{
+			return Type;
+		}
+
+		inline void SetType(EValue NewType)
+		{
+			Type = NewType;
+		}
+
+		inline FParameterValue& operator=(double InValue)
+		{
+			Value.Real = InValue;
+			return *this;
+		}
+
+		inline FParameterValue& operator=(int32 InValue)
+		{
+			Value.Integer = InValue;
+			return *this;
+		}
+
+		inline FParameterValue& operator=(const FString& InValue)
+		{
+			StringValue = InValue;
+			return *this;
+		}
+
+		inline FParameterValue& operator=(bool InValue)
+		{
+			Value.Boolean = InValue;
+			return *this;
+		}
+
+		inline operator double& ()
+		{
+			return Value.Real;
+		}
+
+		inline operator FString& ()
+		{
+			return StringValue;
+		}
+
+		inline operator int32& ()
+		{
+			return Value.Integer;
+		}
+
+		inline operator bool& ()
+		{
+			return Value.Boolean;
+		}
+
+		inline operator const double& () const
+		{
+			return Value.Real;
+		}
+
+		inline operator const FString& () const
+		{
+			return StringValue;
+		}
+
+		inline operator const int32& () const
+		{
+			return Value.Integer;
+		}
+
+		inline operator const bool& () const
+		{
+			return Value.Boolean;
+		}
+
+		bool operator==(const FParameterValue& Other) const;
+	};
+
+	class CADKERNEL_API FParameter
+	{
+	private:
+		FString Name;
+
+		FParameterValue Value;
+		FParameterValue DefaultValue;
+
+		const TCHAR** EnumLabels;
+
+		void AddToParameterMap(FParameter& Parameter, FParameters& Parameters);
+
+	public:
+
+		template<typename ValueType>
+		FParameter(const FString& ParameterName, const ValueType& FirstValue, FParameters& Parameters)
+			: Name(ParameterName)
+			, Value(FirstValue)
+			, DefaultValue(FirstValue)
+			, EnumLabels(nullptr)
+		{
+			AddToParameterMap(*this, Parameters);
+		}
+
+		FParameter(const FString& ParameterName, const TCHAR* v, FParameters& Parameters)
+			: Name(ParameterName)
+			, Value(FString(v))
+			, DefaultValue(FString(v))
+			, EnumLabels(nullptr)
+		{
+			AddToParameterMap(*this, Parameters);
+		}
+
+		FParameter(const FParameter& Other, FParameters& Parameters)
+			: Name(Other.Name)
+			, Value(Other.Value)
+			, DefaultValue(Other.DefaultValue)
+			, EnumLabels(Other.EnumLabels)
+		{
+			AddToParameterMap(*this, Parameters);
+		}
+
+		FParameter()
+			: Name()
+			, Value()
+			, DefaultValue()
+			, EnumLabels()
+		{
+		}
+
+		virtual ~FParameter()
+		{
+		}
+
+		const FParameterValue& GetValue() const
+		{
+			return Value;
+		}
+
+		const FParameterValue& GetDefaultValue() const
+		{
+			return DefaultValue;
+		}
+
+		template<typename ValueType>
+		void SetValue(const ValueType NewValue)
+		{
+			Value = NewValue;
+		}
+
+		template<typename ValueType>
+		void SetDefaultValue(const ValueType NewValue)
+		{
+			DefaultValue = NewValue;
+		}
+
+		template<typename ValueType>
+		inline FParameter& operator= (const ValueType& NewValue)
+		{
+			SetValue(NewValue);
+			return *this;
+		}
+
+		inline FParameter& operator= (const FParameter& Parameter)
+		{
+			SetValue(Parameter.Value);
+			return *this;
+		}
+
+		inline operator int32 () const
+		{
+			return Value.Value.Integer;
+		}
+
+		inline operator double() const
+		{
+			return Value.Value.Real;
+		}
+
+		inline operator bool() const
+		{
+			return Value.Value.Boolean;
+		}
+
+		inline operator const FString& () const
+		{
+			return Value.StringValue;
+		}
+
+		inline bool IsDefault() const 
+		{
+			return Value == DefaultValue;
+		}
+
+		inline const FString& GetName() const
+		{
+			return Name;
+		}
+
+		inline EValue GetType() const
+		{
+			return Value.GetType();
+		}
+
+		FString GetTypeString() const;
+
+		void SetFromString(const FString& str);
+		FString ToString() const;
+	};
+
+	class CADKERNEL_API FAngleParameter
+		: public FParameter
+	{
+	private:
+		double RadianValue;
+		double CosineValue;
+
+	public:
+		FAngleParameter(const FString paramName, double AngleValue, FParameters& Parameters)
+			: FParameter(paramName, AngleValue, Parameters)
+		{
+			RadianValue = FMath::DegreesToRadians(AngleValue);
+			CosineValue = cos(RadianValue);
+		}
+
+		void SetValue(const FParameterValue& NewValue)
+		{
+			FParameter::SetValue(NewValue);
+			RadianValue = FMath::DegreesToRadians((double) NewValue);
+			CosineValue = cos(RadianValue);
+		}
+
+		const double GetAngle() const
+		{
+			return RadianValue;
+		};
+
+		const double GetCosine() const
+		{
+			return CosineValue;
+		};
+	};
+} // namespace CADKernel
+
diff --git a/Engine/Source/Runtime/Datasmith/CADKernel/Public/CADKernel/Core/Parameters.h b/Engine/Source/Runtime/Datasmith/CADKernel/Public/CADKernel/Core/Parameters.h
new file mode 100644
--- /dev/null
+++ b/Engine/Source/Runtime/Datasmith/CADKernel/Public/CADKernel/Core/Parameters.h
@@ -0,0 +1,77 @@
+// Copyright Epic Games, Inc. All Rights Reserved.
+#pragma once
+
+#include "CADKernel/Core/Types.h"
+
+#include "CADKernel/Core/Parameter.h"
+
+#include "Serialization/Archive.h"
+
+namespace CADKernel
+{
+	class FDatabaseFile;
+
+	class CADKERNEL_API FParameters
+	{
+		friend FParameter; 
+	protected:
+
+		void Add(FParameter& NewParameter)
+		{
+			FParameter** Parameter = Map.Find(NewParameter.GetName());
+			if (Parameter)
+			{
+				(*Parameter)->SetValue(NewParameter.GetValue());
+			}
+			else
+			{
+				Map.Add(NewParameter.GetName(), &NewParameter);
+			}
+		}
+
+		TMap<FString, FParameter*> Map;
+
+	public:
+		FParameters(int32 ParametersNum = 3)
+		{
+			Map.Reserve(ParametersNum);
+		}
+
+		FParameters& operator=(const FParameters& Other)
+		{
+			for (const TPair<FString, FParameter*>& Parameter : Other.Map)
+			{
+				Map.Add(Parameter.Key, Parameter.Value);
+			}
+			return *this;
+		}
+
+		virtual ~FParameters() = default;
+
+		FParameter* GetByName(const FString& Name)
+		{
+			FParameter** Parameter = Map.Find(Name);
+			return (Parameter != nullptr) ? *Parameter : nullptr;
+		}
+
+		void SetFromString(const FString& ParameterStr);
+
+		FString ToString(bool bOnlyChanged = false) const;
+
+		void PrintParameterList();
+
+		int32 Count()
+		{
+			return Map.Num();
+		}
+
+		const TMap<FString, FParameter*>& GetMap() const 
+		{
+			return Map;
+		}
+
+		friend FArchive& operator<<(FArchive& Ar, FParameters& Database)
+		{}
+
+	};
+} 
diff --git a/Engine/Source/Runtime/Datasmith/CADKernel/Public/CADKernel/Core/Session.h b/Engine/Source/Runtime/Datasmith/CADKernel/Public/CADKernel/Core/Session.h
new file mode 100644
--- /dev/null
+++ b/Engine/Source/Runtime/Datasmith/CADKernel/Public/CADKernel/Core/Session.h
@@ -0,0 +1,109 @@
+// Copyright Epic Games, Inc. All Rights Reserved.
+#pragma once
+
+#include "CADKernel/Core/Types.h"
+#include "CADKernel/Core/Entity.h"
+#include "CADKernel/Core/Database.h"
+
+namespace CADKernel
+{
+	class FEntity;
+	class CADKERNEL_API FSession
+	{
+		friend FEntity;
+		friend FCADKernelArchive;
+
+	protected:
+
+		double GeometricTolerance;
+		FDatabase Database;
+
+		TSharedPtr<FModel> Model;
+
+	public:
+
+		FSession(double InGeometricTolerance)
+			: GeometricTolerance(InGeometricTolerance)
+		{
+		}
+
+		TSharedRef<FModel> GetModel();
+
+		void Serialize(FCADKernelArchive& Ar)
+		{
+			Ar << GeometricTolerance;
+		}
+
+		FDatabase& GetDatabase()
+		{
+			return Database;
+		}
+
+		double GetGeometricTolerance() const
+		{
+			return GeometricTolerance;
+		}
+
+		/**
+		 * Save the database as a FAchive in a file
+		 * Mandatory: all entity have to have a defined ID
+		 * Use SpawnEntityIdent if needed
+		 */
+		void SaveDatabase(const TCHAR* FilePath);
+
+		/**
+		 * Save a selection and all the dependencies as a FAchive in a file
+		 * Mandatory: all entity have to have a defined ID
+		 * Use SpawnEntityIdent if needed
+		 */
+		void SaveDatabase(const TCHAR* FileName, const TArray<TSharedPtr<FEntity>>& Entities);
+
+		/**
+		 * Save a selection and all the dependencies as a FAchive in a file
+		 * Mandatory: all entity have to have a defined ID
+		 * Use SpawnEntityIdent if needed
+		 */
+		void SaveDatabase(const TCHAR* FileName, const TSharedPtr<FEntity> Entity)
+		{
+			TArray<TSharedPtr<FEntity>> Entities;
+			Entities.Emplace(Entity);
+			SaveDatabase(FileName, Entities);
+		}
+
+		/**
+		 * load and add a database in the current session database
+		 * Entity ID is set for all loaded entities
+		 */
+		void LoadDatabase(const TCHAR* FilePath);
+
+		void Clear()
+		{
+			Model.Reset();
+			Database.Empty();
+		}
+
+		/**
+		 * To be consistent,  all entity to save have to had an Id.
+		 * This method browses all sub entities and set their Id if needed
+		 * @param bForceSpawning If false, the process does not iterate through the children of entities with a defined ID
+		 */
+		uint32 SpawnEntityIdent(const TArray<TSharedPtr<FEntity>>& SelectedEntities, bool bForceSpawning = false)
+		{
+			return Database.SpawnEntityIdent(SelectedEntities, bForceSpawning);
+		}
+
+		uint32 SpawnEntityIdent(TSharedPtr<FEntity>& SelectedEntity, bool bForceSpawning)
+		{
+			return Database.SpawnEntityIdent(SelectedEntity, bForceSpawning);
+		}
+
+		template<typename PointType>
+		uint32 SpawnEntityIdent(TSharedPtr<PointType>& SelectedEntity, bool bForceSpawning = false)
+		{
+			return Database.SpawnEntityIdent((TSharedPtr<FEntity>&) SelectedEntity, bForceSpawning);
+		}
+
+	};
+
+} // namespace CADKernel
+
diff --git a/Engine/Source/Runtime/Datasmith/CADKernel/Public/CADKernel/Core/System.h b/Engine/Source/Runtime/Datasmith/CADKernel/Public/CADKernel/Core/System.h
new file mode 100644
--- /dev/null
+++ b/Engine/Source/Runtime/Datasmith/CADKernel/Public/CADKernel/Core/System.h
@@ -0,0 +1,139 @@
+// Copyright Epic Games, Inc. All Rights Reserved.
+#pragma once
+
+#include "CADKernel/Core/Types.h"
+
+#include "CADKernel/UI/Visu.h"
+#include "CADKernel/UI/Console.h"
+#include "CADKernel/UI/Progress.h"
+
+namespace CADKernel
+{
+	class FConsole;
+	class FKernelParameters;
+	class FVisu;
+
+	class FSystem
+	{
+		friend class FDatabase;
+
+	protected:
+
+		FString ProductName = L"CADKernel";
+
+		TSharedRef<FKernelParameters> Parameters;
+
+		FVisu DefaultVisu;
+		FVisu* Viewer;
+
+		FConsole DefaultConsole;
+		FConsole* Console;
+
+		FProgressManager DefaultProgressManager;
+		FProgressManager* ProgressManager;
+
+		static TUniquePtr<FSystem> Instance;
+
+		TSharedPtr<FArchive> LogFile;
+		EVerboseLevel LogLevel;
+		TSharedPtr<FArchive> SpyFile;
+
+		EVerboseLevel VerboseLevel;
+
+	public:
+
+		FSystem();
+
+		void Initialize(bool bIsDll = true, const FString& LogFilePath = TEXT(""), const FString& SpyFilePath = TEXT(""));
+
+		void Shutdown();
+
+		void CloseLogFiles();
+
+		FString GetToolkitVersion() const;
+		FString GetCompilationDate() const;
+
+		FVisu* GetVisu() const
+		{
+			return Viewer;
+		}
+
+		FConsole& GetConsole() const
+		{
+			ensureCADKernel(Console);
+			return *Console;
+		}
+
+		void SetViewer(FVisu* NewViewer)
+		{
+			Viewer = NewViewer;
+		}
+
+		void SetConsole(FConsole* InConsole)
+		{
+			Console = InConsole;
+		}
+
+		FProgressManager& GetProgressManager()
+		{
+			ensureCADKernel(ProgressManager);
+			return *ProgressManager;
+		}
+			
+		void SetProgressManager(FProgressManager* InProgressManager)
+		{
+			ProgressManager = InProgressManager;
+		}
+
+		TSharedRef<FKernelParameters> GetParameters()
+		{
+			return Parameters;
+		}
+
+		EVerboseLevel GetVerboseLevel() const
+		{
+			return VerboseLevel;
+		}
+
+		void SetVerboseLevel(EVerboseLevel Level)
+		{
+			VerboseLevel = Level;
+		}
+
+		void DefineLogFile(const FString& LogFilePath, EVerboseLevel Level = Log);
+		void DefineSpyFile(const FString& SpyFilePath);
+
+		void InitializeCADKernel();
+
+		TSharedPtr<FArchive> GetLogFile() const
+		{
+			return LogFile;
+		}
+
+		EVerboseLevel GetLogLevel() const
+		{
+			return LogLevel;
+		}
+
+		TSharedPtr<FArchive> GetSpyFile() const
+		{
+			return SpyFile;
+		}
+
+		static FSystem& Get()
+		{
+			if (!Instance.IsValid())
+			{
+				Instance = MakeUnique<FSystem>();
+			}
+			return *Instance;
+		}
+
+	protected:
+
+		void PrintHeader();
+
+	};
+
+} // namespace CADKernel
+
diff --git a/Engine/Source/Runtime/Datasmith/CADKernel/Public/CADKernel/Core/Types.h b/Engine/Source/Runtime/Datasmith/CADKernel/Public/CADKernel/Core/Types.h
new file mode 100644
--- /dev/null
+++ b/Engine/Source/Runtime/Datasmith/CADKernel/Public/CADKernel/Core/Types.h
@@ -0,0 +1,51 @@
+// Copyright Epic Games, Inc. All Rights Reserved.
+#pragma once
+
+#ifdef DEV_CADKERNEL_LIB
+#include "CADKernel/Utils/UnrealEnvironment.h"
+#else
+#include "CoreMinimal.h"
+#endif
+
+#define SMALL_NUMBER_SQUARE 10e-16
+#define HUGE_VALUE 10e8
+#define HUGE_VALUE_SQUARE 10e16
+
+#define AThird 0.33333333333333333333333333333333
+#define AQuarter 0.25
+#define ASixth 0.16666666666666666666666666666667
+#define AEighth 0.125
+
+typedef uint32 FIdent;
+
+namespace CADKernel
+{
+	enum class EValue : uint8
+	{
+		Entity,
+		OrientedEntity,
+		Point,
+		Matrix,
+		Integer,
+		Double,
+		String,
+		Boolean,
+		Tuple,
+		List,
+		Array
+	};
+
+	enum EVerboseLevel : uint8
+	{
+		NoVerbose = 0,
+		Spy,
+		Log,
+		Debug
+	};
+} // namespace CADKernel
+
+#ifdef DO_ENSURE_CADKERNEL
+#define ensureCADKernel(InExpression) ensure(InExpression)
+#else
+#define ensureCADKernel(InExpression) {}
+#endif
diff --git a/Engine/Source/Runtime/Datasmith/CADKernel/Public/CADKernel/Geo/Curves/Curve.h b/Engine/Source/Runtime/Datasmith/CADKernel/Public/CADKernel/Geo/Curves/Curve.h
new file mode 100644
--- /dev/null
+++ b/Engine/Source/Runtime/Datasmith/CADKernel/Public/CADKernel/Geo/Curves/Curve.h
@@ -0,0 +1,241 @@
+// Copyright Epic Games, Inc. All Rights Reserved.
+#pragma once
+
+#include "CADKernel/Core/Types.h"
+
+#include "CADKernel/Core/EntityGeom.h"
+#include "CADKernel/Core/Parameters.h"
+#include "CADKernel/Geo/GeoPoint.h"
+#include "CADKernel/Geo/GeoEnum.h"
+#include "CADKernel/Math/Boundary.h"
+#include "CADKernel/UI/Message.h"
+#include "CADKernel/Utils/Cache.h"
+
+namespace CADKernel
+{
+	class FSurface;
+
+	struct FCurvePoint;
+	struct FCurvePoint2D;
+
+	class CADKERNEL_API FCurve : public FEntityGeom
+	{
+		friend class FEntity;
+
+	protected:
+
+		/** Geometric Tolerance of the curve in its space i.e. Tolerance 3D for 3d curve, Tolerance 2D for 2d curve */
+		double Tolerance;
+
+		mutable TCache<double> Tolerance1D;
+		mutable TCache<double> GlobalLength;
+		FLinearBoundary Boundary;
+
+		int8 Dimension;
+
+		FCurve(const double InTolerance, int8 InDimension = 3)
+			: FEntityGeom()
+			, Tolerance(InTolerance)
+			, Boundary()
+			, Dimension(InDimension)
+		{
+		}
+
+		FCurve(double InTolerance, const FLinearBoundary& InBounds, int8 InDimension = 3)
+			: FEntityGeom()
+			, Tolerance(InTolerance)
+			, Boundary(InBounds)
+			, Dimension(InDimension)
+		{
+		}
+
+		FCurve() = default;
+
+	public:
+
+		virtual void Serialize(FCADKernelArchive& Ar) override
+		{
+			// Curve's type is serialize because it is used to instantiate the correct entity on deserialization (@see Deserialize(FCADKernelArchive& Archive)) 
+			if (Ar.IsSaving())
+			{
+				ECurve CurveType = GetCurveType();
+				Ar << CurveType;
+			}
+			FEntityGeom::Serialize(Ar);
+			Ar << Tolerance;
+			Ar << Dimension;
+			Ar << Boundary;
+		}
+
+		/**
+		 * Specific method for curve family to instantiate the correct derived class of FCurve 
+		 */
+		static TSharedPtr<FEntity> Deserialize(FCADKernelArchive& Archive);
+		
+#ifdef CADKERNEL_DEV
+		virtual FInfoEntity& GetInfo(FInfoEntity&) const override;
+#endif
+
+		int32 GetDimension() const
+		{
+			return Dimension;
+		}
+
+		virtual EEntity GetEntityType() const override
+		{
+			return EEntity::Curve;
+		}
+
+		virtual ECurve GetCurveType() const = 0;
+
+		double GetUMin() const
+		{
+			return Boundary.Min;
+		};
+
+		double GetUMax() const
+		{
+			return Boundary.Max;
+		};
+
+		const FLinearBoundary& GetBoundary() const
+		{
+			return Boundary;
+		}
+
+		virtual TSharedPtr<FEntityGeom> ApplyMatrix(const FMatrixH& InMatrix) const override = 0;
+
+		double GetLength() const;
+
+		/**
+		 * Return the tolerance in the parametric space of the curve i.e Distance(Curve(U), Curve(U + Tolerance1D)) = Tolerance
+		 */
+		double GetParametricTolerance() const;
+
+		/**
+		 * Evaluate exact 3D point of the curve at the input Coordinate
+		 * The function can only be used with 3D curve (Dimension == 3) 
+		 */
+		virtual void EvaluatePoint(double Coordinate, FCurvePoint& OutPoint, int32 DerivativeOrder = 0) const
+		{
+			ensureCADKernel(false);
+		}
+
+		/**
+		 * Evaluate exact 3D point of the curve at the input Coordinate
+		 * The function can only be used with 3D curve (Dimension == 3)
+		 */
+		virtual FPoint EvaluatePoint(double Coordinate) const
+		{
+			FCurvePoint OutPoint;
+			EvaluatePoint(Coordinate, OutPoint);
+			return OutPoint.Point;
+		}
+
+		/**
+		 * Evaluate exact 2D point of the curve at the input Coordinate
+		 * The function can only be used with 2D curve (Dimension == 2)
+		 */
+		virtual void Evaluate2DPoint(double Coordinate, FCurvePoint2D& OutPoint, int32 DerivativeOrder = 0) const
+		{
+			ensureCADKernel(false);
+		}
+
+		/**
+		 * Evaluate exact 3D point of the curve at the input Coordinate
+		 * The function can only be used with 3D curve (Dimension == 3)
+		 */
+		virtual FPoint2D Evaluate2DPoint(double Coordinate) const
+		{
+			FCurvePoint2D OutPoint;
+			Evaluate2DPoint(Coordinate, OutPoint);
+			return OutPoint.Point;
+		}
+
+		/**
+		 * Evaluate exact 2D point of the curve at the input Coordinate
+		 * The function can only be used with 2D curve (Dimension == 2)
+		 */
+		virtual void Evaluate2DPoint(double Coordinate, FPoint2D& OutPoint) const
+		{
+			ensure(Dimension == 2);
+			FCurvePoint2D CurvePoint;
+			Evaluate2DPoint(Coordinate, CurvePoint);
+			OutPoint = CurvePoint.Point;
+		}
+
+		/**
+		 * Evaluate exact 3D points of the curve at the input Coordinates
+		 * The function can only be used with 3D curve (Dimension == 3)
+		 */
+		virtual void EvaluatePoints(const TArray<double>& Coordinates, TArray<FCurvePoint>& OutPoints, int32 DerivativeOrder = 0) const;
+
+		/**
+		 * Evaluate exact 3D points of the curve at the input Coordinates
+		 * The function can only be used with 3D curve (Dimension == 3)
+		 */
+		virtual void EvaluatePoints(const TArray<double>& Coordinates, TArray<FPoint>& OutPoints) const;
+
+		/**
+		 * Evaluate exact 2D points of the curve at the input Coordinates
+		 * The function can only be used with 2D curve (Dimension == 2)
+		 */
+		virtual void Evaluate2DPoints(const TArray<double>& Coordinates, TArray<FPoint2D>& OutPoints) const;
+
+		/**
+		 * Evaluate exact 2D points of the curve at the input Coordinates
+		 * The function can only be used with 2D curve (Dimension == 2)
+		 */
+		virtual void Evaluate2DPoints(const TArray<double>& Coordinates, TArray<FCurvePoint2D>& OutPoints, int32 DerivativeOrder = 0) const;
+
+		void FindNotDerivableCoordinates(int32 DerivativeOrder, TArray<double>& OutNotDerivableCoordinates) const;
+		virtual void FindNotDerivableCoordinates(const FLinearBoundary& InBoundary, int32 DerivativeOrder, TArray<double>& OutNotDerivableCoordinates) const;
+
+
+		/**
+		 * Generate a pre-sampling of the curve saved in OutCoordinates.
+		 * This sampling is light enough to allow a fast computation, precise enough to compute accurately meshing criteria
+		 */
+		void Presample(TArray<double>& OutSampling) const
+		{
+			Presample(Boundary, OutSampling);
+		}
+
+		virtual void Presample(const FLinearBoundary& InBoundary, TArray<double>& OutSampling) const;
+
+		/**
+		 * Make a new curve based on the new bounds.
+		 * If the new bound is nearly equal to the initial bound, no curve is made (return TSharedPtr<FCurve>())
+		 */
+		virtual TSharedPtr<FCurve> ReboundCurve(const FLinearBoundary& InBoundary);
+
+		/**
+		 * Linear deformation of the curve along the axis [Start point, End point] so that the nearest extremity is at the desired position and the other is not modified
+		 */
+		virtual void ExtendTo(const FPoint& DesiredPosition)
+		{
+			ensure(Dimension == 3);
+
+			NOT_IMPLEMENTED;
+			ensureCADKernel(false);
+		}
+
+		/**
+		 * Linear deformation of the curve along the axis [Start point, End point] so that the nearest extremity is at the desired position and the other is not modified
+		 */
+		virtual void ExtendTo(const FPoint2D& DesiredPosition)
+		{
+			ensure(Dimension == 2);
+
+			NOT_IMPLEMENTED;
+			ensureCADKernel(false);
+		}
+
+	protected:
+
+		virtual double ComputeLength(const FLinearBoundary& InBoundary) const;
+		//virtual double ComputeLength2D(const FLinearBoundary& InBoundary) const;
+		//virtual double ComputeLength3D(const FLinearBoundary& InBoundary) const;
+	};
+} // namespace CADKernel
+
