CADKernel Source Code v0.01

diff --git a/Engine/Source/Runtime/Datasmith/CADKernel/Public/CADKernel/Geo/Curves/PolylineCurve.h b/Engine/Source/Runtime/Datasmith/CADKernel/Public/CADKernel/Geo/Curves/PolylineCurve.h
new file mode 100644
--- /dev/null
+++ b/Engine/Source/Runtime/Datasmith/CADKernel/Public/CADKernel/Geo/Curves/PolylineCurve.h
@@ -0,0 +1,225 @@
+// Copyright Epic Games, Inc. All Rights Reserved.
+#pragma once
+
+#include "CADKernel/Geo/Curves/Curve.h"
+#include "CADKernel/Geo/Sampling/Polyline.h"
+#include "CADKernel/Geo/Sampling/PolylineTools.h"
+
+namespace CADKernel
+{
+	class FInfoEntity;
+	class FSurface;
+
+	template<class PointType, class PointCurveType>
+	class CADKERNEL_API TPolylineCurve : public FCurve, public TPolyline<PointType>
+	{
+		friend class FEntity;
+		friend class FPolylineTools;
+
+	protected:
+		
+		TPolylineApproximator<PointType> Approximator;
+
+		TPolylineCurve(const double InTolerance, const TArray<PointType>& InPoints, const TArray<double>& InCoordinates, int8 InDimension)
+			: FCurve(InTolerance, Dimension)
+			, Approximator(this->Coordinates, this->Points)
+		{
+			this->Coordinates = InCoordinates;
+			this->Points = InPoints;
+			ensureCADKernel(Coordinates[0] < Coordinates.Last());
+			Boundary.Set(this->Coordinates.HeapTop(), this->Coordinates.Last());
+		}
+
+		TPolylineCurve(const double InTolerance, const TArray<PointType>& InPoints, int8 InDimension)
+			: FCurve(InTolerance, Dimension)
+			, TPolyline<PointType>(InPoints)
+			, Approximator(this->Coordinates, this->Points)
+		{
+			this->Coordinates.Reserve(InPoints.Num());
+			this->Coordinates.Add(0.);
+
+			double CurvilineLength = 0;
+			for (int32 iPoint = 1; iPoint < this->Points.Num(); iPoint++)
+			{
+				CurvilineLength += this->Points[iPoint].Distance(this->Points[iPoint - 1]);
+				this->Coordinates.Add(CurvilineLength);
+			}
+
+			Boundary.Set(0, CurvilineLength);
+		}
+
+		TPolylineCurve(FCADKernelArchive& Archive)
+			: FCurve()
+			, Approximator(this->Coordinates, this->Points)
+		{
+			Serialize(Archive);
+		}
+
+	public:
+
+		virtual void Serialize(FCADKernelArchive& Ar) override
+		{
+			FCurve::Serialize(Ar);
+			TPolyline<PointType>::Serialize(Ar);
+		}
+
+		virtual void EvaluateCurvesPoint(double InCoordinate, PointCurveType& OutPoint, int32 InDerivativeOrder = 0) const
+		{
+			Approximator.ApproximatePoint<PointCurveType>(InCoordinate, OutPoint, InDerivativeOrder);
+		}
+
+		PointType EvaluatePointAt(double InCoordinate) const
+		{
+			return Approximator.ApproximatePoint(InCoordinate);
+		}
+
+		virtual void EvaluateCurvesPoints(const TArray<double>& InCoordinates, TArray<PointCurveType>& OutPoints, int32 InDerivativeOrder = 0) const
+		{
+			Approximator.ApproximatePoints<PointCurveType>(InCoordinates, OutPoints, InDerivativeOrder);
+		}
+
+		virtual double ComputeSubLength(const FLinearBoundary& InBoundary) const
+		{
+			return Approximator.ComputeLengthOfSubPolyline(InBoundary);
+		}
+
+		const TArray<PointType>& GetPolylinePoints() const
+		{
+			return this->Points;
+		}
+
+		const TArray<double>& GetPolylineParameters() const
+		{
+			return this->Coordinates;
+		}
+
+		virtual void FindNotDerivableCoordinates(const FLinearBoundary& InBoundary, int32 DerivativeOrder, TArray<double>& OutNotDerivableCoordinates) const override
+		{
+			//TO DO
+			ensureCADKernel(false);
+		}
+
+		void SetPoints(const TArray<PointType>& InPoints)
+		{
+			this->Points = InPoints;
+			GlobalLength.Empty();
+		}
+
+		template<class PolylineType>
+		TSharedPtr<FEntityGeom> ApplyMatrixImpl(const FMatrixH& InMatrix) const
+		{
+			TArray<PointType> NewPoints;
+			NewPoints.Reserve(this->Points.Num());
+
+			for (PointType Point : this->Points)
+			{
+				NewPoints.Emplace(InMatrix.Multiply(Point));
+			}
+
+			return FEntity::MakeShared<PolylineType>(Tolerance, NewPoints, this->Coordinates);
+		}
+
+		virtual void ExtendTo(const FPoint& DesiredPoint) override
+		{
+			PolylineTools::ExtendTo(this->Points, (PointType) DesiredPoint);
+		}
+	};
+
+	class CADKERNEL_API FPolylineCurve : public TPolylineCurve<FPoint, FCurvePoint>
+	{
+		friend class FEntity;
+
+	protected:
+		FPolylineCurve(const double InTolerance, const TArray<FPoint>& InPoints, const TArray<double>& InCoordinates)
+			: TPolylineCurve<FPoint, FCurvePoint>(InTolerance, InPoints, InCoordinates, 3)
+		{
+		}
+
+		FPolylineCurve(const double InTolerance, const TArray<FPoint>& InPoints)
+			: TPolylineCurve<FPoint, FCurvePoint>(InTolerance, InPoints, 3)
+		{
+		}
+
+		FPolylineCurve(FCADKernelArchive& Archive)
+			: TPolylineCurve<FPoint, FCurvePoint>(Archive)
+		{
+		}
+
+	public:
+
+		virtual ECurve GetCurveType() const override
+		{
+			return ECurve::Polyline3D;
+		}
+
+		virtual TSharedPtr<FEntityGeom> ApplyMatrix(const FMatrixH& InMatrix) const override
+		{
+			return ApplyMatrixImpl<FPolylineCurve>(InMatrix);
+		}
+
+		virtual FPoint EvaluatePoint(double InCoordinate) const override
+		{
+			return EvaluatePointAt(InCoordinate);
+		} 
+
+#ifdef CADKERNEL_DEV
+		virtual FInfoEntity& GetInfo(FInfoEntity& Info) const override;
+#endif
+
+	};
+
+	class CADKERNEL_API FPolyline2DCurve : public TPolylineCurve<FPoint2D, FCurvePoint2D>
+	{
+		friend class FEntity;
+
+	protected:
+		FPolyline2DCurve(const double InTolerance, const TArray<FPoint2D>& InPoints, const TArray<double>& InCoordinates)
+			: TPolylineCurve<FPoint2D, FCurvePoint2D>(InTolerance, InPoints, InCoordinates, 3)
+		{
+		}
+
+		FPolyline2DCurve(const double InTolerance, const TArray<FPoint2D>& InPoints)
+			: TPolylineCurve<FPoint2D, FCurvePoint2D>(InTolerance, InPoints, 3)
+		{
+		}
+
+		FPolyline2DCurve(FCADKernelArchive& Archive)
+			: TPolylineCurve<FPoint2D, FCurvePoint2D>(Archive)
+		{
+		}
+
+	public:
+
+		virtual ECurve GetCurveType() const override
+		{
+			return ECurve::Polyline3D;
+		}
+
+		virtual TSharedPtr<FEntityGeom> ApplyMatrix(const FMatrixH& InMatrix) const override
+		{
+			return ApplyMatrixImpl<FPolyline2DCurve>(InMatrix);
+		}
+
+		virtual void Evaluate2DPoint(double InCoordinate, FCurvePoint2D& OutPoint, int32 InDerivativeOrder = 0) const override
+		{
+			EvaluateCurvesPoint(InCoordinate, OutPoint, InDerivativeOrder);
+		}
+
+		virtual FPoint2D Evaluate2DPoint(double InCoordinate) const override
+		{
+			return EvaluatePointAt(InCoordinate);
+		}
+
+		virtual void Evaluate2DPoints(const TArray<double>& InCoordinates, TArray<FCurvePoint2D>& OutPoints, int32 InDerivativeOrder = 0) const override
+		{
+			EvaluateCurvesPoints(InCoordinates, OutPoints, InDerivativeOrder);
+		}
+
+#ifdef CADKERNEL_DEV
+		virtual FInfoEntity& GetInfo(FInfoEntity& Info) const override;
+#endif
+
+	}; 
+	
+} // namespace CADKernel
+
diff --git a/Engine/Source/Runtime/Datasmith/CADKernel/Public/CADKernel/Geo/Curves/RestrictionCurve.h b/Engine/Source/Runtime/Datasmith/CADKernel/Public/CADKernel/Geo/Curves/RestrictionCurve.h
new file mode 100644
--- /dev/null
+++ b/Engine/Source/Runtime/Datasmith/CADKernel/Public/CADKernel/Geo/Curves/RestrictionCurve.h
@@ -0,0 +1,185 @@
+// Copyright Epic Games, Inc. All Rights Reserved.
+#pragma once
+
+#include "CADKernel/Geo/Curves/SurfacicCurve.h"
+#include "CADKernel/Geo/Sampling/SurfacicPolyline.h"
+
+namespace CADKernel
+{
+	/**
+	 * A restriction curve is the curve carrying an edge
+	 * 
+	 * It's defined by
+	 * - A surfacic curve defined by a 2D curve (@see Curve2D) and the carrier surface (@see CarrierSurface) of the TopologicalFace containing the edge:
+	 * - A linear approximation of the surfacic curve (@see Polyline) respecting the System Geometrical Tolerance (@see FKernelParameters::GeometricalTolerance)
+	 *   The linear approximation is:
+	 *      - an array of increasing coordinates (@see Polyline::Coordinates)
+	 *      - an array of 2D Points relating to coordinates: points of curve in the parametric space of the carrier surface (@see Polyline::Points2D)
+	 *      - an array of 3D Points relating to coordinates: 3D points of curve (@see Polyline::Points3D)
+	 *      - an array of Normal relating to coordinates: Surface's normal (@see Polyline::Normal)
+	 */
+	class CADKERNEL_API FRestrictionCurve : public FSurfacicCurve
+	{
+		friend class FEntity;
+		friend class FTopologicalEdge;
+
+	protected:
+
+		FSurfacicPolyline Polyline;
+
+		FRestrictionCurve(const double InTolerance, TSharedRef<FSurface> InCarrierSurface, TSharedRef<FCurve> InCurve2D)
+			: FSurfacicCurve(InTolerance, InCurve2D, InCarrierSurface)
+			, Polyline(InCarrierSurface, InCurve2D, InTolerance)
+		{
+		}
+
+		FRestrictionCurve(FCADKernelArchive& Archive)
+			: FSurfacicCurve()
+		{
+			Serialize(Archive);
+		}
+
+	public:
+
+		virtual void Serialize(FCADKernelArchive& Ar) override
+		{
+			FSurfacicCurve::Serialize(Ar);
+			Polyline.Serialize(Ar);
+		}
+
+#ifdef CADKERNEL_DEV
+		virtual FInfoEntity& GetInfo(FInfoEntity&) const override;
+#endif
+
+		virtual ECurve GetCurveType() const override
+		{
+			return ECurve::Restriction;
+		}
+
+		const TSharedPtr<FCurve>& Get2DCurve() const
+		{
+			return Curve2D;
+		}
+
+		const TSharedPtr<FSurface>& GetCarrierSurface() const
+		{
+			return CarrierSurface;
+		}
+
+		TSharedPtr<FEntityGeom> ApplyMatrix(const FMatrixH& InMatrix) const override
+		{
+			ensureCADKernel(false);
+			return TSharedPtr<FEntityGeom>();
+		}
+
+		/**
+		 * Fast computation of the point in the parametric space of the carrier surface.
+		 * This approximation is based on FSurfacicPolyline::Polyline 
+		 */
+		FPoint2D Approximate2DPoint(double InCoordinate) const
+		{
+			return Polyline.Approximate2DPoint(InCoordinate);
+		}
+
+		/**
+		 * Fast computation of the point in the parametric space of the carrier surface.
+		 * This approximation is based on FSurfacicPolyline::Polyline
+		 */
+		FPoint Approximate3DPoint(double InCoordinate) const
+		{
+			return Polyline.Approximate3DPoint(InCoordinate);
+		}
+
+		/**
+		 * Fast computation of the point in the parametric space of the carrier surface.
+		 * This approximation is based on FSurfacicPolyline::Polyline
+		 */
+		void Approximate2DPoints(const TArray<double>& InCoordinates, TArray<FPoint2D>& OutPoints) const
+		{
+			return Polyline.Approximate2DPoints(InCoordinates, OutPoints);
+		}
+
+		/**
+		 * Fast computation of the 3D point of the curve.
+		 * This approximation is based on FSurfacicPolyline::Polyline
+		 */
+		void Approximate3DPoints(const TArray<double>& InCoordinates, TArray<FPoint>& OutPoints) const
+		{
+			return Polyline.Approximate3DPoints(InCoordinates, OutPoints);
+		}
+
+		/**
+		 * Approximation of surfacic polyline (points 2d, 3d, normals, tangents) defined by its coordinates compute with carrier surface polyline
+		 */
+		void ApproximatePolyline(FSurfacicPolyline& OutPolyline) const
+		{
+			Polyline.ApproximatePolyline(OutPolyline);
+		}
+
+
+		double GetCoordinateOfProjectedPoint(const FLinearBoundary& InBoundary, const FPoint& PointOnEdge, FPoint& ProjectedPoint) const
+		{
+			return Polyline.GetCoordinateOfProjectedPoint(InBoundary, PointOnEdge, ProjectedPoint);
+		}
+
+		void ProjectPoints(const FLinearBoundary& InBoundary, const TArray<FPoint>& InPointsToProject, TArray<double>& ProjectedPointCoordinates, TArray<FPoint>& ProjectedPoints) const
+		{
+			Polyline.ProjectPoints(InBoundary, InPointsToProject, ProjectedPointCoordinates, ProjectedPoints);
+		}
+
+		/**
+		 * Project a set of points of a twin curve on the 3D polyline and return the coordinate of the projected point
+		 */
+		void ProjectTwinCurvePoints(const TArray<FPoint>& InPointsToProject, bool bSameOrientation, TArray<double>& OutProjectedPointCoords) const
+		{
+			Polyline.ProjectCoincidentalPolyline(InPointsToProject, bSameOrientation, OutProjectedPointCoords);
+		}
+
+		/**
+		 * A check is done to verify that:
+		 * - the curve is degenerated in the parametric space of the carrier surface i.e. the 2D length of the curve is not nearly equal to zero
+		 * - the curve is degenerated in 3D i.e. the 3D length of the curve is not nearly equal to zero
+		 * 
+		 * A curve can be degenerated in 3D and not in 2D in the case of locally degenerated carrier surface.
+		 */
+		void CheckIfDegenerated(const FLinearBoundary& Boudary, bool& bDegeneration2D, bool& bDegeneration3D, double& Length3D) const
+		{
+			double Tolerance2D = GetCarrierSurface()->Get2DTolerance();
+			Polyline.CheckIfDegenerated(Tolerance, Tolerance2D, Boudary, bDegeneration2D, bDegeneration3D, Length3D);
+		}
+
+		/**
+		 * @return the size of the polyline i.e. the count of points.
+		 */
+		int32 GetPolylineSize()
+		{
+			return Polyline.Size();
+		}
+
+		/**
+		 * Get the sub polyline bounded by the input InBoundary in the orientation of the input InOrientation and append it to the output OutPoints
+		 */
+		template<class PointType>
+		void GetDiscretizationPoints(const FLinearBoundary& Boundary, EOrientation Orientation, TArray<PointType>& OutPoints) const
+		{
+			Polyline.GetSubPolyline(Boundary, Orientation, OutPoints);
+		}
+
+		/**
+		 * Samples the sub curve limited by the boundary respecting the input Desired segment length 
+		 */
+		void Sample(const FLinearBoundary& InBoundary, const double DesiredSegmentLength, TArray<double>& OutCoordinates) const
+		{
+			Polyline.Sample(InBoundary, DesiredSegmentLength, OutCoordinates);
+		}
+
+		virtual double ComputeLength(const FLinearBoundary& InBoundary) const
+		{
+			return Polyline.GetLength(InBoundary);
+		}
+
+		void ExtendTo2DPoint(const FPoint2D& Point);
+	};
+
+} // namespace CADKernel
+
diff --git a/Engine/Source/Runtime/Datasmith/CADKernel/Public/CADKernel/Geo/GeoPoint.h b/Engine/Source/Runtime/Datasmith/CADKernel/Public/CADKernel/Geo/GeoPoint.h
new file mode 100644
--- /dev/null
+++ b/Engine/Source/Runtime/Datasmith/CADKernel/Public/CADKernel/Geo/GeoPoint.h
@@ -0,0 +1,149 @@
+// Copyright Epic Games, Inc. All Rights Reserved.
+#pragma once
+
+#include "CADKernel/Core/Types.h"
+#include "CADKernel/Geo/GeoEnum.h"
+#include "CADKernel/Math/Point.h"
+
+namespace CADKernel
+{
+	struct FCurvePoint2D
+	{
+		int32 DerivativeOrder = -1;
+		FPoint2D Point = FPoint2D::ZeroPoint;
+		FPoint2D Gradient = FPoint2D::ZeroPoint;
+		FPoint2D Laplacian = FPoint2D::ZeroPoint;
+	};
+
+	struct FSurfacicPoint
+	{
+		int32 DerivativeOrder = -1;
+		FPoint Point = FPoint::ZeroPoint;
+		FPoint GradientU = FPoint::ZeroPoint;
+		FPoint GradientV = FPoint::ZeroPoint;
+		FPoint LaplacianU = FPoint::ZeroPoint;
+		FPoint LaplacianV = FPoint::ZeroPoint;
+		FPoint LaplacianUV = FPoint::ZeroPoint;
+	};
+
+	struct FSurfacicCurvePoint
+	{
+		bool bWithNormals;
+		bool bWithTangent;
+
+		FPoint2D Point2D = FPoint2D::ZeroPoint;
+		FPoint Point = FPoint::ZeroPoint;
+		FPoint Normal = FPoint::ZeroPoint;
+		FPoint Tangent = FPoint::ZeroPoint;
+	};
+
+	struct FCurvePoint
+	{
+		int32 DerivativeOrder = -1;
+		FPoint Point = FPoint::ZeroPoint;
+		FPoint Gradient = FPoint::ZeroPoint;
+		FPoint Laplacian = FPoint::ZeroPoint;
+
+		FCurvePoint() = default;
+
+		FCurvePoint(FPoint InPoint)
+			: Point(InPoint)
+		{
+ 			Gradient = FPoint::ZeroPoint;
+			Laplacian = FPoint::ZeroPoint;
+		}
+
+		void Init()
+		{
+			Point = FPoint::ZeroPoint;
+			Gradient = FPoint::ZeroPoint;
+			Laplacian = FPoint::ZeroPoint;
+		}
+
+		/**
+		 * Compute the 3D surface curve point property (3D Coordinate, Gradient, Laplacian) according to 
+		 * its 2D curve point property and the 3D surface point property
+		 */
+		void Combine(const FCurvePoint2D& Point2D, const FSurfacicPoint& SurfacicPoint)
+		{
+			ensureCADKernel(Point2D.DerivativeOrder >= 0);
+			ensureCADKernel(SurfacicPoint.DerivativeOrder >= 0);
+
+			ensureCADKernel(Point2D.DerivativeOrder <= SurfacicPoint.DerivativeOrder);
+
+			DerivativeOrder = Point2D.DerivativeOrder;
+			Point = SurfacicPoint.Point;
+
+			if (DerivativeOrder > 0)
+			{
+				Gradient = SurfacicPoint.GradientU * Point2D.Gradient.U + SurfacicPoint.GradientV * Point2D.Gradient.V;
+			}
+
+			if (DerivativeOrder > 1) 
+			{
+				Laplacian = SurfacicPoint.LaplacianU * FMath::Square(Point2D.Gradient.U)
+					+ 2.0 * SurfacicPoint.LaplacianUV * Point2D.Gradient.U * Point2D.Gradient.V
+					+ SurfacicPoint.LaplacianV * FMath::Square(Point2D.Gradient.V)
+					+ SurfacicPoint.GradientU * Point2D.Laplacian.U 
+					+ SurfacicPoint.GradientV * Point2D.Laplacian.V;
+			}
+		}
+	};
+
+	struct FCoordinateGrid
+	{
+		TArray<double> Coordinates[2];
+
+		FCoordinateGrid()
+		{
+		}
+
+		FCoordinateGrid(const TArray<double>& InUCoordinates, const TArray<double>& InVCoordinates)
+		{
+			Coordinates[EIso::IsoU] = InUCoordinates;
+			Coordinates[EIso::IsoV] = InVCoordinates;
+		}
+
+		void Swap(TArray<double>& InUCoordinates, TArray<double>& InVCoordinates)
+		{
+			::Swap(Coordinates[EIso::IsoU], InUCoordinates);
+			::Swap(Coordinates[EIso::IsoV], InVCoordinates);
+		}
+
+		int32 Count() const
+		{
+			return Coordinates[0].Num() * Coordinates[1].Num();
+		}
+
+		int32 IsoCount(EIso Iso) const
+		{
+			return Coordinates[Iso].Num();
+		}
+
+		void SetNum(int32 UNumber, int32 VNumber)
+		{
+			Coordinates[0].SetNum(UNumber);
+			Coordinates[1].SetNum(VNumber);
+		}
+
+		void Empty(int32 UNumber = 0, int32 VNumber = 0)
+		{
+			Coordinates[0].Empty(UNumber);
+			Coordinates[1].Empty(VNumber);
+		}
+
+		constexpr TArray<double>& operator[](EIso Iso)
+		{
+			ensureCADKernel(Iso == 0 || Iso == 1);
+			return Coordinates[Iso];
+		}
+
+		constexpr const TArray<double>& operator[](EIso Iso) const
+		{
+			ensureCADKernel(Iso == 0 || Iso == 1);
+			return Coordinates[Iso];
+		}
+	};
+
+
+}
\ No newline at end of file
diff --git a/Engine/Source/Runtime/Datasmith/CADKernel/Public/CADKernel/Geo/Sampling/PolylineTools.h b/Engine/Source/Runtime/Datasmith/CADKernel/Public/CADKernel/Geo/Sampling/PolylineTools.h
new file mode 100644
--- /dev/null
+++ b/Engine/Source/Runtime/Datasmith/CADKernel/Public/CADKernel/Geo/Sampling/PolylineTools.h
@@ -0,0 +1,603 @@
+// Copyright Epic Games, Inc. All Rights Reserved.
+#pragma once
+
+#include "CADKernel/Core/Types.h"
+#include "CADKernel/Math/Boundary.h"
+#include "CADKernel/Math/Geometry.h"
+#include "CADKernel/Math/Point.h"
+#include "CADKernel/Geo/GeoEnum.h"
+#include "CADKernel/Geo/GeoPoint.h"
+#include "CADKernel/Utils/IndexOfCoordinateFinder.h"
+
+#include "Algo/ForEach.h"
+#include "Algo/Reverse.h"
+
+namespace CADKernel
+{
+	namespace PolylineTools
+	{
+		inline bool IsDichotomyToBePreferred(int32 InPolylineSize, int32 ResultSize)
+		{
+			float MeanLinearIteration = (float)InPolylineSize / (float)ResultSize;
+			float MaxDichotomyIteration = FMath::Log2((float)InPolylineSize);
+
+			if (MeanLinearIteration > MaxDichotomyIteration)
+			{
+				return true;
+			}
+			return false;
+		}
+
+		template<typename PointType>
+		inline PointType LinearInterpolation(const TArray<PointType>& Array, const int32 Index, const double Coordinate)
+		{
+			ensureCADKernel(Index + 1 < Array.Num());
+			return Array[Index] + (Array[Index + 1] - Array[Index]) * Coordinate;
+		}
+
+		inline double SectionCoordinate(const TArray<double>& Array, const int32 Index, const double Coordinate)
+		{
+			ensureCADKernel(Index + 1 < Array.Num());
+			const double DeltaU = Array[Index + 1] - Array[Index];
+			if (FMath::IsNearlyZero(DeltaU))
+			{
+				return 0;
+			}
+			return (Coordinate - Array[Index]) / DeltaU;
+		}
+
+		/**
+		 * Progressively deforms a polyline (or a control polygon) so that its end is in the desired position 
+		 */
+		template<typename PointType>
+		void ExtendTo(TArray<PointType>& Polyline, const PointType& DesiredEnd)
+		{
+			double DistanceStart = Polyline[0].SquareDistance(DesiredEnd);
+			double DistanceEnd = Polyline.Last().SquareDistance(DesiredEnd);
+
+			if (DistanceStart > DistanceEnd)
+			{
+				PointType Factor;
+				for (int32 Index = 0; Index < PointType::Dimension(); ++Index)
+				{
+					Factor[Index] = FMath::Abs(Polyline.Last()[Index] - Polyline[0][Index]) > SMALL_NUMBER_SQUARE ? (DesiredEnd[Index] - Polyline[0][Index]) / (Polyline.Last()[Index] - Polyline[0][Index]) : 1.;
+				}
+
+				Algo::ForEach(Polyline, [&, Factor](PointType& Pole)
+					{
+						for (int32 Index = 0; Index < PointType::Dimension(); ++Index)
+						{
+							Pole[Index] = Polyline[0][Index] + (Pole[Index] - Polyline[0][Index]) * Factor[Index];
+						}
+					}
+				);
+			}
+			else
+			{
+				PointType Factor;
+				for (int32 Index = 0; Index < PointType::Dimension(); ++Index)
+				{
+					Factor[Index] = FMath::Abs(Polyline[0][Index] - Polyline.Last()[Index]) > SMALL_NUMBER_SQUARE ? (DesiredEnd[Index] - Polyline.Last()[Index]) / (Polyline[0][Index] - Polyline.Last()[Index]) : 1.;
+				}
+
+				Algo::Reverse(Polyline);
+				Algo::ForEach(Polyline, [&, Factor](PointType& Pole)
+					{
+						for (int32 Index = 0; Index < PointType::Dimension(); ++Index)
+						{
+							Pole[Index] = Polyline.Last()[Index] + (Pole[Index] - Polyline.Last()[Index]) * Factor[Index];
+						}
+					}
+				);
+				Algo::Reverse(Polyline);
+			}
+		}
+	}
+
+	template<class PointType>
+	class TPolylineApproximator
+	{
+	protected:
+		const TArray<double>& PolylineCoordinates;
+		const TArray<PointType>& PolylinePoints;
+
+	public:
+		TPolylineApproximator(const TArray<double>& InPolylineCoordinates, const TArray<PointType>& InPolylinePoints)
+			: PolylineCoordinates(InPolylineCoordinates)
+			, PolylinePoints(InPolylinePoints)
+		{
+		}
+
+	protected:
+
+		double ComputeCurvilinearCoordinatesOfPolyline(const FLinearBoundary& InBoundary, TArray<double>& OutCurvilinearCoordinates, int32& StartIndex, int32& EndIndex) const
+		{
+			GetStartEndIndex(InBoundary, StartIndex, EndIndex);
+
+			OutCurvilinearCoordinates.Reserve(EndIndex - StartIndex + 2);
+
+			double LastSegmentLength;
+			double Length = 0;
+			ensureCADKernel(EndIndex > StartIndex);
+
+			OutCurvilinearCoordinates.Add(0);
+			if (EndIndex > StartIndex + 1)
+			{
+				PointType StartPoint = ComputePoint(StartIndex, InBoundary.Min);
+				PointType EndPoint = ComputePoint(EndIndex, InBoundary.Max);
+				Length = StartPoint.Distance(PolylinePoints[StartIndex + 1]);
+				LastSegmentLength = EndPoint.Distance(PolylinePoints[EndIndex]);
+
+				OutCurvilinearCoordinates.Add(Length);
+				for (int32 Index = StartIndex + 1; Index < EndIndex; ++Index)
+				{
+					Length += PolylinePoints[Index].Distance(PolylinePoints[Index + 1]);
+					OutCurvilinearCoordinates.Add(Length);
+				}
+				Length += LastSegmentLength;
+				OutCurvilinearCoordinates.Add(Length);
+			}
+			else
+			{
+				PointType StartPoint = ComputePoint(StartIndex, InBoundary.Min);
+				PointType EndPoint = ComputePoint(EndIndex, InBoundary.Max);
+				Length = StartPoint.Distance(EndPoint);
+				OutCurvilinearCoordinates.Add(Length);
+			}
+			return Length;
+		}
+
+		PointType ComputePoint(const int32 Index, const double PointCoordinate) const
+		{
+			double Delta = PolylineCoordinates[Index + 1] - PolylineCoordinates[Index];
+			if (FMath::IsNearlyZero(Delta, (double)KINDA_SMALL_NUMBER))
+			{
+				return PolylinePoints[Index];
+			}
+
+			return PolylinePoints[Index] + (PolylinePoints[Index + 1] - PolylinePoints[Index]) * (PointCoordinate - PolylineCoordinates[Index]) / Delta;
+		};
+
+		/**
+		 * Project a Set of points on a restricted polyline (StartIndex & EndIndex define the polyline boundary)
+		 * the points are projected on all segments of the polyline, the closest are selected
+		 */
+		double ProjectPointToPolyline(int32 StartIndex, int32 EndIndex, const PointType& InPointToProject, PointType& OutProjectedPoint) const
+		{
+			double MinDistance = HUGE_VAL;
+			double UForMinDistance = 0;
+
+			double ParamU = 0.;
+			int32 SegmentIndex = 0;
+
+			for (int32 Index = StartIndex; Index <= EndIndex; ++Index)
+			{
+				FPoint ProjectPoint = ProjectPointOnSegment(InPointToProject, PolylinePoints[Index], PolylinePoints[Index + 1], ParamU, true);
+				double SquareDistance = FMath::Square(ProjectPoint[0] - InPointToProject[0]);
+				if (SquareDistance > MinDistance)
+				{
+					continue;
+				}
+				SquareDistance += FMath::Square(ProjectPoint[1] - InPointToProject[1]);
+				if (SquareDistance > MinDistance)
+				{
+					continue;
+				}
+				SquareDistance += FMath::Square(ProjectPoint[2] - InPointToProject[2]);
+				if (SquareDistance > MinDistance)
+				{
+					continue;
+				}
+				MinDistance = SquareDistance;
+				UForMinDistance = ParamU;
+				SegmentIndex = Index;
+				OutProjectedPoint = ProjectPoint;
+			}
+
+			return PolylineTools::LinearInterpolation(PolylineCoordinates, SegmentIndex, UForMinDistance);
+		}
+
+	public:
+
+		void GetStartEndIndex(const FLinearBoundary& InBoundary, int32& StartIndex, int32& EndIndex) const
+		{
+			FDichotomyFinder Finder(PolylineCoordinates);
+			StartIndex = Finder.Find(InBoundary.Min);
+			EndIndex = Finder.Find(InBoundary.Max);
+		}
+
+		/**
+		 * Evaluate the point of the polyline at the input InCoordinate
+		 * If the input coordinate is outside the boundary of the polyline, the coordinate of the nearest boundary is used.
+		 */
+		PointType ApproximatePoint(const double InCoordinate) const
+		{
+			FDichotomyFinder Finder(PolylineCoordinates);
+			int32 Index = Finder.Find(InCoordinate);
+			return ComputePoint(Index, InCoordinate);
+		}
+
+		/**
+		 * Evaluate the point of the polyline at the input Coordinate
+		 * If the input coordinate is outside the boundary of the polyline, the coordinate of the nearest boundary is used.
+		 */
+		template<class CurvePointType>
+		void ApproximatePoint(double InCoordinate, CurvePointType& OutPoint, int32 InDerivativeOrder) const
+		{
+			FDichotomyFinder Finder(PolylineCoordinates);
+			int32 Index = Finder.Find(InCoordinate);
+
+			OutPoint.DerivativeOrder = InDerivativeOrder;
+
+			double DeltaU = PolylineCoordinates[Index + 1] - PolylineCoordinates[Index];
+			if (FMath::IsNearlyZero(DeltaU))
+			{
+				OutPoint.Point = PolylinePoints[Index];
+				OutPoint.Gradient = FPoint::ZeroPoint;
+				OutPoint.Laplacian = FPoint::ZeroPoint;
+				return;
+			}
+
+			double SectionCoordinate = (InCoordinate - PolylineCoordinates[Index]) / DeltaU;
+
+			FPoint Tangent = PolylinePoints[Index + 1] - PolylinePoints[Index];
+			OutPoint.Point = PolylinePoints[Index] + Tangent * SectionCoordinate;
+
+			if (InDerivativeOrder > 0)
+			{
+				OutPoint.Gradient = Tangent;
+				OutPoint.Laplacian = FPoint::ZeroPoint;
+			}
+		}
+
+		template<class CurvePointType>
+		void ApproximatePoints(const TArray<double>& InCoordinates, TArray<CurvePointType>& OutPoints, int32 InDerivativeOrder = 0) const
+		{
+			if (!InCoordinates.Num())
+			{
+				ensureCADKernel(false);
+				return;
+			}
+
+			TFunction<void(FIndexOfCoordinateFinder&)> ComputePoints = [&](FIndexOfCoordinateFinder& Finder)
+			{
+				for (int32 IPoint = 0; IPoint < InCoordinates.Num(); ++IPoint)
+				{
+					int32 Index = Finder.Find(InCoordinates[IPoint]);
+
+					OutPoints[IPoint].DerivativeOrder = InDerivativeOrder;
+
+					double DeltaU = PolylineCoordinates[Index + 1] - PolylineCoordinates[Index];
+					if (FMath::IsNearlyZero(DeltaU))
+					{
+						OutPoints[IPoint].Point = PolylinePoints[Index];
+						OutPoints[IPoint].Gradient = FPoint::ZeroPoint;
+						OutPoints[IPoint].Laplacian = FPoint::ZeroPoint;
+						return;
+					}
+
+					double SectionCoordinate = (InCoordinates[IPoint] - PolylineCoordinates[Index]) / DeltaU;
+
+					FPoint Tangent = PolylinePoints[Index + 1] - PolylinePoints[Index];
+					OutPoints[IPoint].Point = PolylinePoints[Index] + Tangent * SectionCoordinate;
+
+					if (InDerivativeOrder > 0)
+					{
+						OutPoints[IPoint].Gradient = Tangent;
+						OutPoints[IPoint].Laplacian = FPoint::ZeroPoint;
+					}
+				}
+			};
+
+			FDichotomyFinder DichotomyFinder(PolylineCoordinates);
+
+			int32 StartIndex = DichotomyFinder.Find(InCoordinates[0]);
+			int32 EndIndex = DichotomyFinder.Find(InCoordinates.Last());
+			bool bUseDichotomy = PolylineTools::IsDichotomyToBePreferred(EndIndex - StartIndex, InCoordinates.Num());
+
+			OutPoints.Empty(InCoordinates.Num());
+			if (bUseDichotomy)
+			{
+				DichotomyFinder.StartLower = StartIndex;
+				DichotomyFinder.StartUpper = EndIndex;
+				ComputePoints(DichotomyFinder);
+			}
+			else
+			{
+				FLinearFinder LinearFinder(PolylineCoordinates, StartIndex);
+				ComputePoints(LinearFinder);
+			}
+		}
+
+
+		/**
+		 * Evaluate the points of the polyline associated to the increasing array of input Coordinates
+		 * If the input coordinate is outside the boundary of the polyline, the coordinate of the nearest boundary is used.
+		 */
+		void ApproximatePoints(const TArray<double>& InCoordinates, TArray<PointType>& OutPoints) const
+		{
+			if (!InCoordinates.Num())
+			{
+				return;
+			}
+
+			TFunction<void(FIndexOfCoordinateFinder&)> ComputePoints = [&](FIndexOfCoordinateFinder& Finder)
+			{
+				for (double Coordinate : InCoordinates)
+				{
+					int32 Index = Finder.Find(Coordinate);
+					OutPoints.Emplace(ComputePoint(Index, Coordinate));
+				}
+			};
+
+			FDichotomyFinder DichotomyFinder(PolylineCoordinates);
+
+			int32 StartIndex = DichotomyFinder.Find(InCoordinates[0]);
+			int32 EndIndex = DichotomyFinder.Find(InCoordinates.Last());
+			bool bUseDichotomy = PolylineTools::IsDichotomyToBePreferred(EndIndex - StartIndex, InCoordinates.Num());
+
+			OutPoints.Empty(InCoordinates.Num());
+			if (bUseDichotomy)
+			{
+				DichotomyFinder.StartLower = StartIndex;
+				DichotomyFinder.StartUpper = EndIndex;
+				ComputePoints(DichotomyFinder);
+			}
+			else
+			{
+				FLinearFinder LinearFinder(PolylineCoordinates, StartIndex);
+				ComputePoints(LinearFinder);
+			}
+		}
+
+		void SamplePolyline(const FLinearBoundary& InBoundary, const double DesiredSegmentLength, TArray<double>& OutCoordinates) const
+		{
+			int32 StartIndex = 0;
+			int32 EndIndex = 0;
+
+			TArray<double> CurvilinearCoordinates;
+			double CurveLength = ComputeCurvilinearCoordinatesOfPolyline(InBoundary, CurvilinearCoordinates, StartIndex, EndIndex);
+
+			int32 SegmentNum = (int32)FMath::Max(CurveLength / DesiredSegmentLength + 0.5, 1.0);
+
+			double SectionLength = CurveLength / (double)(SegmentNum);
+
+			OutCoordinates.Empty();
+			OutCoordinates.Reserve(SegmentNum + 1);
+			OutCoordinates.Add(InBoundary.Min);
+
+			TFunction<double(const int32, const int32, const double, const double)> ComputeSamplePointCoordinate = [&](const int32 IndexCurvilinear, const int32 IndexCoordinate, const double Length, const double Coordinate)
+			{
+				return Coordinate + (PolylineCoordinates[IndexCoordinate] - Coordinate) * (Length - CurvilinearCoordinates[IndexCurvilinear - 1u]) / (CurvilinearCoordinates[IndexCurvilinear] - CurvilinearCoordinates[IndexCurvilinear - 1u]);
+			};
+
+			double CurvilinearLength = SectionLength;
+			double LastCoordinate = InBoundary.Min;
+			for (int32 IndexCurvilinear = 1, IndexCoordinate = StartIndex + 1; IndexCurvilinear < CurvilinearCoordinates.Num(); ++IndexCurvilinear, ++IndexCoordinate)
+			{
+				while (CurvilinearLength < CurvilinearCoordinates[IndexCurvilinear] + SMALL_NUMBER)
+				{
+					double Coordinate = ComputeSamplePointCoordinate(IndexCurvilinear, IndexCoordinate, CurvilinearLength, LastCoordinate);
+					OutCoordinates.Add(Coordinate);
+					CurvilinearLength += SectionLength;
+					if (CurvilinearLength + SMALL_NUMBER > CurveLength)
+					{
+						OutCoordinates.Add(InBoundary.Max);
+						break;
+					}
+				}
+				LastCoordinate = PolylineCoordinates[IndexCoordinate];
+			}
+		}
+
+		/**
+		 * Project a Set of points on a restricted polyline (StartIndex & EndIndex define the polyline boundary)
+		 * the points are projected on all segments of the polyline, the closest are selected
+		 */
+		double ProjectPointToPolyline(const FLinearBoundary& InBoundary, const PointType& PointOnEdge, PointType& OutProjectedPoint) const
+		{
+			int32 StartIndex = 0;
+			int32 EndIndex = 0;
+			GetStartEndIndex(InBoundary, StartIndex, EndIndex);
+
+			return ProjectPointToPolyline(StartIndex, EndIndex, PointOnEdge, OutProjectedPoint);
+		}
+
+		/**
+		 * Project a Set of points on a restricted polyline (StartIndex & EndIndex define the polyline boundary)
+		 * Each points are projected on all segments of the restricted polyline, the closest are selected
+		 */
+		void ProjectPointsToPolyline(const FLinearBoundary& InBoundary, const TArray<PointType>& InPointsToProject, TArray<double>& OutProjectedPointCoords, TArray<PointType>& OutProjectedPoints) const
+		{
+			OutProjectedPointCoords.Empty(InPointsToProject.Num());
+			OutProjectedPoints.Empty(InPointsToProject.Num());
+
+			int32 StartIndex = 0;
+			int32 EndIndex = 0;
+
+			GetStartEndIndex(InBoundary, StartIndex, EndIndex);
+
+			for (const PointType& Point : InPointsToProject)
+			{
+				PointType ProjectedPoint;
+				double Coordinate = ProjectPointToPolyline(StartIndex, EndIndex, Point, ProjectedPoint);
+
+				OutProjectedPointCoords.Emplace(Coordinate);
+				OutProjectedPoints.Emplace(ProjectedPoint);
+			}
+		}
+
+		/**
+		 * Project each point of a coincidental polyline on the Polyline. 
+		 */
+		void ProjectCoincidentalPolyline(const TArray<PointType>& InPointsToProject, bool bSameOrientation, TArray<double>& OutProjectedPointCoords) const
+		{
+			int32 StartIndex = 0;
+			const int32 EndIndex = PolylinePoints.Num() - 1;
+
+			TFunction<double(const PointType&)> ProjectPointToPolyline = [&](const PointType & InPointToProject)
+			{
+				double MinDistance = HUGE_VAL;
+				double UForMinDistance = 0;
+
+				double ParamU = 0.;
+				int32 SegmentIndex = 0;
+				for (int32 Index = StartIndex; Index < EndIndex; ++Index)
+				{
+					PointType ProjectPoint = ProjectPointOnSegment(InPointToProject, PolylinePoints[Index], PolylinePoints[Index + 1], ParamU, true);
+
+					double SquareDistance = ProjectPoint.SquareDistance(InPointToProject);
+					if (SquareDistance > MinDistance)
+					{
+						break;
+					}
+					MinDistance = SquareDistance;
+					UForMinDistance = ParamU;
+					SegmentIndex = Index;
+				}
+				StartIndex = SegmentIndex;
+				return PolylineTools::LinearInterpolation(PolylineCoordinates, SegmentIndex, UForMinDistance);
+			};
+
+			if (bSameOrientation)
+			{
+				OutProjectedPointCoords.Empty(InPointsToProject.Num());
+				for (const PointType& Point : InPointsToProject)
+				{
+					double Coordinate = ProjectPointToPolyline(Point);
+					OutProjectedPointCoords.Emplace(Coordinate);
+				}
+			}
+			else
+			{
+				OutProjectedPointCoords.SetNum(InPointsToProject.Num());
+				for (int32 Index = InPointsToProject.Num()-1, Pndex = 0; Index >= 0; --Index, ++Pndex)
+				{
+					OutProjectedPointCoords[Pndex] = ProjectPointToPolyline(InPointsToProject[Index]);
+				}
+			}
+		}
+
+
+		/**
+		 * Append to the OutPoints array the sub polyline bounded by InBoundary according to the orientation
+		 */
+		void GetSubPolyline(const FLinearBoundary& InBoundary, const EOrientation Orientation, TArray<PointType>& OutPoints) const
+		{
+			int32 StartIndex = 0;
+			int32 EndIndex = 0;
+			GetStartEndIndex(InBoundary, StartIndex, EndIndex);
+
+			int32 NewSize = OutPoints.Num() + EndIndex - StartIndex + 2;
+			OutPoints.Reserve(NewSize);
+
+			int32 PolylineStartIndex = StartIndex;
+			int32 PolylineEndIndex = EndIndex;
+			if (FMath::IsNearlyEqual(PolylineCoordinates[StartIndex + 1], InBoundary.Min, (double)KINDA_SMALL_NUMBER))
+			{
+				PolylineStartIndex++;
+			}
+			if (FMath::IsNearlyEqual(PolylineCoordinates[EndIndex], InBoundary.Max, (double)KINDA_SMALL_NUMBER))
+			{
+				PolylineEndIndex--;
+			}
+
+			if (Orientation)
+			{
+				OutPoints.Emplace(ComputePoint(StartIndex, InBoundary.Min));
+				if(PolylineEndIndex - PolylineStartIndex > 0)
+				{
+					OutPoints.Append(PolylinePoints.GetData() + PolylineStartIndex + 1, PolylineEndIndex - PolylineStartIndex);
+				}
+				OutPoints.Emplace(ComputePoint(EndIndex, InBoundary.Max));
+			}
+			else
+			{
+				OutPoints.Emplace(ComputePoint(EndIndex, InBoundary.Max));
+				for (int32 Index = PolylineEndIndex; Index > PolylineStartIndex; --Index)
+				{
+					OutPoints.Emplace(PolylinePoints[Index]);
+				}
+				OutPoints.Emplace(ComputePoint(StartIndex, InBoundary.Min));
+			}
+		}
+
+		/**
+		 * Get the subset of point defining the sub polyline bounded by InBoundary
+		 * OutCoordinates and OutPoints are emptied before the process
+		 */
+		void GetSubPolyline(const FLinearBoundary& InBoundary, TArray<double>& OutCoordinates, TArray<PointType>& OutPoints) const
+		{
+			int32 StartIndex = 0;
+			int32 EndIndex = 0;
+			GetStartEndIndex(PolylineCoordinates, InBoundary, StartIndex, EndIndex);
+
+			OutCoordinates.Empty(EndIndex - StartIndex + 2);
+			OutPoints.Empty(EndIndex - StartIndex + 2);
+
+			if (FMath::IsNearlyEqual(PolylineCoordinates[StartIndex + 1], InBoundary.Min, KINDA_SMALL_NUMBER))
+			{
+				StartIndex++;
+			}
+
+			if (!FMath::IsNearlyEqual(PolylineCoordinates[EndIndex], InBoundary.Max, KINDA_SMALL_NUMBER))
+			{
+				EndIndex--;
+			}
+
+			OutPoints.Emplace(ComputePoint(StartIndex, InBoundary.Min));
+			OutPoints.Append(PolylinePoints.GetData() + StartIndex + 1, EndIndex - StartIndex);
+			OutPoints.Emplace(ComputePoint(EndIndex, InBoundary.Max));
+
+			OutCoordinates.Add(InBoundary.Min);
+			OutCoordinates.Append(PolylineCoordinates.GetData() + StartIndex + 1, EndIndex - StartIndex);
+			OutCoordinates.Add(InBoundary.Max);
+		}
+
+		double ComputeLength() const
+		{
+			double Length = 0;
+			for (int32 Index = 1; Index < PolylinePoints.Num(); ++Index)
+			{
+				Length += PolylinePoints[Index-1].Distance(PolylinePoints[Index]);
+			}
+			return Length;
+		}
+
+		double ComputeLengthOfSubPolyline(const FLinearBoundary& InBoundary) const
+		{
+			int32 StartIndex = 0;
+			int32 EndIndex = 0;
+
+			GetStartEndIndex(InBoundary, StartIndex, EndIndex);
+
+			double Length = 0;
+			if (EndIndex > StartIndex)
+			{
+				FPoint StartPoint = ComputePoint(StartIndex, InBoundary.Min);
+				FPoint EndPoint = ComputePoint(EndIndex, InBoundary.Max);
+				Length = StartPoint.Distance(PolylinePoints[StartIndex + 1]);
+				Length += EndPoint.Distance(PolylinePoints[EndIndex]);
+
+				if (EndIndex > StartIndex + 1)
+				{
+					for (int32 Index = StartIndex + 1; Index < EndIndex; ++Index)
+					{
+						Length += PolylinePoints[Index].Distance(PolylinePoints[Index + 1]);
+					}
+				}
+			}
+			else
+			{
+				FPoint StartPoint = ComputePoint(StartIndex, InBoundary.Min);
+				FPoint EndPoint = ComputePoint(EndIndex, InBoundary.Max);
+				Length = StartPoint.Distance(EndPoint);
+			}
+			return Length;
+		}
+
+
+	};
+
+} // namespace CADKernel
+
diff --git a/Engine/Source/Runtime/Datasmith/CADKernel/Public/CADKernel/Geo/Sampling/SurfacicPolyline.h b/Engine/Source/Runtime/Datasmith/CADKernel/Public/CADKernel/Geo/Sampling/SurfacicPolyline.h
new file mode 100644
--- /dev/null
+++ b/Engine/Source/Runtime/Datasmith/CADKernel/Public/CADKernel/Geo/Sampling/SurfacicPolyline.h
@@ -0,0 +1,355 @@
+// Copyright Epic Games, Inc. All Rights Reserved.
+#pragma once
+
+#include "CADKernel/Core/Types.h"
+
+#include "CADKernel/Geo/Curves/Curve.h"
+#include "CADKernel/Geo/Sampling/PolylineTools.h"
+#include "CADKernel/Geo/Surfaces/Surface.h"
+#include "CADKernel/UI/Display.h"
+
+#include "Serialization/Archive.h"
+
+namespace CADKernel
+{
+	class FInfoEntity;
+
+	class FSurfacicPolyline
+	{
+
+	public:
+
+		TArray<double> Coordinates;
+		TArray<FPoint2D> Points2D;
+		TArray<FPoint> Points3D;
+		TArray<FPoint> Normals;
+		TArray<FPoint> Tangents;
+
+		bool bWithNormals;
+		bool bWithTangent;
+
+		FSurfacicPolyline(FCADKernelArchive& Archive)
+		{
+			Serialize(Archive);
+		}
+
+		FSurfacicPolyline(TSharedRef<FSurface> InCarrierSurface, TSharedRef<FCurve> InCurve2D, const double Tolerance);
+
+		FSurfacicPolyline(TSharedRef<FSurface> InCarrierSurface, TSharedRef<FCurve> InCurve2D, const double ChordTolerance, const double ParamTolerance, bool bInWithNormals/* = false*/, bool bWithTangent/* = false*/);
+
+		FSurfacicPolyline(bool bInWithNormals = false, bool bInWithTangent = false)
+			: bWithNormals(bInWithNormals)
+			, bWithTangent(bInWithTangent)
+		{
+		}
+
+		void Serialize(FCADKernelArchive& Ar)
+		{
+			Ar << Points3D;
+			Ar << Points2D;
+			Ar << Normals;
+			Ar << Coordinates;
+			Ar << bWithNormals;
+			Ar << bWithTangent;
+		}
+
+		FInfoEntity& GetInfo(FInfoEntity&) const;
+
+		TSharedPtr<FEntityGeom> ApplyMatrix(const FMatrixH&) const;
+
+		void CheckIfDegenerated(const double Tolerance3D, const double Tolerance2D, const FLinearBoundary& Boudary, bool& bDegeneration2D, bool& bDegeneration3D, double& Length3D) const;
+
+		FPoint Approximate3DPoint(double InCoordinate) const
+		{
+			TPolylineApproximator<FPoint> Approximator3D(Coordinates, Points3D);
+			return Approximator3D.ApproximatePoint(InCoordinate);
+		}
+
+		void Approximate3DPoints(const TArray<double>& InCoordinates, TArray<FPoint>& OutPoints) const
+		{
+			TPolylineApproximator<FPoint> Approximator3D(Coordinates, Points3D);
+			Approximator3D.ApproximatePoints(InCoordinates, OutPoints);
+		}
+
+		FPoint2D Approximate2DPoint(double InCoordinate) const
+		{
+			TPolylineApproximator<FPoint2D> Approximator(Coordinates, Points2D);
+			return Approximator.ApproximatePoint(InCoordinate);
+		}
+
+		void Approximate2DPoints(const TArray<double>& InCoordinates, TArray<FPoint2D>& OutPoints) const
+		{
+			TPolylineApproximator<FPoint2D> Approximator(Coordinates, Points2D);
+			Approximator.ApproximatePoints(InCoordinates, OutPoints);
+		}
+
+		void ApproximatePolyline(FSurfacicPolyline& OutPolyline) const
+		{
+			TFunction<void(FIndexOfCoordinateFinder&)> ComputePoints = [&](FIndexOfCoordinateFinder& Finder)
+			{
+				int32 CoordinateCount = OutPolyline.Coordinates.Num();
+
+				TArray<int32> SegmentIndexes;
+				SegmentIndexes.Reserve(CoordinateCount);
+
+				TArray<double> SegmentCoordinates;
+				SegmentCoordinates.Reserve(CoordinateCount);
+
+
+				//for (int32 Index = 0; Index < CoordinateCount; ++Index)
+				for (const double Coordinate : OutPolyline.Coordinates)
+				{
+					const int32& Index = SegmentIndexes.Emplace_GetRef(Finder.Find(Coordinate));
+					SegmentCoordinates.Emplace(PolylineTools::SectionCoordinate(Coordinates, Index, Coordinate));
+				}
+
+				OutPolyline.Points2D.Reserve(CoordinateCount);
+				for (int32 Index = 0; Index < CoordinateCount; ++Index)
+				{
+					int32 SegmentIndex = SegmentIndexes[Index];
+					double SegmentCoordinate = SegmentCoordinates[Index];
+					OutPolyline.Points2D.Emplace(PolylineTools::LinearInterpolation(Points2D, SegmentIndex, SegmentCoordinate));
+				}
+
+				OutPolyline.Points3D.Reserve(CoordinateCount);
+				for (int32 Index = 0; Index < CoordinateCount; ++Index)
+				{
+					int32 SegmentIndex = SegmentIndexes[Index];
+					double SegmentCoordinate = SegmentCoordinates[Index];
+					OutPolyline.Points3D.Emplace(PolylineTools::LinearInterpolation(Points3D, SegmentIndex, SegmentCoordinate));
+				}
+
+				if (bWithNormals)
+				{
+					for (int32 Index = 0; Index < CoordinateCount; ++Index)
+					{
+						int32 SegmentIndex = SegmentIndexes[Index];
+						double SegmentCoordinate = SegmentCoordinates[Index];
+						OutPolyline.Normals.Emplace(PolylineTools::LinearInterpolation(Normals, SegmentIndex, SegmentCoordinate));
+					}
+				}
+
+				if (bWithTangent)
+				{
+					for (int32 Index = 0; Index < CoordinateCount; ++Index)
+					{
+						int32 SegmentIndex = SegmentIndexes[Index];
+						double SegmentCoordinate = SegmentCoordinates[Index];
+						OutPolyline.Tangents.Emplace(PolylineTools::LinearInterpolation(Tangents, SegmentIndex, SegmentCoordinate));
+					}
+				}
+			};
+
+			FDichotomyFinder DichotomyFinder(Coordinates);
+			int32 StartIndex = 0;
+			int32 EndIndex;
+
+			bool bUseDichotomy = false;
+			StartIndex = DichotomyFinder.Find(OutPolyline.Coordinates[0]);
+			EndIndex = DichotomyFinder.Find(OutPolyline.Coordinates.Last());
+			bUseDichotomy = PolylineTools::IsDichotomyToBePreferred(EndIndex - StartIndex, Coordinates.Num());
+
+
+			if (bUseDichotomy)
+			{
+				DichotomyFinder.StartLower = StartIndex;
+				DichotomyFinder.StartUpper = EndIndex;
+				ComputePoints(DichotomyFinder);
+			}
+			else
+			{
+				FLinearFinder LinearFinder(Coordinates, StartIndex);
+				ComputePoints(LinearFinder);
+			}
+		}
+
+
+		void Sample(const FLinearBoundary& Boundary, const double DesiredSegmentLength, TArray<double>& OutCoordinates) const
+		{
+			TPolylineApproximator<FPoint> Approximator3D(Coordinates, Points3D);
+			Approximator3D.SamplePolyline(Boundary, DesiredSegmentLength, OutCoordinates);
+		}
+
+		double GetCoordinateOfProjectedPoint(const FLinearBoundary& Boundary, const FPoint& PointOnEdge, FPoint& ProjectedPoint) const
+		{
+			TPolylineApproximator<FPoint> Approximator3D(Coordinates, Points3D);
+			return Approximator3D.ProjectPointToPolyline(Boundary, PointOnEdge, ProjectedPoint);
+		}
+
+		void ProjectPoints(const FLinearBoundary& InBoundary, const TArray<FPoint>& InPointsToProject, TArray<double>& ProjectedPointCoordinates, TArray<FPoint>& ProjectedPoints) const
+		{
+			TPolylineApproximator<FPoint> Approximator3D(Coordinates, Points3D);
+			Approximator3D.ProjectPointsToPolyline(InBoundary, InPointsToProject, ProjectedPointCoordinates, ProjectedPoints);
+		}
+
+		/**
+		 * Project each point of a coincidental polyline on the Polyline.
+		 */
+		void ProjectCoincidentalPolyline(const TArray<FPoint>& InPointsToProject, bool bSameOrientation, TArray<double>& OutProjectedPointCoordinates) const
+		{
+			TPolylineApproximator<FPoint> Approximator3D(Coordinates, Points3D);
+			Approximator3D.ProjectCoincidentalPolyline(InPointsToProject, bSameOrientation, OutProjectedPointCoordinates);
+		}
+
+		const TArray<double>& GetCoordinates() const
+		{
+			return Coordinates;
+		}
+
+		const TArray<FPoint2D>& Get2DPoints() const
+		{
+			return Points2D;
+		}
+
+		const TArray<FPoint>& GetPoints() const
+		{
+			return Points3D;
+		}
+
+		const TArray<FPoint>& GetNormals() const
+		{
+			return Normals;
+		}
+
+		const TArray<FPoint>& GetTangents() const
+		{
+			return Tangents;
+		}
+
+		void SwapCoordinates(TArray<double>& NewCoordinates)
+		{
+			Swap(NewCoordinates, Coordinates);
+		}
+
+		/**
+		 * @return the size of the polyline i.e. the count of points.
+		 */
+		int32 Size() const
+		{
+			return Points2D.Num();
+		}
+
+		/**
+		 * Get the sub 2d polyline bounded by the input InBoundary in the orientation of the input InOrientation and append it to the output OutPoints
+		 */
+		void GetSubPolyline(const FLinearBoundary& InBoundary, const EOrientation InOrientation, TArray<FPoint2D>& OutPoints) const
+		{
+			TPolylineApproximator<FPoint2D> Approximator(Coordinates, Points2D);
+			Approximator.GetSubPolyline(InBoundary, InOrientation, OutPoints);
+		}
+
+		/**
+		 * Get the sub polyline bounded by the input InBoundary in the orientation of the input InOrientation and append it to the output OutPoints
+		 */
+		void GetSubPolyline(const FLinearBoundary& InBoundary, const EOrientation InOrientation, TArray<FPoint>& OutPoints) const
+		{
+			TPolylineApproximator<FPoint> Approximator3D(Coordinates, Points3D);
+			Approximator3D.GetSubPolyline(InBoundary, InOrientation, OutPoints);
+		}
+
+		/**
+		 * Reserves memory such that the polyline can contain at least Number elements.
+		 *
+		 * @param Number The number of elements that the polyline should be able to contain after allocation.
+		 */
+		void Reserve(int32 Number)
+		{
+			Points3D.Reserve(Number);
+			Points2D.Reserve(Number);
+			Coordinates.Reserve(Number);
+			if (bWithNormals)
+			{
+				Normals.Reserve(Number);
+			}
+		}
+
+		/**
+		 * Empties the polyline.
+		 *
+		 * @param Slack (Optional) The expected usage size after empty operation. Default is 0.
+		 */
+		void Empty(int32 Slack = 0)
+		{
+			Points3D.Empty(Slack);
+			Points2D.Empty(Slack);
+			Normals.Empty(Slack);
+			Coordinates.Empty(Slack);
+		}
+
+		void EmplaceAt(int32 Index, FSurfacicPolyline& Polyline, int32 PointIndex)
+		{
+			Coordinates.EmplaceAt(Index, Polyline.Coordinates[PointIndex]);
+
+			Points2D.EmplaceAt(Index, Polyline.Points2D[PointIndex]);
+			Points3D.EmplaceAt(Index, Polyline.Points3D[PointIndex]);
+			if (bWithNormals)
+			{
+				Normals.EmplaceAt(Index, Polyline.Normals[PointIndex]);
+			}
+			if (bWithTangent)
+			{
+				Tangents.EmplaceAt(Index, Polyline.Tangents[PointIndex]);
+			}
+		}
+
+		void RemoveAt(int32 Index)
+		{
+			Coordinates.RemoveAt(Index);
+			Points2D.RemoveAt(Index);
+			Points3D.RemoveAt(Index);
+			if (bWithNormals)
+			{
+				Normals.RemoveAt(Index);
+			}
+			if (bWithTangent)
+			{
+				Tangents.RemoveAt(Index);
+			}
+		}
+
+		void Pop()
+		{
+			Coordinates.Pop(false);
+			Points2D.Pop(false);
+			Points3D.Pop(false);
+			if (bWithNormals)
+			{
+				Normals.Pop(false);
+			}
+			if (bWithTangent)
+			{
+				Tangents.Pop(false);
+			}
+		}
+
+		void Reverse()
+		{
+			Algo::Reverse(Coordinates);
+			Algo::Reverse(Points2D);
+			Algo::Reverse(Points3D);
+			if (bWithNormals)
+			{
+				Algo::Reverse(Normals);
+			}
+			if (bWithTangent)
+			{
+				Algo::Reverse(Tangents);
+			}
+		}
+
+		double GetLength(const FLinearBoundary& InBoundary) const
+		{
+			TPolylineApproximator<FPoint> Approximator3D(Coordinates, Points3D);
+			return Approximator3D.ComputeLengthOfSubPolyline(InBoundary);
+		}
+
+		double Get2DLength(const FLinearBoundary& InBoundary) const
+		{
+			TPolylineApproximator<FPoint2D> Approximator(Coordinates, Points2D);
+			return Approximator.ComputeLengthOfSubPolyline(InBoundary);
+		}
+
+	};
+
+}
+
diff --git a/Engine/Source/Runtime/Datasmith/CADKernel/Public/CADKernel/Geo/Surfaces/ConeSurface.h b/Engine/Source/Runtime/Datasmith/CADKernel/Public/CADKernel/Geo/Surfaces/ConeSurface.h
new file mode 100644
--- /dev/null
+++ b/Engine/Source/Runtime/Datasmith/CADKernel/Public/CADKernel/Geo/Surfaces/ConeSurface.h
@@ -0,0 +1,85 @@
+// Copyright Epic Games, Inc. All Rights Reserved.
+#pragma once
+
+#include "CADKernel/Geo/Surfaces/Surface.h"
+
+namespace CADKernel
+{
+	class CADKERNEL_API FConeSurface : public FSurface
+	{
+		friend FEntity;
+
+	protected:
+		FMatrixH Matrix;
+		double   StartRadius;
+		double   ConeAngle;
+
+		/**
+		 * A Cone surface is defined by a point on the axis of the cone, the direction of the axis of the cone, 
+		 * the radius of the cone at the axis point and the cone semi-angle. 
+		 * 
+		 * The local coordinate system is defined with the origin at the axis point and the Z axis in the axis direction, 
+		 * The equation of the surface in this system is x2 + y2 – (StartRadius + z*tan(ConeAngle) )2 = 0
+		 * where ConeAngle is the cone semi-angle and StartRadius is the given cone radius. 
+		 * 
+		 * The cone is placed at its final position and orientation by the Matrix
+		 * 
+		 * The bounds of the cone are defined as follow:
+		 * Bounds[EIso::IsoU].Min = StartRuleLength;
+		 * Bounds[EIso::IsoU].Max = EndRuleLength;
+		 * Bounds[EIso::IsoV].Min = StartAngle;
+		 * Bounds[EIso::IsoV].Max = EndAngle;
+		 */
+		FConeSurface(const double InToleranceGeometric, const FMatrixH& InMatrix, double InStartRadius, double InConeAngle, double InStartRuleLength = -HUGE_VALUE, double InEndRuleLength = HUGE_VALUE, double InStartAngle = 0.0, double InEndAngle = 2.0 * PI)
+			: FSurface(InToleranceGeometric, InStartAngle, InEndAngle, InStartRuleLength, InEndRuleLength)
+			, Matrix(InMatrix)
+			, StartRadius(InStartRadius)
+			, ConeAngle(InConeAngle)
+		{
+		}
+
+		FConeSurface(FCADKernelArchive& Archive)
+			: FSurface()
+		{
+			Serialize(Archive);
+		}
+
+	public:
+
+		virtual void Serialize(FCADKernelArchive& Ar) override
+		{
+			FSurface::Serialize(Ar);
+			Ar << Matrix;
+			Ar << StartRadius;
+			Ar << ConeAngle;
+		}
+
+
+		ESurface GetSurfaceType() const
+		{
+			return ESurface::Cone;
+		}
+
+#ifdef CADKERNEL_DEV
+		virtual FInfoEntity& GetInfo(FInfoEntity&) const override;
+#endif
+
+		virtual TSharedPtr<FEntityGeom> ApplyMatrix(const FMatrixH& InMatrix) const override;
+
+		virtual void EvaluatePoint(const FPoint2D& InSurfacicCoordinate, FSurfacicPoint& OutPoint3D, int32 InDerivativeOrder = 0) const override;
+		virtual void EvaluatePointGrid(const FCoordinateGrid& Coordinates, FSurfacicSampling& OutPoints, bool bComputeNormals = false) const override;
+
+		virtual void Presample(const FSurfacicBoundary& InBoundaries, FCoordinateGrid& OutCoordinates) override 
+		{
+			PresampleIsoCircle(InBoundaries, OutCoordinates, EIso::IsoU);
+
+			OutCoordinates[EIso::IsoV].Empty(3);
+			OutCoordinates[EIso::IsoV].Add(InBoundaries.UVBoundaries[EIso::IsoV].Min);
+			OutCoordinates[EIso::IsoV].Add((InBoundaries.UVBoundaries[EIso::IsoV].Max + InBoundaries.UVBoundaries[EIso::IsoV].Min) / 2.0);
+			OutCoordinates[EIso::IsoV].Add(InBoundaries.UVBoundaries[EIso::IsoV].Max);
+		}
+
+
+	};
+}
+
diff --git a/Engine/Source/Runtime/Datasmith/CADKernel/Public/CADKernel/Geo/Surfaces/CylinderSurface.h b/Engine/Source/Runtime/Datasmith/CADKernel/Public/CADKernel/Geo/Surfaces/CylinderSurface.h
new file mode 100644
--- /dev/null
+++ b/Engine/Source/Runtime/Datasmith/CADKernel/Public/CADKernel/Geo/Surfaces/CylinderSurface.h
@@ -0,0 +1,74 @@
+// Copyright Epic Games, Inc. All Rights Reserved.
+#pragma once
+
+#include "CADKernel/Geo/Surfaces/Surface.h"
+
+namespace CADKernel
+{
+	class FNURBSSurface;
+
+	class CADKERNEL_API FCylinderSurface : public FSurface
+	{
+		friend FEntity;
+
+	protected:
+		FMatrixH Matrix;
+		double Radius;
+
+		/**
+		 * The Cylinder surface is a right circular cylinder surface starting at the axis origin and with z axis as the axis direction. 
+		 * It is defined by its Radius and its bounds (StartAngle, EndAngle, StartLength, EndLength)
+		 * 
+		 * The surface is placed at its final position and orientation by the Matrix
+		 *
+		 * The bounds of the cylinder are defined as follow:
+		 * Bounds[EIso::IsoU].Min = StartAngle;
+		 * Bounds[EIso::IsoU].Max = EndAngle;
+		 * Bounds[EIso::IsoV].Min = StartLength;
+		 * Bounds[EIso::IsoV].Max = EndLength;
+		 */
+		FCylinderSurface(const double InToleranceGeometric, const FMatrixH& InMatrix, double InRadius, double InStartLength = -HUGE_VALUE, double InEndLength = HUGE_VALUE, double InStartAngle = 0.0, double InEndAngle = 2.0 * PI);
+
+		FCylinderSurface(FCADKernelArchive& Archive)
+			: FSurface()
+		{
+			Serialize(Archive);
+		}
+
+	public:
+
+		virtual void Serialize(FCADKernelArchive& Ar) override
+		{
+			FSurface::Serialize(Ar);
+			Ar << Matrix;
+			Ar << Radius;
+		}
+
+		virtual TSharedPtr<FEntityGeom> ApplyMatrix(const FMatrixH& InMatrix) const override;
+
+#ifdef CADKERNEL_DEV
+		virtual FInfoEntity& GetInfo(FInfoEntity&) const override;
+#endif
+
+		ESurface GetSurfaceType() const
+		{
+			return ESurface::Cylinder;
+		}
+
+		virtual void InitBoundary() const override;
+
+		virtual void EvaluatePoint(const FPoint2D& InSurfacicCoordinate, FSurfacicPoint& OutPoint3D, int32 InDerivativeOrder = 0) const override;
+
+		virtual void Presample(const FSurfacicBoundary& InBoundaries, FCoordinateGrid& OutCoordinates) override
+		{
+			PresampleIsoCircle(InBoundaries, OutCoordinates, EIso::IsoU);
+
+			OutCoordinates[EIso::IsoV].Empty(3);
+			OutCoordinates[EIso::IsoV].Add(InBoundaries.UVBoundaries[EIso::IsoV].Min);
+			OutCoordinates[EIso::IsoV].Add((InBoundaries.UVBoundaries[EIso::IsoV].Max + InBoundaries.UVBoundaries[EIso::IsoV].Min) / 2.0);
+			OutCoordinates[EIso::IsoV].Add(InBoundaries.UVBoundaries[EIso::IsoV].Max);
+		}
+
+	};
+}
+
diff --git a/Engine/Source/Runtime/Datasmith/CADKernel/Public/CADKernel/Geo/Surfaces/NURBSSurface.h b/Engine/Source/Runtime/Datasmith/CADKernel/Public/CADKernel/Geo/Surfaces/NURBSSurface.h
new file mode 100644
--- /dev/null
+++ b/Engine/Source/Runtime/Datasmith/CADKernel/Public/CADKernel/Geo/Surfaces/NURBSSurface.h
@@ -0,0 +1,184 @@
+// Copyright Epic Games, Inc. All Rights Reserved.
+#pragma once
+
+#include "CADKernel/Geo/Surfaces/Surface.h"
+#include "CADKernel/Math/BSpline.h"
+
+namespace CADKernel
+{
+
+	class CADKERNEL_API FNURBSSurface : public FSurface
+	{
+		friend FEntity;
+
+	protected:
+		int32 PoleUNum;
+		int32 PoleVNum;
+
+		int32 UDegree;
+		int32 VDegree;
+
+		TArray<double> UNodalVector;
+		TArray<double> VNodalVector;
+
+		TArray<double> Weights;
+
+		TArray<FPoint> Poles;
+
+		bool bIsRational;
+
+		/**
+		 * Data generated at initialization which are not serialized
+		 */
+		TArray<double> HomogeneousPoles;
+
+		/**
+		 * Build a Non uniform B-Spline surface
+		 * @param NodalVectorU Its size is the number of poles in U + the surface degree in U + 1 (PoleUNum + UDegre + 1)
+		 * @param NodalVectorV Its size is the number of poles in V + the surface degree in V + 1 (PoleVNum + VDegre + 1)
+		 */
+		FNURBSSurface(const double InToleranceGeometric, int32 InPoleUNum, int32 InPoleVNum, int32 InDegreU, int32 InDegreV, const TArray<double>& InNodalVectorU, const TArray<double>& InNodalVectorV, const TArray<FPoint>& InPoles)
+			: FSurface(InToleranceGeometric)
+			, PoleUNum(InPoleUNum)
+			, PoleVNum(InPoleVNum)
+			, UDegree(InDegreU)
+			, VDegree(InDegreV)
+			, UNodalVector(InNodalVectorU)
+			, VNodalVector(InNodalVectorV)
+			, Poles(InPoles)
+			, bIsRational(false)
+		{
+			Finalize();
+		}
+
+		/**
+		 * Build a Non uniform rational B-Spline surface
+		 * @param NodalVectorU Its size is the number of poles in U + the surface degree in U + 1 (PoleUNum + UDegre + 1)
+		 * @param NodalVectorV Its size is the number of poles in V + the surface degree in V + 1 (PoleVNum + VDegre + 1)
+		 */
+		FNURBSSurface(const double InToleranceGeometric, int32 InPoleUNum, int32 InPoleVNum, int32 InDegreU, int32 InDegreV, const TArray<double>& InNodalVectorU, const TArray<double>& InNodalVectorV, const TArray<FPoint>& InPoles, const TArray<double>& InWeights)
+			: FSurface(InToleranceGeometric)
+			, PoleUNum(InPoleUNum)
+			, PoleVNum(InPoleVNum)
+			, UDegree(InDegreU)
+			, VDegree(InDegreV)
+			, UNodalVector(InNodalVectorU)
+			, VNodalVector(InNodalVectorV)
+			, Weights(InWeights)
+			, Poles(InPoles)
+			, bIsRational(true)
+		{
+			Finalize();
+		}
+
+		FNURBSSurface(FCADKernelArchive& Archive)
+			: FSurface()
+		{
+			Serialize(Archive);
+			Finalize();
+		}
+
+	public:
+
+		virtual void Serialize(FCADKernelArchive& Ar) override
+		{
+			FSurface::Serialize(Ar);
+			Ar << PoleUNum;
+			Ar << PoleVNum;
+			Ar << UDegree;
+			Ar << VDegree;
+			Ar << UNodalVector;
+			Ar << VNodalVector;
+			Ar << Weights;
+			Ar << Poles;
+		}
+
+		ESurface GetSurfaceType() const
+		{
+			return ESurface::Nurbs;
+		}
+
+#ifdef CADKERNEL_DEV
+		virtual FInfoEntity& GetInfo(FInfoEntity&) const override;
+#endif
+
+		constexpr const int32 GetDegree(EIso Iso) const
+		{
+			switch (Iso)
+			{
+			case EIso::IsoU:
+				return UDegree;
+			case EIso::IsoV:
+			default:
+				return VDegree;
+			}
+		}
+
+		constexpr const int32 GetPoleCount(EIso Iso) const
+		{
+			switch (Iso)
+			{
+			case EIso::IsoU:
+				return PoleUNum;
+			case EIso::IsoV:
+			default:
+				return PoleVNum;
+			}
+		}
+
+		const TArray<FPoint>& GetPoles() const
+		{
+			return Poles;
+		}
+
+		const TArray<double>& GetWeights() const
+		{
+			return Weights;
+		}
+
+		TArray<double> GetHPoles() const
+		{
+			return HomogeneousPoles;
+		}
+
+		constexpr const TArray<double>& GetNodalVector(EIso Iso) const
+		{
+			switch (Iso)
+			{
+			case EIso::IsoU:
+				return UNodalVector;
+			case EIso::IsoV:
+			default:
+				return VNodalVector;
+			}
+		}
+
+		bool IsRational() const
+		{
+			return bIsRational;
+		}
+
+
+		virtual TSharedPtr<FEntityGeom> ApplyMatrix(const FMatrixH& InMatrix) const override;
+
+		virtual void EvaluatePoint(const FPoint2D& InSurfacicCoordinate, FSurfacicPoint& OutPoint3D, int32 InDerivativeOrder = 0) const override
+		{
+			BSpline::EvaluatePoint(*this, InSurfacicCoordinate, OutPoint3D, InDerivativeOrder);
+		}
+
+		void EvaluatePointGrid(const FCoordinateGrid& InSurfacicCoordinates, FSurfacicSampling& OutPoints, bool bComputeNormals) const override
+		{
+			BSpline::EvaluatePointGrid(*this, InSurfacicCoordinates, OutPoints, bComputeNormals);
+		}
+
+		virtual void LinesNotDerivables(const FSurfacicBoundary& InBoundary, int32 InDerivativeOrder, FCoordinateGrid& OutNotDerivableCoordinates) const override
+		{
+			BSpline::FindNotDerivableParameters(*this, InDerivativeOrder, InBoundary, OutNotDerivableCoordinates);
+		}
+
+	private:
+		void Finalize();
+	};
+
+} // namespace CADKernel
+
diff --git a/Engine/Source/Runtime/Datasmith/CADKernel/Public/CADKernel/Geo/Surfaces/PlaneSurface.h b/Engine/Source/Runtime/Datasmith/CADKernel/Public/CADKernel/Geo/Surfaces/PlaneSurface.h
new file mode 100644
--- /dev/null
+++ b/Engine/Source/Runtime/Datasmith/CADKernel/Public/CADKernel/Geo/Surfaces/PlaneSurface.h
@@ -0,0 +1,106 @@
+// Copyright Epic Games, Inc. All Rights Reserved.
+#pragma once
+
+#include "CADKernel/Geo/Surfaces/Surface.h"
+#include "CADKernel/Math/Plane.h"
+
+namespace CADKernel
+{
+	class CADKERNEL_API FPlaneSurface : public FSurface
+	{
+		friend FEntity;
+
+	protected:
+
+		FMatrixH Matrix;
+		FMatrixH InverseMatrix;
+
+		/**
+		 * The plane surface is the plane XY.
+		 * The surface is placed at its final position and orientation by the Matrix
+		 */
+		FPlaneSurface(const double InToleranceGeometric, const FMatrixH& InMatrix);
+
+		/**
+		 * The plane surface is the plane XY.
+		 * The surface is placed at its final position and orientation by the Matrix,
+		 * The matrix is calculated from the normal plane at its final position and its distance from the origin along the normal.
+		 */
+		FPlaneSurface(const double InToleranceGeometric, double DistanceFromOrigin, FPoint Normal)
+			: FPlaneSurface(InToleranceGeometric, Normal* DistanceFromOrigin, Normal)
+		{
+		}
+
+		/**
+		 * The plane surface is the plane XY.
+		 * The surface is placed at its final position and orientation by the Matrix 
+		 * The matrix is calculated from the plan origine at its final position and its final normal
+		 */
+		FPlaneSurface(const double InToleranceGeometric, const FPoint& Position, FPoint Normal);
+
+		FPlaneSurface(FCADKernelArchive& Archive)
+			: FSurface()
+		{
+			Serialize(Archive);
+		}
+		
+
+	public:
+
+		virtual void Serialize(FCADKernelArchive& Ar) override
+		{
+			FSurface::Serialize(Ar);
+			Ar << Matrix;
+			Ar << InverseMatrix;
+		}
+
+		ESurface GetSurfaceType() const
+		{
+			return ESurface::Plane;
+		}
+
+		const FMatrixH& GetMatrix() const
+		{
+			return Matrix;
+		}
+
+		FPlane GetPlane() const;
+
+#ifdef CADKERNEL_DEV
+		virtual FInfoEntity& GetInfo(FInfoEntity&) const override;
+#endif
+
+		virtual TSharedPtr<FEntityGeom> ApplyMatrix(const FMatrixH& InMatrix) const override;
+
+		virtual void EvaluatePoint(const FPoint2D& InSurfacicCoordinate, FSurfacicPoint& OutPoint3D, int32 InDerivativeOrder = 0) const override;
+		virtual void EvaluatePoints(const TArray<FPoint2D>& InSurfacicCoordinates, TArray<FSurfacicPoint>& OutPoint3D, int32 InDerivativeOrder = 0) const override;
+		virtual void EvaluatePointGrid(const FCoordinateGrid& Coordinates, FSurfacicSampling& OutPoints, bool bComputeNormals = false) const override;
+
+		virtual FPoint ProjectPoint(const FPoint& InPoint, FPoint* OutProjectedPoint = nullptr) const;// override;
+		virtual void ProjectPoints(const TArray<FPoint>& InPoints, TArray<FPoint>* OutProjectedPointCoordinates, TArray<FPoint>* OutProjectedPoints = nullptr) const;// override;
+
+		virtual void Presample(const FSurfacicBoundary& InBoundaries, FCoordinateGrid& OutCoordinates) override
+		{
+			OutCoordinates[EIso::IsoU].Empty(3);
+			OutCoordinates[EIso::IsoU].Add(InBoundaries.UVBoundaries[EIso::IsoU].Min);
+			OutCoordinates[EIso::IsoU].Add(InBoundaries.UVBoundaries[EIso::IsoU].GetMiddle());
+			OutCoordinates[EIso::IsoU].Add(InBoundaries.UVBoundaries[EIso::IsoU].Max);
+
+			OutCoordinates[EIso::IsoV].Empty(3);
+			OutCoordinates[EIso::IsoV].Add(InBoundaries.UVBoundaries[EIso::IsoV].Min);
+			OutCoordinates[EIso::IsoV].Add(InBoundaries.UVBoundaries[EIso::IsoV].GetMiddle());
+			OutCoordinates[EIso::IsoV].Add(InBoundaries.UVBoundaries[EIso::IsoV].Max);
+		}
+
+		virtual void IsSurfaceClosed(bool& bOutClosedAlongU, bool& bOutClosedAlongV) const override
+		{
+			bOutClosedAlongU = false;
+			bOutClosedAlongV = false;
+		}
+
+	private:
+		virtual const void ComputeIsoTolerances() const override;
+
+	};
+
+} // namespace CADKernel
diff --git a/Engine/Source/Runtime/Datasmith/CADKernel/Public/CADKernel/Geo/Surfaces/SphericalSurface.h b/Engine/Source/Runtime/Datasmith/CADKernel/Public/CADKernel/Geo/Surfaces/SphericalSurface.h
new file mode 100644
--- /dev/null
+++ b/Engine/Source/Runtime/Datasmith/CADKernel/Public/CADKernel/Geo/Surfaces/SphericalSurface.h
@@ -0,0 +1,72 @@
+// Copyright Epic Games, Inc. All Rights Reserved.
+#pragma once
+
+#include "CADKernel/Geo/Surfaces/Surface.h"
+
+namespace CADKernel
+{
+	class CADKERNEL_API FSphericalSurface : public FSurface
+	{
+		friend FEntity;
+
+	protected:
+		FMatrixH Matrix;
+		double Radius;
+
+		/**
+		 * The spherical surface is defined its radius. 
+		 *
+		 * It's defined as the rotation around Z axis of an semicircle defined in the plan XY centered at the origin. 
+		 *
+		 * The cone is placed at its final position and orientation by the Matrix
+		 * 
+		 * The bounds of the spherical surface are defined as follow:
+		 * Bounds[EIso::IsoU].Min = MeridianStartAngle
+		 * Bounds[EIso::IsoU].Max = MeridianEndAngle
+		 * Bounds[EIso::IsoV].Min = ParallelStartAngle
+		 * Bounds[EIso::IsoV].Max = ParallelEndAngle
+		 */
+		FSphericalSurface(const double InToleranceGeometric, const FMatrixH& InMatrix, double InRadius, double InParallelStartAngle = 0.0, double InParallelEndAngle = 2.0 * PI, double InMeridianStartAngle = -PI / 2.0, double InMeridianEndAngle = PI / 2.0)
+			: FSurface(InToleranceGeometric, InParallelStartAngle, InParallelEndAngle, InMeridianStartAngle, InMeridianEndAngle)
+			, Matrix(InMatrix)
+			, Radius(InRadius)
+		{
+		}
+
+		FSphericalSurface(FCADKernelArchive& Archive)
+			: FSurface()
+		{
+			Serialize(Archive);
+		}
+
+	public:
+
+		virtual void Serialize(FCADKernelArchive& Ar) override
+		{
+			FSurface::Serialize(Ar);
+			Ar << Matrix;
+			Ar << Radius;
+		}
+
+		ESurface GetSurfaceType() const
+		{
+			return ESurface::Sphere;
+		}
+
+#ifdef CADKERNEL_DEV
+		virtual FInfoEntity& GetInfo(FInfoEntity&) const override;
+#endif
+
+		virtual TSharedPtr<FEntityGeom> ApplyMatrix(const FMatrixH& InMatrix) const override;
+
+		virtual void EvaluatePoint(const FPoint2D& InSurfacicCoordinate, FSurfacicPoint& OutPoint3D, int32 InDerivativeOrder = 0) const override;
+		virtual void EvaluatePointGrid(const FCoordinateGrid& Coordinates, FSurfacicSampling& OutPoints, bool bComputeNormals = false) const override;
+
+		virtual void Presample(const FSurfacicBoundary& InBoundaries, FCoordinateGrid& OutCoordinates) override
+		{
+			PresampleIsoCircle(InBoundaries, OutCoordinates, EIso::IsoU);
+			PresampleIsoCircle(InBoundaries, OutCoordinates, EIso::IsoV);
+		}
+
+	};
+}
diff --git a/Engine/Source/Runtime/Datasmith/CADKernel/Public/CADKernel/Geo/Surfaces/Surface.h b/Engine/Source/Runtime/Datasmith/CADKernel/Public/CADKernel/Geo/Surfaces/Surface.h
new file mode 100644
--- /dev/null
+++ b/Engine/Source/Runtime/Datasmith/CADKernel/Public/CADKernel/Geo/Surfaces/Surface.h
@@ -0,0 +1,230 @@
+// Copyright Epic Games, Inc. All Rights Reserved.
+#pragma once
+
+#include "CADKernel/Core/CADKernelArchive.h"
+#include "CADKernel/Core/EntityGeom.h"
+
+#include "CADKernel/Geo/GeoEnum.h"
+#include "CADKernel/Geo/GeoPoint.h"
+#include "CADKernel/Math/Boundary.h"
+#include "CADKernel/Math/Point.h"
+#include "CADKernel/Utils/Cache.h"
+
+namespace CADKernel
+{
+	class FGrid;
+	class FInfoEntity;
+	class FSurfacicPolyline;
+
+	struct FCurvePoint;
+	struct FCurvePoint2D;
+	struct FSurfacicSampling;
+
+	class CADKERNEL_API FSurface : public FEntityGeom
+	{
+		friend FEntity;
+
+	protected:
+
+		double Tolerance3D;
+		mutable TCache<FSurfacicTolerance> ToleranceIsos;
+		mutable FSurfacicBoundary Boundary;
+
+	private:
+
+		/**
+		 * Return the tolerance in the parametric space of the surface along isoU and isoV i.e Distance(Point(U, V), Point(U + ToleranceIsoU, V)) = Tolerance3D
+		 * With Tolerance3D = FSysteme.GeometricalTolerance
+		 */
+		const virtual void ComputeIsoTolerances() const;
+
+		virtual void InitBoundary() const
+		{
+			Boundary.Set();
+		}
+
+	protected:
+
+		FSurface() = default;
+
+		FSurface(double InToleranceGeometric)
+			: FEntityGeom()
+			, Tolerance3D(InToleranceGeometric)
+		{
+
+		}
+
+		FSurface(double InToleranceGeometric, double UMin, double UMax, double VMin, double VMax)
+			: FEntityGeom()
+			, Tolerance3D(InToleranceGeometric)
+			, Boundary(UMin, UMax, VMin, VMax)
+		{
+		}
+
+	public:
+
+
+		virtual void Serialize(FCADKernelArchive& Ar) override
+		{
+			// Surface's type is serialize because it is used to instantiate the correct entity on deserialization (@see Deserialize(FCADKernelArchive& Archive)) 
+			if (Ar.IsSaving())
+			{
+				ESurface SurfaceType = GetSurfaceType();
+				Ar << SurfaceType;
+			}
+			FEntityGeom::Serialize(Ar);
+
+			Ar << Tolerance3D;
+			Ar << ToleranceIsos;
+			Ar << Boundary;
+		}
+
+		/**
+		 * Specific method for surface family to instantiate the correct derived class of FCurve
+		 */
+		static TSharedPtr<FSurface> Deserialize(FCADKernelArchive& Archive);
+
+		virtual ~FSurface() = default;
+
+#ifdef CADKERNEL_DEV
+		virtual FInfoEntity& GetInfo(FInfoEntity&) const override;
+#endif
+
+		virtual EEntity GetEntityType() const override
+		{
+			return EEntity::Surface;
+		}
+
+		virtual ESurface GetSurfaceType() const = 0;
+
+		const FSurfacicBoundary& GetBoundary() const
+		{
+			if (!Boundary.IsValid())
+			{
+				InitBoundary();
+			}
+			return Boundary;
+		};
+
+		void ExtendBoundaryTo(const FSurfacicBoundary MaxLimit)
+		{
+			Boundary.ExtendTo(MaxLimit);
+		}
+
+		virtual TSharedPtr<FEntityGeom> ApplyMatrix(const FMatrixH& InMatrix) const = 0;
+
+		/**
+		 * Return the tolerance in the parametric space of the surface i.e Distance(Point(U, V), Point(U', V')) = Tolerance3D with Distance((U,V), (U', V')) = Tolerance2D
+		 */
+		const virtual double Get2DTolerance() const
+		{
+			if (!ToleranceIsos.IsValid())
+			{
+				GetIsoTolerances();
+			}
+			FSurfacicTolerance& Tolerance = ToleranceIsos;
+			return FMath::Min(Tolerance[EIso::IsoU], Tolerance[EIso::IsoV]);
+		}
+
+		const double Get3DTolerance()
+		{
+			return Tolerance3D;
+		}
+
+		/**
+		 * Return the tolerance in the parametric space of the surface along isoU and isoV i.e Distance(Point(U, V), Point(U + ToleranceIsoU, V)) = Tolerance3D
+		 * With Tolerance3D = FSysteme.GeometricalTolerance
+		 */
+		const FSurfacicTolerance& GetIsoTolerances() const
+		{
+			if (!ToleranceIsos.IsValid())
+			{
+				ComputeIsoTolerances();
+			}
+			return ToleranceIsos;
+		}
+
+		/**
+		 * Return the tolerance in the parametric space of the surface along the specified axis 
+		 * With Tolerance3D = FSysteme.GeometricalTolerance
+		 * @see GetIsoTolerances
+		 */
+		const double& GetIsoTolerance(EIso Iso) const
+		{
+			if (!ToleranceIsos.IsValid())
+			{
+				ComputeIsoTolerances();
+			}
+			return ((FSurfacicTolerance) ToleranceIsos)[Iso];
+		}
+
+		virtual void EvaluatePoint(const FPoint2D& InSurfacicCoordinate, FSurfacicPoint& OutPoint3D, int32 InDerivativeOrder = 0) const = 0;
+
+		virtual void EvaluatePoints(const TArray<FPoint2D>& InSurfacicCoordinates, TArray<FSurfacicPoint>& OutPoint3D, int32 InDerivativeOrder = 0) const;
+		virtual void EvaluatePoints(const TArray<FCurvePoint2D>& InSurfacicCoordinates, TArray<FSurfacicPoint>& OutPoint3D, int32 InDerivativeOrder = 0) const;
+		virtual void EvaluatePoints(const TArray<FCurvePoint2D>& InSurfacicCoordinates, TArray<FCurvePoint>& OutPoint3D, int32 InDerivativeOrder = 0) const;
+		virtual void EvaluatePoints(FSurfacicPolyline& Polyline) const;
+		virtual void EvaluatePoints(const TArray<FCurvePoint2D>& Points2D, FSurfacicPolyline& Polyline) const;
+
+		virtual void EvaluatePointGrid(const FCoordinateGrid& Coordinates, FSurfacicSampling& OutPoints, bool bComputeNormals = false) const;
+		void EvaluateGrid(FGrid& Grid) const;
+
+		virtual void EvaluateNormals(const TArray<FPoint2D>& Points2D, TArray<FPoint>& Normals) const;
+
+		/**
+		 * Divide the parametric space in the desired number of regular subdivisions and compute the associated PointGrid
+		 */
+		void Sample(const FSurfacicBoundary& Bounds, int32 NumberOfSubdivisions[2], FSurfacicSampling& OutPointSampling) const;
+
+		/**
+		 * Divide the parametric space in the desired number of regular subdivisions and compute the associated CoordinateGrid
+		 */
+		void Sample(const FSurfacicBoundary& Boundary, int32 NumberOfSubdivisions[2], FCoordinateGrid& OutCoordinateSampling) const;
+
+		/**
+		 * Generate a pre-sampling of the surface saved in OutCoordinates.
+		 * This sampling is light enough to allow a fast computation of the grid, precise enough to compute accurately meshing criteria
+		 */
+		virtual void Presample(const FSurfacicBoundary& InBoundaries, FCoordinateGrid& OutCoordinates);
+
+		virtual void LinesNotDerivables(const FSurfacicBoundary& Bounds, int32 InDerivativeOrder, FCoordinateGrid& OutNotDerivableCoordinates) const
+		{
+			OutNotDerivableCoordinates.Empty();
+		}
+
+		void LinesNotDerivables(int32 InDerivativeOrder, FCoordinateGrid& OutNotDerivableCoordinates) const
+		{
+			LinesNotDerivables(Boundary, InDerivativeOrder, OutNotDerivableCoordinates);
+		}
+
+		/**
+		 * A surface is closed along an iso axis if it's connected to itself e.g. the edge carried by the iso curve U = UMin is linked to the edge carried by the iso curve U = UMax
+		 * A complete sphere, torus are example of surface closed along U and V
+		 * A cylinder with circle section or a cone are example of surface closed along an iso
+		 */
+		virtual void IsSurfaceClosed(bool& bOutClosedAlongU, bool& bOutClosedAlongV) const
+		{
+			bOutClosedAlongU = false;
+			bOutClosedAlongV = false;
+		}
+
+		void PresampleIsoCircle(const FSurfacicBoundary& InBoundaries, FCoordinateGrid& Coordinates, EIso Iso)
+		{
+			constexpr const double ThreeOverPi = 3 / PI;
+			double Delta = InBoundaries.Length(Iso);
+			int32 SampleCount = /*FMath::Max(3, */ (int32)(Delta * ThreeOverPi + 1)/*)*/;
+			Delta /= SampleCount;
+
+			Coordinates[Iso].Empty(SampleCount + 1);
+			double Sample = InBoundaries[Iso].GetMin();
+			for (int32 Index = 0; Index <= SampleCount; ++Index)
+			{
+				Coordinates[Iso].Add(Sample);
+				Sample += Delta;
+			}
+		}
+
+
+	};
+} // namespace CADKernel
+
diff --git a/Engine/Source/Runtime/Datasmith/CADKernel/Public/CADKernel/Math/Aabb.h b/Engine/Source/Runtime/Datasmith/CADKernel/Public/CADKernel/Math/Aabb.h
new file mode 100644
--- /dev/null
+++ b/Engine/Source/Runtime/Datasmith/CADKernel/Public/CADKernel/Math/Aabb.h
@@ -0,0 +1,313 @@
+// Copyright Epic Games, Inc. All Rights Reserved.
+#pragma once
+
+#include "CADKernel/Math/Point.h"
+
+
+enum EAABBBoundary : uint32
+{
+	XMax = 0x00000000u,
+	YMax = 0x00000000u,
+	ZMax = 0x00000000u,
+	XMin = 0x00000001u,
+	YMin = 0x00000002u,
+	ZMin = 0x00000004u,
+};
+
+ENUM_CLASS_FLAGS(EAABBBoundary);
+
+namespace CADKernel
+{
+	class FSphere;
+	class FPlane;
+
+	class FAABB
+	{
+
+	private:
+		FPoint MinCorner;
+		FPoint MaxCorner;
+
+	public:
+		FAABB()
+			: MinCorner(HUGE_VALUE, HUGE_VALUE, HUGE_VALUE)
+			, MaxCorner(-HUGE_VALUE, -HUGE_VALUE, -HUGE_VALUE)
+		{
+		}
+
+		FAABB(const FPoint& InMinCorner, const FPoint& InMaxCorner)
+			: MinCorner(InMinCorner)
+			, MaxCorner(InMaxCorner)
+		{
+		}
+
+		bool IsValid() const
+		{
+			return (MinCorner.X < MaxCorner.Y);
+		}
+
+		void Empty()
+		{
+			MinCorner.Set(HUGE_VALUE, HUGE_VALUE, HUGE_VALUE);
+			MaxCorner.Set(-HUGE_VALUE, -HUGE_VALUE, -HUGE_VALUE);
+		}
+
+		bool Contains(const FPoint& Point) const
+		{
+			return ((Point.X > MinCorner.X - SMALL_NUMBER) && (Point.X < MaxCorner.X + SMALL_NUMBER) &&
+				(Point.Y > MinCorner.Y - SMALL_NUMBER) && (Point.Y < MaxCorner.Y + SMALL_NUMBER) &&
+				(Point.Z > MinCorner.Z - SMALL_NUMBER) && (Point.Z < MaxCorner.Z + SMALL_NUMBER));
+		}
+
+		void SetMinDimension(double MinDimension)
+		{
+			for (int32 Axis = 0; Axis < 3; Axis++)
+			{
+				double AxisDimension = GetDimension(Axis);
+				if (AxisDimension < MinDimension)
+				{
+					double Offset = (MinDimension - AxisDimension) / 2;
+					MinCorner[Axis] -= Offset;
+					MaxCorner[Axis] += Offset;
+				}
+			}
+		}
+
+		double GetMaxDimension() const
+		{
+			double MaxDimension = 0;
+			for (int32 Index = 0; Index < 3; Index++)
+			{
+				double Dimension = GetDimension(Index);
+				if (Dimension > MaxDimension)
+				{
+					MaxDimension = Dimension;
+				}
+			}
+			return MaxDimension;
+		}
+
+		double GetDimension(int32 Axis) const
+		{
+			return MaxCorner[Axis] - MinCorner[Axis];
+		}
+
+		bool Contains(const FAABB& Aabb) const
+		{
+			return IsValid() && Aabb.IsValid() && Contains(Aabb.MinCorner) && Contains(Aabb.MaxCorner);
+		}
+
+		const FPoint& GetMin() const
+		{
+			return MinCorner;
+		}
+
+		const FPoint& GetMax() const
+		{
+			return MaxCorner;
+		}
+
+		FPoint GetCorner(int32 Corner) const
+		{
+			return FPoint(
+				Corner & EAABBBoundary::XMin ? MinCorner[0] : MaxCorner[0],
+				Corner & EAABBBoundary::YMin ? MinCorner[1] : MaxCorner[1],
+				Corner & EAABBBoundary::ZMin ? MinCorner[2] : MaxCorner[2]
+			);
+		}
+
+		FAABB& operator+= (const FPoint& Point)
+		{
+			for (int32 Index = 0; Index < 3; Index++)
+			{
+				if (Point[Index] < MinCorner[Index])
+				{
+					MinCorner[Index] = Point[Index];
+				}
+				if (Point[Index] > MaxCorner[Index])
+				{
+					MaxCorner[Index] = Point[Index];
+				}
+			}
+			return *this;
+		}
+
+		void Offset (double Offset)
+		{
+			for (int32 Index = 0; Index < 3; Index++)
+			{
+				MinCorner[Index] -= Offset;
+				MaxCorner[Index] += Offset;
+			}
+		}
+
+		FAABB& operator+= (const FAABB& aabb)
+		{
+			*this += aabb.MinCorner;
+			*this += aabb.MaxCorner;
+			return *this;
+		}
+
+		FAABB operator+ (const FPoint& Point) const
+		{
+			FAABB Other = *this;
+			Other += Point;
+			return Other;
+		}
+
+
+		FAABB operator+ (const FAABB& Aabb) const
+		{
+			FAABB Other = *this;
+			Other += Aabb;
+			return Other;
+		}
+
+	};
+
+	class CADKERNEL_API FAABB2D
+	{
+	private:
+
+		FPoint2D MinCorner;
+		FPoint2D MaxCorner;
+
+	public:
+
+		FAABB2D()
+			: MinCorner(HUGE_VALUE, HUGE_VALUE)
+			, MaxCorner(-HUGE_VALUE, -HUGE_VALUE)
+		{
+		}
+
+		FAABB2D(const FPoint2D& InMinCorner, const FPoint2D& InMaxCorner)
+			: MinCorner(InMinCorner)
+			, MaxCorner(InMaxCorner)
+		{
+		}
+
+		friend FArchive& operator<<(FArchive& Ar, FAABB2D& AABB)
+		{
+			Ar << AABB.MinCorner;
+			Ar << AABB.MaxCorner;
+			return Ar;
+		}
+
+		const FPoint2D& GetMin() const
+		{
+			return MinCorner;
+		}
+
+		const FPoint2D& GetMax() const
+		{
+			return MaxCorner;
+		}
+
+		bool IsValid() const
+		{
+			return (MinCorner.U < MaxCorner.U);
+		}
+
+		void Empty()
+		{
+			MinCorner.Set(HUGE_VALUE, HUGE_VALUE);
+			MaxCorner.Set(-HUGE_VALUE, -HUGE_VALUE);
+		}
+
+		bool Contains(const FPoint2D& Point) const
+		{
+			return ((Point.U > MinCorner.U - SMALL_NUMBER) && (Point.U < MaxCorner.U + SMALL_NUMBER) &&
+				(Point.V > MinCorner.V - SMALL_NUMBER) && (Point.V < MaxCorner.V + SMALL_NUMBER));
+		}
+
+		FPoint2D GetCorner(int32 CornerIndex) const
+		{
+			return FPoint2D(
+				CornerIndex & EAABBBoundary::XMin ? MinCorner[0] : MaxCorner[0],
+				CornerIndex & EAABBBoundary::YMin ? MinCorner[1] : MaxCorner[1]
+			);
+		}
+
+		bool Contains(const FAABB2D& Aabb2d) const
+		{
+			return IsValid() && Aabb2d.IsValid() && Contains(Aabb2d.MinCorner) && Contains(Aabb2d.MaxCorner);
+		}
+
+		double DiagonalLength() const
+		{
+			return MinCorner.Distance(MaxCorner);
+		}
+
+		FPoint2D Diagonal() const
+		{
+			return MaxCorner - MinCorner;
+		}
+
+		double GetDimension(int32 Axis) const
+		{
+			return MaxCorner[Axis] - MinCorner[Axis];
+		}
+
+		FPoint2D Center() const
+		{
+			return (MinCorner + MaxCorner) / 2.;
+		}
+
+		void Set(const FPoint2D& InMinCorner, const FPoint2D& InMaxCorner)
+		{
+			MinCorner = InMinCorner;
+			MaxCorner = InMaxCorner;
+		}
+
+		FAABB2D& operator+= (const FAABB2D& Aabb2d)
+		{
+			*this += Aabb2d.GetMin();
+			*this += Aabb2d.GetMax();
+			return *this;
+		}
+
+		FAABB2D& operator+= (TArray<FPoint2D> Points)
+		{
+			for (FPoint Point : Points)
+			{
+				*this += Point;
+			}
+			return *this;
+		}
+
+		FAABB2D& operator+= (const FPoint2D& Point)
+		{
+			if (Point.U < MinCorner.U)
+			{
+				MinCorner.U = Point.U;
+			}
+
+			if (Point.V < MinCorner.V)
+			{
+				MinCorner.V = Point.V;
+			}
+
+			if (Point.U > MaxCorner.U)
+			{
+				MaxCorner.U = Point.U;
+			}
+
+			if (Point.V > MaxCorner.V)
+			{
+				MaxCorner.V = Point.V;
+			}
+
+			return *this;
+		}
+
+		void Offset(double Offset)
+		{
+			MinCorner.U -= Offset;
+			MinCorner.V -= Offset;
+			MaxCorner.U += Offset;
+			MaxCorner.V += Offset;
+		}
+	};
+
+} // namespace CADKernel
+
diff --git a/Engine/Source/Runtime/Datasmith/CADKernel/Public/CADKernel/Math/Boundary.h b/Engine/Source/Runtime/Datasmith/CADKernel/Public/CADKernel/Math/Boundary.h
new file mode 100644
--- /dev/null
+++ b/Engine/Source/Runtime/Datasmith/CADKernel/Public/CADKernel/Math/Boundary.h
@@ -0,0 +1,335 @@
+// Copyright Epic Games, Inc. All Rights Reserved.
+#pragma once
+
+#include "CADKernel/Core/Types.h"
+#include "CADKernel/Geo/GeoEnum.h"
+#include "CADKernel/Math/MathConst.h"
+#include "CADKernel/Math/Point.h"
+
+namespace CADKernel
+{
+
+	struct CADKERNEL_API FLinearBoundary
+	{
+		/** A default boundary (0., 1.)*/
+		static const FLinearBoundary DefaultBoundary;
+
+		double Min;
+		double Max;
+
+		FLinearBoundary()
+		{
+			Min = 0.;
+			Max = 1.;
+		}
+
+		FLinearBoundary(double UMin, double UMax)
+		{
+			Set(UMin, UMax);
+		}
+
+		friend FArchive& operator<<(FArchive& Ar, FLinearBoundary& Bounds)
+		{
+			Ar << Bounds.Min;
+			Ar << Bounds.Max;
+			return Ar;
+		}
+
+		constexpr double GetMin() const
+		{
+			return Min;
+		}
+
+		constexpr double GetMax() const
+		{
+			return Max;
+		}
+
+		constexpr double GetAt(double Coordinate) const
+		{
+			return Min + (Max - Min) * Coordinate;
+		}
+
+		constexpr double GetMiddle() const
+		{
+			return (Min + Max) * 0.5;
+		}
+
+		double Size() const { return Max - Min; }
+
+		void SetMin(double Coordinates)
+		{
+			Sort(Coordinates, Max, Min, Max);
+		}
+
+		void SetMax(double Coordinates)
+		{
+			Sort(Min, Coordinates, Min, Max);
+		}
+
+		void Set(double InUMin = 0., double InUMax = 1.)
+		{
+			Sort(InUMin, InUMax, Min, Max);
+		}
+
+		bool IsValid() const
+		{
+			return Min <= Max;
+		}
+
+		bool Contains(double Coordinate) const
+		{
+			return RealCompare(Coordinate, Min) >= 0 && RealCompare(Coordinate, Max) <= 0;
+		}
+
+		double Length() const
+		{
+			return GetMax() - GetMin();
+		}
+
+		/**
+		 * Return true if the parametric domain is to small
+		 */
+		bool IsDegenerated() const
+		{
+			double DeltaU = (Max - Min);
+			return (DeltaU < KINDA_SMALL_NUMBER);
+		}
+
+
+		/**
+		 * If a coordinate is outside the bounds, set the coordinate at the closed limit
+		 */
+		void MoveInsideIfNot(double& Coordinate, double Tolerance = SMALL_NUMBER) const
+		{
+			if (Coordinate <= Min)
+			{
+				Coordinate = Min + Tolerance;
+			}
+			else if (Coordinate >= Max)
+			{
+				Coordinate = Max - Tolerance;
+			}
+		}
+
+		/**
+		 * Uses to initiate a boundary computation with ExtendTo 
+		 */
+		void Init()
+		{
+			Min = HUGE_VALUE;
+			Max = -HUGE_VALUE;
+		}
+
+		void ExtendTo(double MinCoordinate, double MaxCoordinate)
+		{
+			Sort(MinCoordinate, MaxCoordinate);
+			Min = FMath::Min(Min, MinCoordinate);
+			Max = FMath::Max(Max, MaxCoordinate);
+		}
+
+		void ExtendTo(FLinearBoundary MaxBound)
+		{
+			Min = FMath::Min(Min, MaxBound.Min);
+			Max = FMath::Max(Max, MaxBound.Max);
+		}
+
+		void ExtendTo(double Coordinate)
+		{
+			if (Coordinate < Min)
+			{
+				Min = Coordinate;
+			}
+
+			if (Coordinate > Max)
+			{
+				Max = Coordinate;
+			}
+		}
+
+		void RestrictTo(FLinearBoundary MaxBound)
+		{
+			if (MaxBound.Min > Min)
+			{
+				Min = MaxBound.Min;
+			}
+			if (MaxBound.Max < Max)
+			{
+				Max = MaxBound.Max;
+			}
+		}
+
+		/**
+		 * If the boundary width is near or equal to zero, it's widened by +/- SMALL_NUMBER
+		 */
+		void WidenIfDegenerated()
+		{
+			if (FMath::IsNearlyEqual(Min, Max))
+			{
+				Min -= SMALL_NUMBER;
+				Max += SMALL_NUMBER;
+			}
+		}
+
+		FLinearBoundary& operator=(const FLinearBoundary& InBounds)
+		{
+			Min = InBounds.Min;
+			Max = InBounds.Max;
+			return *this;
+		}
+
+	};
+
+	struct CADKERNEL_API FSurfacicBoundary
+	{
+		/** A default boundary (0., 1., 0., 1.)*/
+		static const FSurfacicBoundary DefaultBoundary;
+
+		FLinearBoundary UVBoundaries[2];
+
+		FSurfacicBoundary() = default;
+
+		FSurfacicBoundary(double InUMin, double InUMax, double InVMin, double InVMax)
+		{
+			UVBoundaries[EIso::IsoU].Set(InUMin, InUMax);
+			UVBoundaries[EIso::IsoV].Set(InVMin, InVMax);
+		}
+
+		FSurfacicBoundary(const FPoint2D& Point1, const FPoint2D& Point2)
+		{
+			Set(Point1, Point2);
+		}
+		
+		void Set(const FPoint2D& Point1, const FPoint2D& Point2)
+		{
+			UVBoundaries[EIso::IsoU].Set(Point1.U, Point2.U);
+			UVBoundaries[EIso::IsoV].Set(Point1.V, Point2.V);
+		}
+
+		friend FArchive& operator<<(FArchive& Ar, FSurfacicBoundary& Bounds)
+		{
+			Ar << Bounds.UVBoundaries[EIso::IsoU];
+			Ar << Bounds.UVBoundaries[EIso::IsoV];
+			return Ar;
+		}
+
+		void Set(const FLinearBoundary& BoundU, const FLinearBoundary& BoundV)
+		{
+			UVBoundaries[EIso::IsoU] = BoundU;
+			UVBoundaries[EIso::IsoV] = BoundV;
+		}
+
+		void Set(double InUMin, double InUMax, double InVMin, double InVMax)
+		{
+			UVBoundaries[EIso::IsoU].Set(InUMin, InUMax);
+			UVBoundaries[EIso::IsoV].Set(InVMin, InVMax);
+		}
+
+		void Set()
+		{
+			UVBoundaries[EIso::IsoU].Set();
+			UVBoundaries[EIso::IsoV].Set();
+		}
+
+		const FLinearBoundary& Get(EIso Type) const
+		{
+			return UVBoundaries[Type];
+		}
+
+		bool IsValid() const
+		{
+			return UVBoundaries[EIso::IsoU].IsValid() && UVBoundaries[EIso::IsoV].IsValid();
+		}
+
+		/**
+		 * Return true if the parametric domain is to small
+		 */
+		bool IsDegenerated() const
+		{
+			return UVBoundaries[EIso::IsoU].IsDegenerated() || UVBoundaries[EIso::IsoV].IsDegenerated();
+		}
+
+		/**
+		 * Uses to initiate a boundary computation with ExtendTo
+		 */
+		void Init()
+		{
+			UVBoundaries[EIso::IsoU].Init();
+			UVBoundaries[EIso::IsoV].Init();
+		}
+
+		void ExtendTo(const FSurfacicBoundary MaxLimit)
+		{
+			UVBoundaries[EIso::IsoU].ExtendTo(MaxLimit[EIso::IsoU]);
+			UVBoundaries[EIso::IsoV].ExtendTo(MaxLimit[EIso::IsoV]);
+		}
+
+		void ExtendTo(FPoint2D Point)
+		{
+			UVBoundaries[EIso::IsoU].ExtendTo(Point.U);
+			UVBoundaries[EIso::IsoV].ExtendTo(Point.V);
+		}
+
+		void ExtendTo(FPoint Point)
+		{
+			UVBoundaries[EIso::IsoU].ExtendTo(Point.X);
+			UVBoundaries[EIso::IsoV].ExtendTo(Point.Y);
+		}
+
+		void RestrictTo(const FSurfacicBoundary& MaxBound)
+		{
+			UVBoundaries[EIso::IsoU].RestrictTo(MaxBound.UVBoundaries[EIso::IsoU]);
+			UVBoundaries[EIso::IsoV].RestrictTo(MaxBound.UVBoundaries[EIso::IsoV]);
+		}
+
+		/**
+		 * If Along each axis, the bound width is near equal to zero, it's widened by +/- SMALL_NUMBER
+		 */
+		void WidenIfDegenerated()
+		{
+			UVBoundaries[EIso::IsoU].WidenIfDegenerated();
+			UVBoundaries[EIso::IsoV].WidenIfDegenerated();
+		}
+
+		/**
+		 * If a point is outside the bounds, set the coordinate to insert the point inside the bounds
+		 */
+		void MoveInsideIfNot(FPoint& Point, double Tolerance = SMALL_NUMBER) const
+		{
+			UVBoundaries[EIso::IsoU].MoveInsideIfNot(Point.X, Tolerance);
+			UVBoundaries[EIso::IsoV].MoveInsideIfNot(Point.Y, Tolerance);
+		}
+
+		/**
+		 * If a point is outside the bounds, set the coordinate to insert the point inside the bounds
+		 */
+		void MoveInsideIfNot(FPoint2D& Point, double Tolerance = SMALL_NUMBER) const
+		{
+			UVBoundaries[EIso::IsoU].MoveInsideIfNot(Point.U, Tolerance);
+			UVBoundaries[EIso::IsoV].MoveInsideIfNot(Point.V, Tolerance);
+		}
+
+		FSurfacicBoundary& operator=(const FSurfacicBoundary& InBounds)
+		{
+			UVBoundaries[EIso::IsoU] = InBounds.UVBoundaries[EIso::IsoU];
+			UVBoundaries[EIso::IsoV] = InBounds.UVBoundaries[EIso::IsoV];
+			return *this;
+		}
+
+		double Length(const EIso& Iso) const
+		{
+			return UVBoundaries[Iso].Length();
+		}
+
+		constexpr const FLinearBoundary& operator[](const EIso& Iso) const
+		{
+			return UVBoundaries[Iso];
+		}
+
+		constexpr FLinearBoundary& operator[](const EIso& Iso)
+		{
+			return UVBoundaries[Iso];
+		}
+	};
+}
+
diff --git a/Engine/Source/Runtime/Datasmith/CADKernel/Public/CADKernel/Math/MathConst.h b/Engine/Source/Runtime/Datasmith/CADKernel/Public/CADKernel/Math/MathConst.h
new file mode 100644
--- /dev/null
+++ b/Engine/Source/Runtime/Datasmith/CADKernel/Public/CADKernel/Math/MathConst.h
@@ -0,0 +1,57 @@
+// Copyright Epic Games, Inc. All Rights Reserved.
+#pragma once
+
+#include "CADKernel/Core/Types.h"
+
+namespace CADKernel
+{
+	template<typename ValueType> const void Sort(const ValueType& ValueA, const ValueType& ValueB, ValueType& OutMin, ValueType& OutMax)
+	{
+		if (ValueA < ValueB)
+		{
+			OutMin = ValueA;
+			OutMax = ValueB;
+		}
+		else
+		{
+			OutMin = ValueB;
+			OutMax = ValueA;
+		}
+	}
+
+	template<typename ValueType> const void Sort(ValueType& Min, ValueType& Max)
+	{
+		if (Max < Min)
+		{
+			Swap(Min, Max);
+		}
+	}
+
+	inline int32 RealCompare(const double Value1, const double Value2, const double Tolerance = SMALL_NUMBER)
+	{
+		double Difference = Value1 - Value2;
+		if (Difference < -Tolerance) 
+		{
+			return -1;
+		}
+		if (Difference > Tolerance) 
+		{
+			return 1;
+		}
+		return 0;
+	}
+
+	template< class T >
+	static T Cubic(const T A)
+	{
+		return A * A * A;
+	}
+
+	template< class T >
+	uint8 ToUInt8(T Value)
+	{
+		return FMath::Clamp((uint8)(Value / 255.), (uint8) 0, (uint8) 255);
+	}
+
+
+} // namespace CADKernel	
diff --git a/Engine/Source/Runtime/Datasmith/CADKernel/Public/CADKernel/Math/MatrixH.h b/Engine/Source/Runtime/Datasmith/CADKernel/Public/CADKernel/Math/MatrixH.h
new file mode 100644
--- /dev/null
+++ b/Engine/Source/Runtime/Datasmith/CADKernel/Public/CADKernel/Math/MatrixH.h
@@ -0,0 +1,177 @@
+// Copyright Epic Games, Inc. All Rights Reserved.
+#pragma once
+
+#include "CADKernel/Core/Types.h"
+#include "CADKernel/Math/Point.h"
+#include "CADKernel/Math/MathConst.h"
+
+namespace CADKernel
+{
+	/**
+	* Should be unified with the math utilities implemented by the Geometry team.
+	*/
+	class CADKERNEL_API FMatrixH
+	{
+	private:
+		double Matrix[16];
+
+	public:
+
+		static const FMatrixH Identity;
+
+		FMatrixH()
+		{
+			SetIdentity();
+		}
+
+		FMatrixH(const double* const InMatrix16)
+		{
+			for (int32 Index = 0; Index < 16; Index++)
+			{
+				Matrix[Index] = InMatrix16[Index];
+			}
+		}
+
+		FMatrixH(const double InMatrix44[][4])
+		{
+			for (int32 Row = 0; Row < 4; Row++)
+			{
+				for (int32 Column = 0; Column < 4; Column++)
+				{
+					Matrix[4 * Row + Column] = InMatrix44[Row][Column];
+				}
+			}
+		}
+
+		FMatrixH(const FPoint& Origin, const FPoint& Ox, const FPoint& Oy, const FPoint& Oz)
+		{
+			BuildChangeOfCoordinateSystemMatrix(Ox, Oy, Oz, Origin);
+		}
+
+		friend FArchive& operator<<(FArchive& Ar, FMatrixH& Matrix)
+		{
+			Ar.Serialize(Matrix.Matrix, 16 * sizeof(double));
+			return Ar;
+		}
+
+		void SetIdentity()
+		{
+			for (int32 Row = 0; Row < 4; Row++)
+			{
+				for (int32 Column = 0; Column < 4; Column++)
+				{
+					Get(Row, Column) = Row == Column ? 1. : 0.;
+				}
+			}
+		}
+
+		void FromAxisOrigin(const FPoint& Axis, const FPoint& Origin);
+
+		/**
+		 * @param Xaxis the vector X of the old coordinate system in the new coordinate system
+		 * @param Yaxis the vector Y of the old coordinate system in the new coordinate system
+		 * @param Zaxis the vector Z of the old coordinate system in the new coordinate system
+		 * @param Origin Origin of the old coordinate system in the new coordinate system
+		 * @return the transform matrix
+		 */
+		void BuildChangeOfCoordinateSystemMatrix(const FPoint& Xaxis, const FPoint& Yaxis, const FPoint& Zaxis, const FPoint& Origin);
+
+		FPoint Multiply(const FPoint& InPoint) const;
+
+		FPoint MultiplyVector(const FPoint& InVector) const;
+
+		static FMatrixH MakeRotationMatrix(double InAngle, const FPoint InAxe);
+
+		static FMatrixH MakeTranslationMatrix(const FPoint& InPoint);
+
+		static FMatrixH MakeScaleMatrix(double ScaleX, double ScaleY, double ScaleZ);
+
+		/**
+		 * Apply the rotation centered in origin to PointToRotate
+		 */
+		FPoint PointRotation(const FPoint& PointToRotate, const FPoint& Origin) const;
+		FPoint2D PointRotation(const FPoint2D& PointToRotate, const FPoint2D& Origin) const;
+
+		void Inverse();
+
+		FMatrixH GetInverse() const
+		{
+			FMatrixH NewMatrix = *this;
+			NewMatrix.Inverse();
+			return NewMatrix;
+		}
+
+		void Transpose();
+
+		double& Get(int32 Row, int32 Column)
+		{
+			return Matrix[Row * 4 + Column];
+		}
+
+		double Get(int32 Row, int32 Column) const
+		{
+			return Matrix[Row * 4 + Column];
+		}
+
+		double& operator()(int32 Row, int32 Column)
+		{
+			return Matrix[Row * 4 + Column];
+		}
+
+		double operator()(int32 Row, int32 Column) const
+		{
+			return Matrix[Row * 4 + Column];
+		}
+
+		double& operator[](int32 Index)
+		{
+			return Matrix[Index];
+		}
+
+		FMatrixH operator*(const FMatrixH& m) const;
+
+		FPoint operator*(const FPoint& Point) const
+		{
+			return Multiply(Point);
+		}
+
+		FMatrixH operator+(const FMatrixH& m) const;
+
+		void GetMatrixDouble(double* Matrix) const;
+
+		FPoint Column(int32 Index) const;
+
+		FPoint Row(int32 Index) const;
+
+		void CrossProduct(const FPoint& vec);
+
+		void Print(EVerboseLevel level) const;
+
+		bool IsId() const
+		{
+			for (int32 Row = 0; Row < 4; Row++)
+			{
+				for (int32 Column = 0; Column < 4; Column++)
+				{
+					if (Row == Column)
+					{
+						if (!FMath::IsNearlyEqual(Get(Row, Column), 1.)) return false;
+					}
+					else
+					{
+						if (!FMath::IsNearlyZero(Get(Row, Column))) return false;
+					}
+				}
+			}
+			return true;
+		}
+	};
+
+	CADKERNEL_API void InverseMatrixN(double* Matrix, int32 n);
+
+	CADKERNEL_API void MatrixProduct(int32 ARowNum, int32 AColumnNum, int32 ResultRank, const double* InMatrixA, const double* InMatrixB, double* OutMatrix);
+
+	CADKERNEL_API void TransposeMatrix(int32 RowNum, int32 ColumnNum, const double* InMatrix, double* OutMatrix);
+
+} // namespace CADKernel
+
diff --git a/Engine/Source/Runtime/Datasmith/CADKernel/Public/CADKernel/Math/Plane.h b/Engine/Source/Runtime/Datasmith/CADKernel/Public/CADKernel/Math/Plane.h
new file mode 100644
--- /dev/null
+++ b/Engine/Source/Runtime/Datasmith/CADKernel/Public/CADKernel/Math/Plane.h
@@ -0,0 +1,89 @@
+// Copyright Epic Games, Inc. All Rights Reserved.
+#pragma once
+
+#include "CADKernel/Math/Point.h"
+#include "CADKernel/Math/Geometry.h"
+
+namespace CADKernel
+{
+	enum ESide 
+	{
+		OnPlane,
+		AbovePlane,
+		BelowPlane
+	};
+
+	class FPlane
+	{
+	private:
+		FPoint Origin;
+		FPoint Normal;
+
+	public:
+
+		FPlane(const FPoint& Point = FPoint(), const FPoint& Normal = FPoint())
+			: Origin(Point)
+			, Normal(Normal)
+		{}
+
+		FPoint GetPoint(double Lambda = 0.) const
+		{
+			return Origin + Normal * Lambda;
+		}
+
+		void SetOrigin(const FPoint& Point)
+		{
+			Origin = Point;
+		}
+
+		const FPoint& GetOrigin() const 
+		{
+			return Origin;
+		}
+
+		const FPoint& GetNormal() const
+		{
+			return Normal;
+		}
+
+		void SetNormal(const FPoint& InNormal)
+		{
+			Normal = InNormal;
+		}
+
+		FPoint PointProjection(const FPoint& Point) const
+		{
+			double Distance;
+			return ProjectPointOnPlane(Point, Origin, Normal, Distance);
+		}
+
+		ESide GetSide(const FPoint& Point) const
+		{
+			double DotProduct = (Point - Origin) * Normal;
+			if (DotProduct  < 0)
+			{
+				return ESide::BelowPlane;
+			}
+
+			if (DotProduct  > 0)
+			{
+				return ESide::AbovePlane;
+			}
+
+			return ESide::OnPlane;
+		}
+
+		FPlane& TranslateAlongNormal(double Step)
+		{
+			Origin += Normal * Step;
+			return *this;
+		}
+
+		FPlane TranslateAlongNormal(double Step) const 
+		{
+			return FPlane(Origin + Normal * Step, Normal);
+		}
+	};
+
+}
+
