CADKernel Source Code v0.01

diff --git a/Engine/Source/Runtime/Datasmith/CADKernel/Public/CADKernel/Math/Point.h b/Engine/Source/Runtime/Datasmith/CADKernel/Public/CADKernel/Math/Point.h
new file mode 100644
--- /dev/null
+++ b/Engine/Source/Runtime/Datasmith/CADKernel/Public/CADKernel/Math/Point.h
@@ -0,0 +1,618 @@
+// Copyright Epic Games, Inc. All Rights Reserved.
+#pragma once
+
+#include "CADKernel/Core/Types.h"
+#include "CADKernel/Math/MathConst.h"
+
+// this will eventually be reconciled with what the Geometry team is using.
+namespace CADKernel
+{
+	class FPoint2D;
+
+	class CADKERNEL_API FPoint
+	{
+	public:
+		double X;
+		double Y;
+		double Z;
+
+	public:
+
+		/** A zero point (0,0,0)*/
+		static const FPoint ZeroPoint;
+
+		explicit FPoint(double InCoordX = 0., double InCoordY = 0., double InCoordZ = 0.)
+		{
+			X = InCoordX; Y = InCoordY; Z = InCoordZ;
+		}
+
+		FPoint(const FPoint2D& InPoint2D);
+
+		FPoint(const double* const InPoint)
+			: X(InPoint[0])
+			, Y(InPoint[1])
+			, Z(InPoint[2])
+		{
+		}
+
+		FPoint(const FPoint& Point)
+			: X(Point.X)
+			, Y(Point.Y)
+			, Z(Point.Z)
+		{
+		}
+
+		constexpr double& operator[](int32 Index)
+		{
+			return *(&X + Index);
+		}
+
+		constexpr double operator[](int32 Index) const
+		{
+			return *(&X + Index);
+		}
+
+		bool operator==(const FPoint& Point) const
+		{
+			return FMath::IsNearlyEqual(X, Point.X) && FMath::IsNearlyEqual(Y, Point.Y) && FMath::IsNearlyEqual(Z, Point.Z);
+		}
+
+		bool operator!=(const FPoint& Point) const
+		{
+			return !(*this == Point);
+		};
+
+		FPoint operator+(const FPoint& Point) const
+		{
+			return FPoint(X + Point.X, Y + Point.Y, Z + Point.Z);
+		}
+
+		FPoint operator-(const FPoint& Point) const
+		{
+			return FPoint(X - Point.X, Y - Point.Y, Z - Point.Z);
+		}
+
+		FPoint operator^(const FPoint& Point) const
+		{
+			return FPoint(Y * Point.Z - Z * Point.Y, Z * Point.X - X * Point.Z, X * Point.Y - Y * Point.X);
+		}
+
+		double operator*(const FPoint& Point) const
+		{
+			return X * Point.X + Y * Point.Y + Z * Point.Z;
+		}
+
+		FPoint operator*(double Scale) const
+		{
+			return FPoint(X * Scale, Y * Scale, Z * Scale);
+		}
+
+		FPoint operator/(double InvScale) const
+		{
+			return FPoint(X / InvScale, Y / InvScale, Z / InvScale);
+		}
+
+		FPoint operator-() const
+		{
+			return FPoint(-X, -Y, -Z);
+		}
+
+		FPoint& operator=(const FPoint& Point)
+		{
+			X = Point.X;
+			Y = Point.Y;
+			Z = Point.Z;
+			return *this;
+		}
+
+		FPoint& operator+=(const FPoint& Point)
+		{
+			X += Point.X;
+			Y += Point.Y;
+			Z += Point.Z;
+			return *this;
+		}
+
+		FPoint& operator-=(const FPoint& Point)
+		{
+			X -= Point.X;
+			Y -= Point.Y;
+			Z -= Point.Z;
+			return *this;
+		}
+
+		FPoint& operator*=(double Scale)
+		{
+			X *= Scale;
+			Y *= Scale;
+			Z *= Scale;
+			return *this;
+		}
+
+		FPoint& operator/=(double InvScale)
+		{
+			X /= InvScale;
+			Y /= InvScale;
+			Z /= InvScale;
+			return *this;
+		}
+
+		bool operator<(const FPoint& Other) const
+		{
+			int32 Compare = CADKernel::RealCompare(X, Other.X);
+			if (Compare < 0)
+			{
+				return true;
+			}
+			if (Compare > 0)
+			{
+				return false;
+			}
+
+			Compare = CADKernel::RealCompare(Y, Other.Y);
+			if (Compare < 0)
+			{
+				return true;
+			}
+			if (Compare > 0)
+			{
+				return false;
+			}
+
+			Compare = CADKernel::RealCompare(Z, Other.Z);
+			return (Compare < 0);
+		}
+
+		void SetMin(const FPoint& Point)
+		{
+			X = FMath::Min(X, Point.X);
+			Y = FMath::Min(Y, Point.Y);
+			Z = FMath::Min(Z, Point.Z);
+		}
+
+		void SetMax(const FPoint& Point)
+		{
+			X = FMath::Max(X, Point.X);
+			Y = FMath::Max(Y, Point.Y);
+			Z = FMath::Max(Z, Point.Z);
+		}
+
+		void Set(double InCoordX = 0., double InCoordY = 0., double InCoordZ = 0.)
+		{
+			X = InCoordX;
+			Y = InCoordY;
+			Z = InCoordZ;
+		}
+
+		void Set(double* InCoordinates)
+		{
+			X = InCoordinates[0];
+			Y = InCoordinates[1];
+			Z = InCoordinates[2];
+		}
+
+		double Length() const
+		{
+			return sqrt((X * X) + (Y * Y) + (Z * Z));
+		}
+
+		double SquareLength() const
+		{
+			return (X * X) + (Y * Y) + (Z * Z);
+		}
+
+		FPoint& Normalize()
+		{
+			double Norm = Length();
+			if (FMath::IsNearlyZero(Norm))
+			{
+				X = Y = Z = 0.0;
+			}
+			else
+			{
+				X /= Norm;
+				Y /= Norm;
+				Z /= Norm;
+			}
+			return *this;
+		}
+
+		FPoint Middle(const FPoint& Point) const
+		{
+			return FPoint(0.5 * (X + Point.X), 0.5 * (Y + Point.Y), 0.5 * (Z + Point.Z));
+		}
+
+		double Distance(const FPoint& Point) const
+		{
+			return sqrt(FMath::Square(Point.X - X) + FMath::Square(Point.Y - Y) + FMath::Square(Point.Z - Z));
+		}
+
+		double SquareDistance(const FPoint& Point) const
+		{
+			return FMath::Square(Point.X - X) + FMath::Square(Point.Y - Y) + FMath::Square(Point.Z - Z);
+		}
+
+		double ComputeCosinus(const FPoint& Vector2) const;
+
+		double ComputeSinus(const FPoint& Vector2) const;
+
+		double ComputeAngle(const FPoint& Vector2) const;
+
+		double SignedAngle(const FPoint& Vector2, const FPoint& Normal) const;
+
+		static double MixedTripleProduct(const FPoint& VectorA, const FPoint& VectorB, const FPoint& VectorC)
+		{
+			return VectorA * (VectorB ^ VectorC);
+		}
+
+		friend FArchive& operator<<(FArchive& Ar, FPoint& Point)
+		{
+			Ar << Point.X;
+			Ar << Point.Y;
+			Ar << Point.Z;
+			return Ar;
+		}
+
+		static int32 Dimension() { return 3; }
+	};
+
+	class CADKERNEL_API FPointH
+	{
+	public:
+		double X;
+		double Y;
+		double Z;
+		double W;
+
+	public:
+		/** A zero point (0,0,0)*/
+		static const FPointH ZeroPoint;
+
+		FPointH()
+			: X(0.)
+			, Y(0.)
+			, Z(0.)
+			, W(1.)
+		{
+		}
+
+		FPointH(const FPoint& Point, double Weight)
+			: X(Point.X* Weight)
+			, Y(Point.Y* Weight)
+			, Z(Point.Z* Weight)
+			, W(Weight)
+		{
+		}
+
+		FPointH(const double InX, const double InY, const double InZ, const double InW)
+			: X(InX)
+			, Y(InY)
+			, Z(InZ)
+			, W(InW)
+		{
+		}
+
+		constexpr double& operator[](int32 Index)
+		{
+			return *(&X + Index);
+		}
+
+		constexpr double operator[](int32 Index) const
+		{
+			return *(&X + Index);
+		}
+
+		FPointH operator+(const FPointH& Point) const
+		{
+			return FPointH(X + Point.X, Y + Point.Y, Z + Point.Z, W + Point.W);
+		}
+
+		FPointH operator-(const FPointH& Point) const
+		{
+			return FPointH(X - Point.X, Y - Point.Y, Z - Point.Z, W - Point.W);
+		}
+
+		FPointH operator*(double Factor) const
+		{
+			return FPointH(X * Factor, Y * Factor, Z * Factor, W * Factor);
+		}
+
+		operator FPoint() const
+		{
+			return FPoint(X / W, Y / W, Z / W);
+		}
+
+		friend FArchive& operator<<(FArchive& Ar, FPointH& Point)
+		{
+			Ar << Point.X;
+			Ar << Point.Y;
+			Ar << Point.Z;
+			Ar << Point.W;
+			return Ar;
+		}
+	};
+
+	class CADKERNEL_API FPoint2D
+	{
+	public:
+		double U;
+		double V;
+
+	public:
+
+		/** A zero point (0,0,0)*/
+		static const FPoint2D ZeroPoint;
+
+		FPoint2D()
+		{
+			U = V = 0.0;
+		}
+
+		FPoint2D(double NewU, double NewV)
+			: U(NewU)
+			, V(NewV)
+		{
+		}
+
+		FPoint2D(const FPoint& Point)
+			: U(Point.X)
+			, V(Point.Y)
+		{
+		}
+
+		FPoint2D(const FPoint2D& Point)
+			: U(Point.U)
+			, V(Point.V)
+		{
+		}
+
+		FPoint2D& operator=(const FPoint2D& Point)
+		{
+			U = Point.U;
+			V = Point.V;
+			return *this;
+		}
+
+		FPoint2D Rotate(double Theta)
+		{
+			return FPoint2D(U * cos(Theta) - V * sin(Theta), U * sin(Theta) + V * cos(Theta));
+		}
+
+		double Distance(const FPoint2D& Point) const
+		{
+			return sqrt(FMath::Square(Point.U - U) + FMath::Square(Point.V - V));
+		}
+
+		double SquareDistance(const FPoint2D& Point) const
+		{
+			return FMath::Square(Point.U - U) + FMath::Square(Point.V - V);
+		}
+
+		double Length() const
+		{
+			return sqrt(FMath::Square(U) + FMath::Square(V));
+		}
+
+		double SquareLength() const
+		{
+			return FMath::Square(U) + FMath::Square(V);
+		}
+
+		FPoint2D& Normalize()
+		{
+			double Norm = Length();
+			if (FMath::IsNearlyZero(Norm))
+			{
+				U = V = 0.0;
+			}
+			else
+			{
+				U /= Norm;
+				V /= Norm;
+			}
+			return *this;
+		}
+
+		constexpr double& operator[](int32 Index)
+		{
+			return *(&U + Index);
+		}
+
+		constexpr double operator[](int32 Index) const
+		{
+			return *(&U + Index);
+		}
+
+		bool operator==(const FPoint2D& Point) const
+		{
+			return FMath::IsNearlyEqual(U, Point.U) && FMath::IsNearlyEqual(V, Point.V);
+		}
+
+		bool operator<(const FPoint2D& other) const
+		{
+			int32 cmp = RealCompare(U, other.U);
+			if (cmp < 0)
+			{
+				return true;
+			}
+			if (cmp > 0)
+			{
+				return false;
+			}
+
+			cmp = RealCompare(V, other.V);
+			return (cmp < 0);
+		}
+
+		void Set(double NewU, double NewV = 0.)
+		{
+			U = NewU;
+			V = NewV;
+		}
+
+		void Set(double* NewCoordinate)
+		{
+			U = NewCoordinate[0];
+			V = NewCoordinate[1];
+		}
+
+		FPoint2D operator+(const FPoint2D& Point) const
+		{
+			return FPoint2D(U + Point.U, V + Point.V);
+		}
+
+		FPoint2D operator-(const FPoint2D& Point) const
+		{
+			return FPoint2D(U - Point.U, V - Point.V);
+		}
+
+		double operator*(const FPoint2D& Point) const
+		{
+			return U * Point.U + V * Point.V;
+		}
+
+		double operator^(const FPoint2D& Point) const
+		{
+			return (U * Point.V - V * Point.U);
+		}
+
+		FPoint2D operator*(double Scale) const
+		{
+			return FPoint2D(U * Scale, V * Scale);
+		}
+
+		FPoint2D operator/(double InvScale) const
+		{
+			return FPoint2D(U / InvScale, V / InvScale);
+		}
+
+		FPoint2D operator-() const
+		{
+			return FPoint2D(-U, -V);
+		}
+
+		FPoint2D& operator+=(const FPoint2D& Point)
+		{
+			U += Point.U;
+			V += Point.V;
+			return *this;
+		}
+
+		FPoint2D& operator-=(const FPoint2D& Point)
+		{
+			U -= Point.U;
+			V -= Point.V;
+			return *this;
+		}
+
+		FPoint2D& operator*=(double Scale)
+		{
+			U *= Scale;
+			V *= Scale;
+			return *this;
+		}
+
+		FPoint2D& operator/=(double InvScale)
+		{
+			U /= InvScale;
+			V /= InvScale;
+			return *this;
+		}
+
+		friend FArchive& operator<<(FArchive& Ar, FPoint2D& Point)
+		{
+			Ar << Point.U;
+			Ar << Point.V;
+			return Ar;
+		}
+
+		FPoint2D GetPerpendicularVector()
+		{
+			return FPoint2D(-V, U);
+		}
+
+		static constexpr const int32 Dimension() { return 2; }
+	};
+
+	inline FPoint::FPoint(const FPoint2D& Point)
+	{
+		X = Point.U; Y = Point.V; Z = 0.0;
+	}
+
+	typedef FPoint2D FSurfacicTolerance;
+
+	class CADKERNEL_API FFPoint
+	{
+	public:
+		float X;
+		float Y;
+		float Z;
+
+		/** A zero point (0,0,0)*/
+		static const FFPoint ZeroPoint;
+
+		FFPoint(float InX = 0., float InY = 0., float InZ = 0.)
+		{
+			X = InX;
+			Y = InY;
+			Z = InZ;
+		}
+
+		FFPoint(const FPoint& Point)
+		{
+			X = (float)Point.X;
+			Y = (float)Point.Y;
+			Z = (float)Point.Z;
+		}
+
+		constexpr float& operator[](int32 Index)
+		{
+			return *(&X + Index);
+		}
+
+		constexpr float operator[](int32 Index) const
+		{
+			return *(&X + Index);
+		}
+
+		FFPoint& operator=(const FPoint& Point)
+		{
+			X = (float)Point.X; Y = (float)Point.Y; Z = (float)Point.Z; return *this;
+		}
+
+		friend FArchive& operator<<(FArchive& Ar, FFPoint& Point)
+		{
+			Ar << Point.X;
+			Ar << Point.Y;
+			Ar << Point.Z;
+			return Ar;
+		}
+
+		static constexpr const int32 Dimension() { return 3; }
+	};
+
+	inline FPoint operator*(double Scale, const FPoint& Point)
+	{
+		return FPoint(Point.X * Scale, Point.Y * Scale, Point.Z * Scale);
+	}
+
+	inline FPoint2D operator*(double s, const FPoint2D& Point)
+	{
+		return FPoint2D(Point.U * s, Point.V * s);
+	}
+}
+
+inline uint32 GetTypeHash(const CADKernel::FPoint& Point)
+{
+	return HashCombine(GetTypeHash(Point.X), HashCombine(GetTypeHash(Point.Y), GetTypeHash(Point.Z)));
+}
+
+inline uint32 GetTypeHash(const CADKernel::FPoint2D& Point)
+{
+	return HashCombine(GetTypeHash(Point.U), GetTypeHash(Point.V));
+}
+
+inline uint32 GetTypeHash(const CADKernel::FPointH& Point)
+{
+	return HashCombine(GetTypeHash(Point.X), HashCombine(GetTypeHash(Point.Y), HashCombine(GetTypeHash(Point.Z), GetTypeHash(Point.W))));
+}
\ No newline at end of file
diff --git a/Engine/Source/Runtime/Datasmith/CADKernel/Public/CADKernel/Mesh/Criteria/Criterion.h b/Engine/Source/Runtime/Datasmith/CADKernel/Public/CADKernel/Mesh/Criteria/Criterion.h
new file mode 100644
--- /dev/null
+++ b/Engine/Source/Runtime/Datasmith/CADKernel/Public/CADKernel/Mesh/Criteria/Criterion.h
@@ -0,0 +1,146 @@
+// Copyright Epic Games, Inc. All Rights Reserved.
+#pragma once
+
+#include "CADKernel/Core/Types.h"
+#include "CADKernel/Core/Entity.h"
+#include "CADKernel/Geo/GeoPoint.h"
+#include "CADKernel/Math/Point.h"
+#include "CADKernel/Mesh/Criteria/CriterionType.h"
+
+
+namespace CADKernel
+{
+	extern const TCHAR* CriterionTypeNames[];
+
+	// Defined for Python purpose
+	extern const char* CriterionTypeConstNames[];
+	extern const char* CriterionTypeConstDescHelp[];
+
+	class FTopologicalEdge;
+	struct FCurvePoint;
+	struct FIsoCurvature;
+
+	class CADKERNEL_API FCriterion : public FEntity
+	{
+		friend class FEntity;
+
+	protected:
+		ECriterion CriterionType;
+		
+		FCriterion(ECriterion InCriterionType)
+			: FEntity()
+			, CriterionType(InCriterionType)
+		{
+		}
+
+	public:
+
+		virtual void Serialize(FCADKernelArchive& Ar) override
+		{
+			// Criterion's type is serialize because it is used to instantiate the correct entity on deserialization (@see Deserialize(FArchive& Archive)) 
+			if (Ar.IsSaving())
+			{
+				Ar << CriterionType;
+			}
+			FEntity::Serialize(Ar);
+		}
+
+		/**
+		 * Specific method for curve family to instantiate the correct derived class of FCurve
+		 */
+		static TSharedPtr<FCriterion> Deserialize(FCADKernelArchive& Archive);
+
+		static TSharedPtr<FCriterion> CreateCriterion(ECriterion type, double value = 0.);
+
+#ifdef CADKERNEL_DEV
+		virtual FInfoEntity& GetInfo(FInfoEntity&) const override;
+#endif
+
+		virtual EEntity GetEntityType() const override
+		{
+			return EEntity::Criterion;
+		}
+
+		/**
+		 * Sag & Angle criterion.pdf
+		 * https://docs.google.com/presentation/d/1bUnrRFWCW3sDn9ngb9ftfQS-2JxNJaUZlh783hZMMEw/edit?usp=sharing
+		 */
+		virtual void ApplyOnEdgeParameters(TSharedRef<FTopologicalEdge> Edge, const TArray<double>& Coordinates, const TArray<FCurvePoint>& Points) const;
+
+		const ECriterion GetCriterionType() const
+		{
+			return CriterionType;
+		}
+
+		FString GetCriterionName()
+		{
+			return CriterionTypeNames[(int32) CriterionType];
+		}
+
+		virtual double Value() const = 0;
+
+		static double EvaluateSag(const FPoint& PointPoint, const FPoint& PointNext, const FPoint& PointMiddle, double& Length)
+		{
+			FPoint ChordVec = PointNext - PointPoint;
+			FPoint MiddleVec = PointMiddle - PointPoint;
+
+			double Sag = 0.0;
+			double NormSqrVec = ChordVec * ChordVec;
+			Length = sqrt(NormSqrVec);
+
+			if (NormSqrVec > SMALL_NUMBER)
+			{
+				FPoint VecSag = ChordVec ^ MiddleVec;
+				double NormSqrSag = VecSag * VecSag;
+				Sag = NormSqrSag / NormSqrVec;
+
+				if (Sag <SMALL_NUMBER_SQUARE)
+				{
+					Sag = 0.;
+				}
+				else
+				{
+					Sag = sqrt(Sag);
+				}
+			}
+			return Sag;
+		}
+
+		virtual bool IsAppliedBetweenBreaks() const
+		{
+			return false;
+		}
+
+		static FString GetCriterionName(ECriterion CriterionType)
+		{
+			return CriterionTypeNames[(int32) CriterionType];
+		}
+
+
+		static double DefaultValue(ECriterion type);
+
+		virtual void UpdateDelta(double InDeltaU, double InUSag, double InDiagonalSag, double InVSag, double ChordLength, double DiagonalLength, double& OutSagDeltaUMax, double& OutSagDeltaUMin, FIsoCurvature& SurfaceCurvature) const
+		{
+			// the component according to U of the sag along the diagonal = (diagonal sag - V Sag) * U Length / diag length
+			double DiagonalSagU = FMath::Abs(InDiagonalSag - InVSag) * ChordLength / DiagonalLength;
+			InUSag = FMath::Max(InUSag, DiagonalSagU);
+
+			if (InUSag > SMALL_NUMBER)
+			{
+				double DeltaUMax = ComputeDeltaU(ChordLength, InDeltaU, InUSag);
+				OutSagDeltaUMax = (OutSagDeltaUMax <DeltaUMax) ? OutSagDeltaUMax : DeltaUMax;
+			}
+		}
+
+	protected:
+
+		virtual double ComputeDeltaU(double ChordLength, double DeltaU, double Sag) const
+		{
+			ensureCADKernel(false);
+			return 0.0;
+		};
+
+	};
+
+} // namespace CADKernel
+
diff --git a/Engine/Source/Runtime/Datasmith/CADKernel/Public/CADKernel/Mesh/Criteria/CriterionType.h b/Engine/Source/Runtime/Datasmith/CADKernel/Public/CADKernel/Mesh/Criteria/CriterionType.h
new file mode 100644
--- /dev/null
+++ b/Engine/Source/Runtime/Datasmith/CADKernel/Public/CADKernel/Mesh/Criteria/CriterionType.h
@@ -0,0 +1,18 @@
+// Copyright Epic Games, Inc. All Rights Reserved.
+#pragma once
+
+#include "CADKernel/Core/Types.h"
+
+namespace CADKernel
+{
+	enum class ECriterion : uint8
+	{
+		Size = 0,
+		MaxSize = 1,
+		MinSize = 2,
+		Angle = 3,
+		Sag = 4,
+		CADCurvature = 5,
+		None = 255
+	} ;
+} // namespace CADKernel
diff --git a/Engine/Source/Runtime/Datasmith/CADKernel/Public/CADKernel/Mesh/Meshers/ParametricMesher.h b/Engine/Source/Runtime/Datasmith/CADKernel/Public/CADKernel/Mesh/Meshers/ParametricMesher.h
new file mode 100644
--- /dev/null
+++ b/Engine/Source/Runtime/Datasmith/CADKernel/Public/CADKernel/Mesh/Meshers/ParametricMesher.h
@@ -0,0 +1,194 @@
+// Copyright Epic Games, Inc. All Rights Reserved.
+#pragma once
+
+#include "CADKernel/Core/Chrono.h"
+#include "CADKernel/Core/Types.h"
+#include "CADKernel/Topo/TopologicalEdge.h"
+#include "CADKernel/Topo/Shell.h"
+#include "CADKernel/UI/Progress.h"
+
+//#define DEBUG_IntersectEdgeIsos
+
+namespace CADKernel
+{
+	class FCriterion;
+	class FGrid;
+	class FThinZoneSide;
+	class FTopologicalEntity;
+	class FTopologicalLoop;
+
+	struct FCostToFace
+	{
+		double Cost;
+		TSharedRef<FTopologicalFace> Face;
+
+		FCostToFace(double NewCost, TSharedRef<FTopologicalFace> NewFace)
+			: Cost(NewCost)
+			, Face(NewFace)
+		{
+		}
+
+	};
+
+	class CADKERNEL_API FMesherParameters : public FParameters
+	{
+	public:
+		FParameter InconsistencyAngle;
+
+		FMesherParameters();
+
+		void SetInconsistencyAngle(const double Value)
+		{
+			InconsistencyAngle = Value;
+		}
+	};
+
+	struct CADKERNEL_API FMesherChronos
+	{
+		FDuration GlobalDuration;
+		FDuration ApplyCriteriaDuration;
+		FDuration IsolateQuadPatchDuration;
+		FDuration GlobalMeshDuration;
+		FDuration GlobalPointCloudDuration;
+		FDuration GlobalGeneratePointCloudDuration;
+		FDuration GlobalTriangulateDuration;
+		FDuration GlobalDelaunayDuration;
+		FDuration GlobalMeshAndGetLoopNodes;
+		FDuration GlobalMeshEdges;
+		FDuration GlobalThinZones;
+
+		FDuration GlobalFindThinZones;
+		//FDuration GlobalScaleGrid;
+		FDuration GlobalMeshThinZones;
+
+		FMesherChronos()
+			: GlobalDuration(FChrono::Init())
+			, ApplyCriteriaDuration(FChrono::Init())
+			, IsolateQuadPatchDuration(FChrono::Init())
+			, GlobalMeshDuration(FChrono::Init())
+			, GlobalPointCloudDuration(FChrono::Init())
+			, GlobalGeneratePointCloudDuration(FChrono::Init())
+			, GlobalTriangulateDuration(FChrono::Init())
+			, GlobalDelaunayDuration(FChrono::Init())
+			, GlobalMeshAndGetLoopNodes(FChrono::Init())
+			, GlobalMeshEdges(FChrono::Init())
+			, GlobalThinZones(FChrono::Init())
+			, GlobalFindThinZones(FChrono::Init())
+			//, GlobalScaleGrid(FChrono::Init())
+			, GlobalMeshThinZones(FChrono::Init())
+		{}
+
+		void PrintTimeElapse() const
+		{
+			FMessage::Printf(Log, TEXT("\n\n\n"));
+			FChrono::PrintClockElapse(Log, TEXT(""), TEXT("Total"), GlobalDuration);
+			FChrono::PrintClockElapse(Log, TEXT("  |  "), TEXT("Apply Criteria"), ApplyCriteriaDuration);
+			FChrono::PrintClockElapse(Log, TEXT("  |  "), TEXT("Find Quad Surfaces"), IsolateQuadPatchDuration);
+			FChrono::PrintClockElapse(Log, TEXT("  |  "), TEXT("Mesh Time"), GlobalMeshDuration);
+			FChrono::PrintClockElapse(Log, TEXT("  |   |  "), TEXT("GeneratePoint Cloud "), GlobalGeneratePointCloudDuration);
+			FChrono::PrintClockElapse(Log, TEXT("  |   |  |  "), TEXT("Point Cloud "), GlobalPointCloudDuration);
+			FChrono::PrintClockElapse(Log, TEXT("  |   |  "), TEXT("ThinZones "), GlobalThinZones);
+			//FChrono::PrintClockElapse(Log, TEXT("  |   |  |  "), TEXT("GlobalScaleGrid"), GlobalScaleGrid);
+			FChrono::PrintClockElapse(Log, TEXT("  |   |  "), TEXT("Mesh ThinZones "), GlobalMeshThinZones);
+			FChrono::PrintClockElapse(Log, TEXT("  |   |  "), TEXT("MeshEdges "), GlobalMeshEdges);
+			FChrono::PrintClockElapse(Log, TEXT("  |   |  "), TEXT("TriangulateDuration "), GlobalTriangulateDuration);
+			FChrono::PrintClockElapse(Log, TEXT("  |   |   |  "), TEXT("Delaunay Duration "), GlobalDelaunayDuration);
+		}
+	};
+
+	class CADKERNEL_API FParametricMesher
+	{
+	protected:
+
+		/**
+		 * Limit of flatness of quad face
+		 */
+		const double MinCurvature = 0.001;
+
+		TSharedRef<FModelMesh> MeshModel;
+		TSharedRef<FMesherParameters> Parameters;
+
+		TArray<TSharedPtr<FTopologicalFace>> Faces;
+		TArray<TSharedPtr<FTopologicalEdge>> Edges;
+		TArray<TSharedPtr<FTopologicalVertex>> Vertices;
+
+		FMesherChronos Chronos;
+
+	public:
+
+		FParametricMesher(TSharedRef<FModelMesh> MeshModel);
+		FParametricMesher(TSharedRef<FModelMesh> MeshModel, TArray<TSharedPtr<FEntity>>& InEntities);
+		FParametricMesher(TSharedRef<FModelMesh> MeshModel, TSharedRef<FTopologicalEntity>& InEntity);
+
+		const TSharedRef<FModelMesh>& GetMeshModel() const
+		{
+			return MeshModel;
+		}
+
+		TSharedRef<FModelMesh>& GetMeshModel()
+		{
+			return MeshModel;
+		}
+
+		void InitParameters(const FString& ParametersString);
+
+		const TSharedRef<FMesherParameters>& GetParameters() const
+		{
+			return Parameters;
+		}
+
+		void MeshEntities(TArray<TSharedPtr<FEntity>>& InEntities);
+
+		template<typename EntityType>
+		void MeshEntity(TSharedRef<EntityType>& Entity)
+		{
+			MeshEntity((TSharedRef<FTopologicalEntity>&) Entity);
+		}
+
+		void MeshEntity(TSharedRef<FModel>& InModel)
+		{
+			MeshEntity((TSharedRef<FTopologicalEntity>&) InModel);
+		}
+
+		void Mesh(TSharedRef<FTopologicalFace> Face);
+		void Mesh(TSharedRef<FTopologicalEdge> InEdge, TSharedRef<FTopologicalFace> CarrierFace);
+		void Mesh(TSharedRef<FTopologicalVertex> Vertex);
+
+		void MeshFaceLoops(FGrid& Grid);
+
+		void MeshThinZoneEdges(FGrid&);
+		void MeshThinZoneSide(const FThinZoneSide& Side);
+		void GetThinZoneBoundary(const FThinZoneSide& Side);
+
+		void GenerateCloud(FGrid& Grid);
+
+	protected:
+
+		void MeshEntities();
+
+		void MeshEntity(TSharedRef<FTopologicalEntity>& InEntity)
+		{
+			TArray<TSharedPtr<FEntity>> Entities;
+			Entities.Add(InEntity);
+			MeshEntities(Entities);
+		}
+
+		void IsolateQuadFace(TArray<FCostToFace>& QuadSurfaces, TArray<TSharedPtr<FTopologicalFace>>& OtherSurfaces) const;
+
+		void LinkQuadSurfaceForMesh(TArray<FCostToFace>& QuadTrimmedSurfaceSet, TArray<TArray<TSharedPtr<FTopologicalFace>>>& OutStrips);
+		void MeshSurfaceByFront(TArray<FCostToFace>& QuadTrimmedSurfaceSet);
+
+		void ApplyEdgeCriteria(TSharedRef<FTopologicalEdge> Edge);
+		void ApplySurfaceCriteria(TSharedRef<FTopologicalFace> Surface);
+
+		void GenerateEdgeElements(TSharedRef<FTopologicalEdge> Edge);
+
+		void IntersectEdgeIsos(TSharedPtr<FTopologicalEdge> Edge, const TArray<FPoint2D>& EdgeCrossingPoints2D, const double ToleranceIso, const TArray<double>& IsoCoordinates, EIso Iso, TArray<double>& OutIntersections);
+
+#ifdef DEBUG_IntersectEdgeIsos
+		void DebugIntersectEdgeIsos(TSharedPtr<FTopologicalEdge> Edge, TSharedPtr<FTopologicalFace> Surface, const TArray<FPoint>& EdgeCrossingPoints2D, const double ToleranceIso, const TArray<double>& IsoCoordinates, EIso TypeIso);
+#endif
+	};
+
+} // namespace CADKernel
+
diff --git a/Engine/Source/Runtime/Datasmith/CADKernel/Public/CADKernel/Mesh/Structure/Mesh.h b/Engine/Source/Runtime/Datasmith/CADKernel/Public/CADKernel/Mesh/Structure/Mesh.h
new file mode 100644
--- /dev/null
+++ b/Engine/Source/Runtime/Datasmith/CADKernel/Public/CADKernel/Mesh/Structure/Mesh.h
@@ -0,0 +1,89 @@
+// Copyright Epic Games, Inc. All Rights Reserved.
+#pragma once
+
+#include "CADKernel/Core/EntityGeom.h"
+#include "CADKernel/Core/Types.h"
+
+namespace CADKernel
+{
+	class FNode;
+	class FModelMesh;
+	class FTopologicalEntity;
+
+	class CADKERNEL_API FMesh : public FEntityGeom
+	{
+	protected:
+		TWeakPtr<FModelMesh> ModelMesh;
+		TWeakPtr<FTopologicalEntity> TopologicalEntity;
+
+		int32 StartNodeId;
+		int32 LastNodeIndex;
+
+		TArray<FPoint> NodeCoordinates;
+		int32 MeshModelIndex;
+
+	public:
+
+		FMesh(TSharedRef<FModelMesh> InMeshModel, TSharedRef<FTopologicalEntity> InTopologicalEntity)
+			: FEntityGeom()
+			, ModelMesh(InMeshModel)
+			, TopologicalEntity(InTopologicalEntity)
+		{
+		}
+
+		virtual ~FMesh() = default;
+
+#ifdef CADKERNEL_DEV
+		virtual FInfoEntity& GetInfo(FInfoEntity&) const override;
+#endif
+
+		virtual EEntity GetEntityType() const override
+		{
+			return EEntity::Mesh;
+		}
+
+		TArray<FPoint>& GetNodeCoordinates()
+		{
+			return NodeCoordinates;
+		}
+
+		const TArray<FPoint>& GetNodeCoordinates() const
+		{
+			return NodeCoordinates;
+		}
+
+		int32 RegisterCoordinates();
+
+		const int32 GetStartVertexId() const
+		{
+			return StartNodeId;
+		}
+
+		const int32 GetLastVertexIndex() const
+		{
+			return LastNodeIndex;
+		}
+
+		const int32 GetIndexInMeshModel() const
+		{
+			return MeshModelIndex;
+		}
+
+		TSharedRef<FModelMesh> GetMeshModel()
+		{
+			ensureCADKernel(ModelMesh.Pin().IsValid());
+			return ModelMesh.Pin().ToSharedRef();
+		}
+
+		const TSharedRef<FModelMesh> GetMeshModel() const
+		{
+			return ModelMesh.Pin().ToSharedRef();
+		}
+
+		const TSharedRef<FTopologicalEntity> GetGeometricEntity() const
+		{
+			return TopologicalEntity.Pin().ToSharedRef();
+		}
+	};
+}
+
diff --git a/Engine/Source/Runtime/Datasmith/CADKernel/Public/CADKernel/Mesh/Structure/ModelMesh.h b/Engine/Source/Runtime/Datasmith/CADKernel/Public/CADKernel/Mesh/Structure/ModelMesh.h
new file mode 100644
--- /dev/null
+++ b/Engine/Source/Runtime/Datasmith/CADKernel/Public/CADKernel/Mesh/Structure/ModelMesh.h
@@ -0,0 +1,147 @@
+// Copyright Epic Games, Inc. All Rights Reserved.
+#pragma once
+
+#include "CADKernel/Core/EntityGeom.h"
+#include "CADKernel/Core/Session.h"
+#include "CADKernel/Core/Types.h"
+#include "CADKernel/Mesh/Structure/Mesh.h"
+
+namespace CADKernel
+{
+	class FCriterion;
+	class FMesh;
+
+	class FVertexMesh;
+	class FEdgeMesh;
+	class FFaceMesh;
+	class FSession;
+
+	class CADKERNEL_API FModelMesh : public FEntityGeom
+	{
+		friend FEntity;
+
+		//TSharedRef<FSession> Session;
+		TArray<TSharedPtr<FCriterion>> Criteria;
+
+		TArray<TArray<FPoint>*> GlobalPointCloud;
+		FIdent LastIdUsed;
+
+		TArray<TSharedPtr<FVertexMesh>> VertexMeshes;
+		TArray<TSharedPtr<FEdgeMesh>> EdgeMeshes;
+		TArray<TSharedPtr<FFaceMesh>> FaceMeshes;
+
+		bool QuadAnalyse = false;
+		double MinSize = SMALL_NUMBER;
+		double MaxSize = HUGE_VALUE;
+		double MaxAngle = PI;
+		double Sag = HUGE_VALUE;
+
+		FModelMesh(/*TSharedRef<FSession>& InSession*/)
+			: FEntityGeom()
+			//, Session(InSession)
+			, LastIdUsed(0)
+		{
+		}
+
+	public:
+
+		virtual void SpawnIdent(FDatabase& Database) override
+		{
+			if (!FEntity::SetId(Database))
+			{
+				return;
+			}
+
+			SpawnIdentOnEntities(VertexMeshes, Database);
+			SpawnIdentOnEntities(EdgeMeshes, Database);
+			SpawnIdentOnEntities(FaceMeshes, Database);
+		};
+
+		virtual void ResetMarkersRecursively() 
+		{
+			ResetMarkers();
+			ResetMarkersRecursivelyOnEntities(VertexMeshes);
+			ResetMarkersRecursivelyOnEntities(EdgeMeshes  );
+			ResetMarkersRecursivelyOnEntities(FaceMeshes  );
+		};
+
+
+		virtual EEntity GetEntityType() const override
+		{
+			return EEntity::MeshModel;
+		}
+
+		const TArray<TSharedPtr<FCriterion>>& GetCriteria() const
+		{
+			return Criteria;
+		}
+
+		void AddCriterion(TSharedPtr<FCriterion>& Criterion);
+
+		double GetGeometricTolerance() const
+		{
+			return 0.02; //Session->GetGeometricTolerance();
+		}
+
+		double GetMinSize() const
+		{
+			return MinSize;
+		}
+
+		double GetMaxSize() const
+		{
+			return MaxSize;
+		}
+
+		double GetAngleCriteria()
+		{
+			return MaxAngle;
+		}
+
+		double GetSag() const
+		{
+			return Sag;
+		}
+
+		void AddMesh(TSharedRef<FVertexMesh> Mesh)
+		{
+			VertexMeshes.Add(Mesh);
+		}
+
+		void AddMesh(TSharedRef<FEdgeMesh> Mesh)
+		{
+			EdgeMeshes.Add(Mesh);
+		}
+
+		void AddMesh(TSharedRef<FFaceMesh> Mesh)
+		{
+			FaceMeshes.Add(Mesh);
+		}
+
+		void RegisterCoordinates(TArray<FPoint>& Coordinates, int32& OutStartVertexId, int32& OutIndex)
+		{
+			OutIndex = (int32)GlobalPointCloud.Num();
+			OutStartVertexId = LastIdUsed;
+
+			GlobalPointCloud.Add(&Coordinates);
+			LastIdUsed += (int32)Coordinates.Num();
+		}
+
+		const int32 GetIndexOfVertexFromId(const int32 Id) const;
+		const int32 GetIndexOfEdgeFromId(const int32 Id) const;
+		const int32 GetIndexOfSurfaceFromId(const int32 Id) const;
+
+		const TSharedPtr<FVertexMesh> GetMeshOfVertexNodeId(const int32 Id) const;
+
+		void GetNodeCoordinates(TArray<FPoint>& NodeCoordinates) const;
+		void GetNodeCoordinates(TArray<FVector>& NodeCoordinates) const;
+
+		const TArray<TSharedPtr<FMesh>>& GetMeshes() const;
+
+#ifdef CADKERNEL_DEV
+		virtual FInfoEntity& GetInfo(FInfoEntity&) const override;
+#endif
+	};
+
+}
+
diff --git a/Engine/Source/Runtime/Datasmith/CADKernel/Public/CADKernel/Mesh/Structure/VertexMesh.h b/Engine/Source/Runtime/Datasmith/CADKernel/Public/CADKernel/Mesh/Structure/VertexMesh.h
new file mode 100644
--- /dev/null
+++ b/Engine/Source/Runtime/Datasmith/CADKernel/Public/CADKernel/Mesh/Structure/VertexMesh.h
@@ -0,0 +1,26 @@
+// Copyright Epic Games, Inc. All Rights Reserved.
+#pragma once
+
+#include "CADKernel/Mesh/Structure/Mesh.h"
+
+namespace CADKernel
+{
+	class FModelMesh;
+
+	class CADKERNEL_API FVertexMesh : public FMesh
+	{
+	protected:
+
+	public:
+		FVertexMesh(TSharedRef<FModelMesh> Model, TSharedRef<FTopologicalEntity> TopologicalEntity)
+			: FMesh(Model, TopologicalEntity)
+		{
+		}
+
+		int32 GetMesh() const
+		{
+			return StartNodeId;
+		}
+	};
+}
+
diff --git a/Engine/Source/Runtime/Datasmith/CADKernel/Public/CADKernel/Topo/Body.h b/Engine/Source/Runtime/Datasmith/CADKernel/Public/CADKernel/Topo/Body.h
new file mode 100644
--- /dev/null
+++ b/Engine/Source/Runtime/Datasmith/CADKernel/Public/CADKernel/Topo/Body.h
@@ -0,0 +1,112 @@
+// Copyright Epic Games, Inc. All Rights Reserved.
+#pragma once
+
+#include "CADKernel/Core/MetadataDictionary.h"
+#include "CADKernel/Geo/GeoEnum.h"
+#include "CADKernel/Topo/Shell.h"
+#include "CADKernel/Topo/TopologicalEntity.h"
+#include "CADKernel/Topo/TopologicalFace.h"
+
+namespace CADKernel
+{
+	class FShell;
+
+	class CADKERNEL_API FBody : public FTopologicalEntity, public FMetadataDictionary
+	{
+		friend FEntity;
+
+	private:
+		TArray<TSharedPtr<FShell>> Shells;
+
+		FBody()
+			: FTopologicalEntity()
+		{
+		}
+
+		FBody(const TArray<TSharedPtr<FShell>>& InShells)
+			: FTopologicalEntity()
+		{
+			for (TSharedPtr<FShell> Shell : InShells)
+			{
+				if (Shell.IsValid())
+				{
+					AddShell(Shell.ToSharedRef());
+				}
+			}
+		}
+
+		FBody(FCADKernelArchive& Archive)
+			: FTopologicalEntity()
+		{
+			Serialize(Archive);
+		}
+
+	public:
+		virtual void Serialize(FCADKernelArchive& Ar) override
+		{
+			FTopologicalEntity::Serialize(Ar);
+			SerializeIdents(Ar, Shells);
+		}
+
+		virtual void SpawnIdent(FDatabase& Database) override
+		{
+			if (!FEntity::SetId(Database))
+			{
+				return;
+			}
+
+			SpawnIdentOnEntities(Shells, Database);
+		}
+
+		virtual void ResetMarkersRecursively() override
+		{
+			ResetMarkers();
+			ResetMarkersRecursivelyOnEntities(Shells);
+		}
+
+#ifdef CADKERNEL_DEV
+		virtual FInfoEntity& GetInfo(FInfoEntity&) const override;
+#endif
+
+		virtual EEntity GetEntityType() const override
+		{
+			return EEntity::Body;
+		}
+
+		void AddShell(TSharedRef<FShell> Shell);
+
+		const TArray<TSharedPtr<FShell>>& GetShells() const
+		{
+			return Shells;
+		}
+
+		virtual int32 FaceCount() const override
+		{
+			int32 FaceCount = 0;
+			for (const TSharedPtr<FShell>& Shell : Shells)
+			{
+				FaceCount += Shell->FaceCount();
+			}
+			return FaceCount;
+		}
+
+		virtual void GetFaces(TArray<TSharedPtr<FTopologicalFace>>& Faces) override
+		{
+			for (const TSharedPtr<FShell>& Shell : Shells)
+			{
+				Shell->GetFaces(Faces);
+			}
+		}
+
+		virtual void SpreadBodyOrientation() override
+		{
+			for (const TSharedPtr<FShell>& Shell : Shells)
+			{
+				Shell->SpreadBodyOrientation();
+			}
+		}
+
+		virtual TSharedPtr<FEntityGeom> ApplyMatrix(const FMatrixH& InMatrix) const override;
+	};
+}
+
diff --git a/Engine/Source/Runtime/Datasmith/CADKernel/Public/CADKernel/Topo/Linkable.h b/Engine/Source/Runtime/Datasmith/CADKernel/Public/CADKernel/Topo/Linkable.h
new file mode 100644
--- /dev/null
+++ b/Engine/Source/Runtime/Datasmith/CADKernel/Public/CADKernel/Topo/Linkable.h
@@ -0,0 +1,220 @@
+// Copyright Epic Games, Inc. All Rights Reserved.
+#pragma once
+
+#include "CADKernel/Core/EntityGeom.h"
+#include "CADKernel/Topo/TopologicalEntity.h"
+#include "CADKernel/Topo/TopologicalLink.h"
+
+namespace CADKernel
+{
+	class FModelMesh;
+
+	template<typename EntityType, typename LinkType>
+	class CADKERNEL_API TLinkable : public FTopologicalEntity
+	{
+	protected:
+		mutable TSharedPtr<LinkType> TopologicalLink;
+
+	public:
+		TLinkable()
+			: FTopologicalEntity()
+			, TopologicalLink(TSharedPtr<LinkType>())
+		{
+		}
+
+		TLinkable(FCADKernelArchive& Archive)
+			: TLinkable()
+		{
+			Serialize(Archive);
+		}
+
+		virtual void Serialize(FCADKernelArchive& Ar) override
+		{
+			FEntityGeom::Serialize(Ar);
+			SerializeIdent(Ar, TopologicalLink);
+		}
+
+		const TSharedRef<const EntityType> GetLinkActiveEntity() const
+		{
+			if (!TopologicalLink.IsValid()) 
+			{
+				return StaticCastSharedRef<const EntityType>(AsShared());
+			}
+			TWeakPtr<EntityType>& ActiveEntity = TopologicalLink->GetActiveEntity();
+			return ActiveEntity.Pin().ToSharedRef();
+		}
+
+		TSharedRef<EntityType> GetLinkActiveEntity()
+		{
+			if (!TopologicalLink.IsValid())
+			{
+				return StaticCastSharedRef<EntityType>(AsShared());
+			}
+			TWeakPtr<EntityType> ActiveEntity = TopologicalLink->GetActiveEntity();
+			return ActiveEntity.Pin().ToSharedRef();
+		}
+
+		bool IsActiveEntity() const 
+		{
+			if (!TopologicalLink.IsValid())
+			{
+				return true;
+			}
+			
+			if(TopologicalLink->GetTwinsEntitieNum() == 1)
+			{
+				return true;
+			}
+
+			return (TopologicalLink->GetActiveEntity() == StaticCastSharedRef<const EntityType>(AsShared()));
+		}
+
+		void Activate()
+		{
+			if (TopologicalLink.IsValid())
+			{
+				TopologicalLink->ActivateEntity(StaticCastSharedRef<EntityType>(AsShared()));
+			}
+		}
+
+		virtual TSharedPtr<const LinkType> GetLink() const
+		{
+			ensureCADKernel(TopologicalLink.IsValid());
+			return TopologicalLink;
+		}
+
+		virtual TSharedPtr<LinkType> GetLink()
+		{
+			if (!TopologicalLink.IsValid())
+			{
+				TopologicalLink = FEntity::MakeShared<LinkType>(StaticCastSharedRef<EntityType>(AsShared()));
+			}
+			return TopologicalLink;
+		}
+
+		void ResetTopologicalLink()
+		{
+			TopologicalLink = TSharedPtr<LinkType>();
+		}
+
+		bool IsLinkedTo(TSharedRef<EntityType> Entity) const
+		{
+			if (Entity == AsShared())
+			{
+				return true;
+			}
+			return (Entity->TopologicalLink == TopologicalLink);
+		}
+
+		int32 GetTwinsEntityCount() const
+		{
+			if (TopologicalLink.IsValid())
+			{
+				return TopologicalLink->GetTwinsEntitieNum();
+			}
+			return 1;
+		}
+
+
+		bool HasTwin() const
+		{
+			return GetTwinsEntityCount() != 1;
+		}
+
+		const TArray<TWeakPtr<EntityType>>& GetTwinsEntities() const
+		{
+			if (!TopologicalLink.IsValid())
+			{
+				TopologicalLink = FEntity::MakeShared<LinkType>(ConstCastSharedRef<EntityType>(StaticCastSharedRef<const EntityType>(AsShared())));
+			}
+			return TopologicalLink->GetTwinsEntities();
+		}
+
+		const TArray<TWeakPtr<EntityType>>& GetTwinsEntities()
+		{
+			if(!TopologicalLink.IsValid())
+			{
+				TopologicalLink = FEntity::MakeShared<LinkType>(StaticCastSharedRef<EntityType>(AsShared()));
+			}
+			return TopologicalLink->GetTwinsEntities();
+		}
+
+		void RemoveFromLink() 
+		{
+			if (TopologicalLink.IsValid())
+			{
+				TopologicalLink->RemoveEntity(StaticCastSharedRef<EntityType>(AsShared()));
+				TopologicalLink.Reset();
+			}
+		}
+
+		const bool IsThinZone() const
+		{
+			return ((States & EHaveStates::ThinZone) == EHaveStates::ThinZone);
+		}
+		
+		virtual void SetThinZone()
+		{
+			States |= EHaveStates::ThinZone;
+		}
+		
+		virtual void ResetThinZone()
+		{
+			States &= ~EHaveStates::ThinZone;
+		}
+
+	protected:
+
+		void MakeLink(TSharedRef<EntityType> Twin)
+		{
+			TSharedPtr<LinkType> Link1 = TopologicalLink;
+			TSharedPtr<LinkType> Link2 = Twin->TopologicalLink;
+
+			if (!Link1.IsValid() && !Link2.IsValid())
+			{
+				TopologicalLink = FEntity::MakeShared<LinkType>(StaticCastSharedRef<EntityType>(AsShared()));
+				TopologicalLink->AddEntity(Twin);
+				Twin->SetTopologicalLink(TopologicalLink);
+			}
+			else if (Link1.IsValid() && Link2.IsValid())
+			{
+				if (Link1 == Link2)
+				{
+					return;
+				}
+
+				if (Link2->GetTwinsEntities().Num() > Link1->GetTwinsEntities().Num())
+				{
+					Swap(Link1, Link2);
+				}
+				
+				Link1->AddEntities(Link2->GetTwinsEntities());
+				for (const TWeakPtr<EntityType>& Entity : Link2->GetTwinsEntities())
+				{
+					Entity.Pin()->SetTopologicalLink(Link1);
+				}
+			}
+			else
+			{
+				if (Link1.IsValid())
+				{
+					Link1->AddEntity(Twin);
+					Twin->SetTopologicalLink(Link1);
+				}
+				else
+				{
+					Link2->AddEntity(StaticCastSharedRef<EntityType>(AsShared()));
+					SetTopologicalLink(Link2);
+				}
+			}
+		}
+
+	protected :
+		void SetTopologicalLink(TSharedPtr<LinkType> Link)
+		{
+			TopologicalLink = Link;
+		}
+	};
+
+} // namespace CADKernel
+
diff --git a/Engine/Source/Runtime/Datasmith/CADKernel/Public/CADKernel/Topo/Model.h b/Engine/Source/Runtime/Datasmith/CADKernel/Public/CADKernel/Topo/Model.h
new file mode 100644
--- /dev/null
+++ b/Engine/Source/Runtime/Datasmith/CADKernel/Public/CADKernel/Topo/Model.h
@@ -0,0 +1,108 @@
+// Copyright Epic Games, Inc. All Rights Reserved.
+#pragma once
+
+#include "CADKernel/Core/EntityGeom.h"
+#include "CADKernel/Topo/TopologicalFace.h"
+
+namespace CADKernel
+{
+
+	class FBody;
+	class FTopologicalFace;
+	class FTopologicalEdge;
+	class FGroup;
+	class FTopologicalVertex;
+
+	class CADKERNEL_API FModel : public FTopologicalEntity, public FMetadataDictionary
+	{
+		friend FEntity;
+
+	protected:
+
+		TArray<TSharedPtr<FBody>> Bodies;
+		TArray<TSharedPtr<FTopologicalFace>> Faces;
+
+		FModel()
+			: FTopologicalEntity()
+		{
+			Bodies.Reserve(100);
+			Faces.Reserve(100);
+		}
+
+		FModel(FCADKernelArchive& Archive)
+			: FTopologicalEntity()
+		{
+			Serialize(Archive);
+		}
+
+	public:
+
+		virtual void Serialize(FCADKernelArchive& Ar) override
+		{
+			FEntityGeom::Serialize(Ar);
+			SerializeIdents(Ar, (TArray<TSharedPtr<FEntity>>&) Bodies);
+			SerializeIdents(Ar, (TArray<TSharedPtr<FEntity>>&) Faces);
+		}
+
+		virtual void SpawnIdent(FDatabase& Database) override
+		{
+			if (!FEntity::SetId(Database))
+			{
+				return;
+			}
+
+			SpawnIdentOnEntities(Bodies, Database);
+			SpawnIdentOnEntities(Faces, Database);
+		}
+
+		virtual void ResetMarkersRecursively() override
+		{
+			ResetMarkers();
+			ResetMarkersRecursivelyOnEntities(Bodies);
+			ResetMarkersRecursivelyOnEntities(Faces);
+		}
+
+#ifdef CADKERNEL_DEV
+		virtual FInfoEntity& GetInfo(FInfoEntity&) const override;
+#endif
+
+		virtual EEntity GetEntityType() const override
+		{
+			return EEntity::Model;
+		}
+
+		void AddEntity(TSharedRef<FTopologicalEntity> Entity);
+
+		void Add(TSharedPtr<FTopologicalFace> Face);
+		void Add(TSharedPtr<FBody> Body);
+
+		void RemoveEntity(TSharedPtr<FTopologicalEntity> Entity);
+		void RemoveDomain(TSharedPtr<FTopologicalFace> ToplologicalFace);
+		void RemoveBody(TSharedPtr<FBody> Body);
+		bool Contains(TSharedPtr<FTopologicalEntity> Entity);
+		
+		const TArray<TSharedPtr<FTopologicalFace>>& GetFaces() const
+		{
+			return Faces;
+		}
+
+		const TArray<TSharedPtr<FTopologicalFace>>& GetSurfaceList() const
+		{
+			return Faces;
+		}
+
+		const TArray<TSharedPtr<FBody>>& GetBodyList() const
+		{
+			return Bodies;
+		}
+
+		virtual int32 FaceCount() const override;
+
+		virtual void GetFaces(TArray<TSharedPtr<FTopologicalFace>>& OutFaces) override;
+
+		virtual void SpreadBodyOrientation() override;
+
+		virtual TSharedPtr<FEntityGeom> ApplyMatrix(const FMatrixH& InMatrix) const override;
+	};
+
+} // namespace CADKernel
diff --git a/Engine/Source/Runtime/Datasmith/CADKernel/Public/CADKernel/Topo/TopologicalEntity.h b/Engine/Source/Runtime/Datasmith/CADKernel/Public/CADKernel/Topo/TopologicalEntity.h
new file mode 100644
--- /dev/null
+++ b/Engine/Source/Runtime/Datasmith/CADKernel/Public/CADKernel/Topo/TopologicalEntity.h
@@ -0,0 +1,65 @@
+// Copyright Epic Games, Inc. All Rights Reserved.
+#pragma once
+
+#include "CADKernel/Core/EntityGeom.h"
+
+namespace CADKernel
+{
+	class FModelMesh;
+	class FTopologicalFace;
+
+	class CADKERNEL_API FTopologicalEntity : public FEntityGeom
+	{
+	public:
+
+		virtual int32 FaceCount() const
+		{
+			return 0;
+		}
+
+		virtual void GetFaces(TArray<TSharedPtr<FTopologicalFace>>& OutFaces) {}
+
+		virtual void SpreadBodyOrientation() {}
+
+		virtual void Serialize(FCADKernelArchive& Ar) override
+		{
+			FEntityGeom::Serialize(Ar);
+		}
+		
+#ifdef CADKERNEL_DEV
+		virtual FInfoEntity& GetInfo(FInfoEntity&) const override;
+#endif
+
+		const bool IsApplyCriteria() const
+		{
+			return ((States & EHaveStates::IsApplyCriteria) == EHaveStates::IsApplyCriteria);
+		}
+
+		virtual void SetApplyCriteria() const
+		{
+			States |= EHaveStates::IsApplyCriteria;
+		}
+		
+		virtual void ResetApplyCriteria()
+		{
+			States &= ~EHaveStates::IsApplyCriteria;
+		}
+
+		bool IsMeshed() const
+		{
+			return ((States & EHaveStates::IsMeshed) == EHaveStates::IsMeshed);
+		}
+		
+		virtual void SetMeshed()
+		{
+			States |= EHaveStates::IsMeshed;
+		}
+		
+		virtual void ResetMeshed()
+		{
+			States &= ~EHaveStates::IsMeshed;
+		}
+	};
+
+} // namespace CADKernel
+
diff --git a/Engine/Source/Runtime/Datasmith/CADKernel/Public/CADKernel/Topo/TopologicalLink.h b/Engine/Source/Runtime/Datasmith/CADKernel/Public/CADKernel/Topo/TopologicalLink.h
new file mode 100644
--- /dev/null
+++ b/Engine/Source/Runtime/Datasmith/CADKernel/Public/CADKernel/Topo/TopologicalLink.h
@@ -0,0 +1,210 @@
+// Copyright Epic Games, Inc. All Rights Reserved.
+#pragma once
+
+#include "CADKernel/Core/Entity.h" 
+#include "CADKernel/Math/Point.h"
+#include "CADKernel/UI/Message.h"
+
+namespace CADKernel
+{
+	class FTopologicalEdge;
+	class FTopologicalVertex;
+
+	template<typename EntityType>
+	class CADKERNEL_API TTopologicalLink : public FEntity
+	{
+		friend FEntity;
+
+	protected:
+		friend EntityType;
+		TWeakPtr<EntityType> ActiveEntity;
+
+		TArray<TWeakPtr<EntityType>> TwinsEntities;
+
+		TTopologicalLink()
+		{
+			ActiveEntity = TWeakPtr<EntityType>();
+		}
+
+		TTopologicalLink(TSharedRef<EntityType> Entity)
+			: FEntity()
+		{
+			TwinsEntities.Add(Entity);
+			ActiveEntity = Entity;
+		}
+
+	public:
+
+		TTopologicalLink(FCADKernelArchive& Ar)
+			: FEntity()
+		{
+			Serialize(Ar);
+		}
+
+		virtual void Serialize(FCADKernelArchive& Ar) override
+		{
+			FEntity::Serialize(Ar);
+			SerializeIdent(Ar, ActiveEntity, false);
+			SerializeIdents(Ar, TwinsEntities, false);
+		}
+
+		TWeakPtr<const EntityType>& GetActiveEntity() const
+		{
+			ensureCADKernel(ActiveEntity.IsValid());
+			return ActiveEntity;
+		}
+
+		TWeakPtr<EntityType>& GetActiveEntity()
+		{
+			ensureCADKernel(ActiveEntity.IsValid());
+			return ActiveEntity;
+		}
+
+		int32 GetTwinsEntitieNum() const
+		{
+			return (int32)TwinsEntities.Num();
+		}
+
+		const TArray<TWeakPtr<EntityType>>& GetTwinsEntities() const
+		{
+			return TwinsEntities;
+		}
+
+		void ActivateEntity(TSharedRef<EntityType> NewActiveEntity)
+		{
+			TFunction<bool()> CheckEntityIsATwin = [&]() {
+				for (TWeakPtr<EntityType>& Entity : TwinsEntities)
+				{
+					if (Entity.Pin() == NewActiveEntity)
+					{
+						return true;
+					}
+				}
+				FMessage::Error(TEXT("FTopologicalLink::ActivateEntity, the topological entity is not found in the twins entities"));
+				return false;
+			};
+
+			ensureCADKernel(CheckEntityIsATwin());
+			ActiveEntity = NewActiveEntity;
+		}
+
+		void RemoveEntity(const TSharedPtr<EntityType>& Entity)
+		{
+			TwinsEntities.Remove(Entity);
+			if (Entity == ActiveEntity && TwinsEntities.Num() > 0)
+			{
+				ActiveEntity = TwinsEntities.HeapTop();
+			}
+		}
+
+#ifdef CADKERNEL_DEV
+		virtual FInfoEntity& GetInfo(FInfoEntity& Info) const override;
+#endif
+
+		virtual EEntity GetEntityType() const override
+		{
+			return EEntity::EdgeLink;
+		}
+
+		void AddEntity(TSharedRef<EntityType> Entity)
+		{
+			TwinsEntities.Add(Entity);
+		}
+
+		template <typename ArrayType>
+		void AddEntities(const ArrayType& Entities)
+		{
+			TwinsEntities.Insert(Entities, TwinsEntities.Num());
+		}
+
+		bool CleanLink()
+		{
+			TArray<TWeakPtr<EntityType>> NewTwinsEntities;
+			NewTwinsEntities.Reserve(TwinsEntities.Num());
+			for (TWeakPtr<EntityType>& Entity : TwinsEntities)
+			{
+				if (Entity.IsValid())
+				{
+					NewTwinsEntities.Add(Entity);
+				}
+			}
+
+			if (NewTwinsEntities.Num() != TwinsEntities.Num())
+			{
+				Swap(NewTwinsEntities, TwinsEntities);
+				ActiveEntity = TwinsEntities.HeapTop();
+				return true;
+			}
+			return false;
+		}
+	};
+
+	/**
+	 * TTopologicalLink overload dedicated to FVertex to manage the barycenter of twin vertices
+	 */
+	class CADKERNEL_API FVertexLink : public TTopologicalLink<FTopologicalVertex>
+	{
+		friend class FTopologicalVertex;
+	protected:
+		FPoint Barycenter;
+
+		void SetBarycenter(const FPoint& Point)
+		{
+			Barycenter = Point;
+		}
+
+	public:
+		FVertexLink()
+			: TTopologicalLink<FTopologicalVertex>()
+			, Barycenter(FPoint::ZeroPoint)
+		{
+		}
+
+		FVertexLink(TSharedRef<FTopologicalVertex> Entity)
+			: TTopologicalLink<FTopologicalVertex>(Entity)
+			, Barycenter(FPoint::ZeroPoint)
+
+		{
+		}
+
+		FVertexLink(FCADKernelArchive& Ar)
+			: TTopologicalLink<FTopologicalVertex>()
+		{
+			Serialize(Ar);
+		}
+
+		virtual void Serialize(FCADKernelArchive& Ar) override
+		{
+			TTopologicalLink<FTopologicalVertex>::Serialize(Ar);
+			Ar << Barycenter;
+		}
+
+#ifdef CADKERNEL_DEV
+		virtual FInfoEntity& GetInfo(FInfoEntity& Info) const override;
+#endif
+
+		const FPoint& GetBarycenter() const
+		{
+			return Barycenter;
+		}
+
+		void CleanLink()
+		{
+			if(TTopologicalLink::CleanLink())
+			{
+				ComputeBarycenter();
+				DefineActiveEntity();
+			}
+		}
+
+		virtual EEntity GetEntityType() const override
+		{
+			return EEntity::VertexLink;
+		}
+
+		void ComputeBarycenter();
+		void DefineActiveEntity();
+	};
+
+
+} // namespace CADKernel
diff --git a/Engine/Source/Runtime/Datasmith/CADKernel/Public/CADKernel/Topo/TopologicalLoop.h b/Engine/Source/Runtime/Datasmith/CADKernel/Public/CADKernel/Topo/TopologicalLoop.h
new file mode 100644
--- /dev/null
+++ b/Engine/Source/Runtime/Datasmith/CADKernel/Public/CADKernel/Topo/TopologicalLoop.h
@@ -0,0 +1,178 @@
+// Copyright Epic Games, Inc. All Rights Reserved.
+#pragma once
+
+#include "CADKernel/Core/OrientedEntity.h"
+#include "CADKernel/Geo/Curves/CompositeCurve.h"
+#include "CADKernel/Geo/GeoEnum.h"
+#include "CADKernel/Topo/TopologicalEdge.h"
+#include "CADKernel/Topo/TopologicalEntity.h"
+#include "CADKernel/Topo/TopologicalVertex.h"
+
+namespace CADKernel
+{
+
+	class FTopologicalFace;
+	class FTopologicalVertex;
+
+	class CADKERNEL_API FOrientedEdge : public TOrientedEntity<FTopologicalEdge>
+	{
+	public:
+		FOrientedEdge(TSharedPtr<FTopologicalEdge>& InEntity, EOrientation InDirection)
+			: TOrientedEntity(InEntity, InDirection)
+		{
+		}
+
+		FOrientedEdge(const FOrientedEdge& OrientiredEntity)
+			: TOrientedEntity(OrientiredEntity)
+		{
+		}
+
+		FOrientedEdge()
+			: TOrientedEntity()
+		{
+		}
+	};
+
+	class CADKERNEL_API FTopologicalLoop : public FTopologicalEntity
+	{
+		friend class FEntity;
+		friend class FTopologicalFace;
+		friend class FTopologicalFace;
+		friend class FTopologicalEdge;
+
+	public:
+		FSurfacicBoundary Boundary;
+
+	protected:
+		TArray<FOrientedEdge> Edges;
+
+		TWeakPtr<FTopologicalFace> Surface;
+		bool bExternalLoop;
+
+		FTopologicalLoop(const TArray<TSharedPtr<FTopologicalEdge>>& Edges, const TArray<EOrientation>& EdgeDirections);
+
+		FTopologicalLoop(FCADKernelArchive& Archive)
+			: FTopologicalEntity()
+		{
+			Serialize(Archive);
+		}
+
+	private:
+
+		void SetSurface(TSharedRef<FTopologicalFace> NewDomain)
+		{
+			Surface = NewDomain;
+		}
+
+		void ResetSurface()
+		{
+			Surface.Reset();
+		}
+
+	public:
+
+		~FTopologicalLoop() = default;
+
+		static TSharedPtr<FTopologicalLoop> Make(const TArray<TSharedPtr<FTopologicalEdge>>& EdgeList, const TArray<EOrientation>& EdgeDirections);
+
+		virtual void Serialize(FCADKernelArchive& Ar) override
+		{
+			FTopologicalEntity::Serialize(Ar);
+			SerializeIdents(Ar, (TArray<TOrientedEntity<FEntity>>&) Edges);
+			SerializeIdent(Ar, Surface);
+			Ar << bExternalLoop;
+		}
+
+		virtual void SpawnIdent(FDatabase& Database) override
+		{
+			if (!FEntity::SetId(Database))
+			{
+				return;
+			}
+
+			SpawnIdentOnEntities((TArray<TOrientedEntity<FEntity>>&) Edges, Database);
+		}
+
+		virtual void ResetMarkersRecursively() override
+		{
+			ResetMarkers();
+			ResetMarkersRecursivelyOnEntities((TArray<TOrientedEntity<FEntity>>&) Edges);
+		}
+
+#ifdef CADKERNEL_DEV
+		virtual FInfoEntity& GetInfo(FInfoEntity&) const override;
+#endif
+
+		virtual EEntity GetEntityType() const override
+		{
+			return EEntity::TopologicalLoop;
+		}
+
+		const int32 EdgeCount() const
+		{
+			return Edges.Num();
+		}
+
+		const TArray<FOrientedEdge>& GetEdges() const
+		{
+			return Edges;
+		}
+
+		TArray<FOrientedEdge>& GetEdges()
+		{
+			return Edges;
+		}
+
+		TSharedRef<FTopologicalFace> GetFace() const
+		{
+			return Surface.Pin().ToSharedRef();
+		}
+
+		void Orient();
+		void SwapOrientation();
+
+		void SetAsInnerBoundary()
+		{
+			bExternalLoop = false;
+		}
+
+		void ReplaceEdge(TSharedPtr<FTopologicalEdge>& OldEdge, TSharedPtr<FTopologicalEdge>& NewEdge);
+		void ReplaceEdge(TSharedPtr<FTopologicalEdge>& Edge, TArray<TSharedPtr<FTopologicalEdge>>& NewEdges);
+
+		void RemoveEdge(TSharedPtr<FTopologicalEdge>& Edge);
+		void ReplaceEdgesWithMergedEdge(TArray<TSharedPtr<FTopologicalEdge>>& OldEdges, TSharedPtr<FTopologicalVertex>& MiddleVertex, TSharedPtr<FTopologicalEdge>& NewEdge);
+
+		EOrientation GetDirection(TSharedPtr<FTopologicalEdge>& Edge, bool bAllowLinkedEdge = false) const;
+
+		EOrientation GetDirection(int32 Index) const
+		{
+			return Edges[Index].Direction;
+		}
+
+		const TSharedPtr<FTopologicalEdge>& GetEdge(int32 Index) const
+		{
+			return Edges[Index].Entity;
+		}
+
+		int32 GetEdgeIndex(const TSharedPtr<FTopologicalEdge>& Edge) const
+		{
+			for (int32 Index = 0; Index < Edges.Num(); ++Index)
+			{
+				if (Edge == Edges[Index].Entity)
+				{
+					return Index;
+				}
+			}
+			return -1;
+		}
+
+		void Get2DSampling(TArray<FPoint2D>& LoopSampling);
+
+		void FindSurfaceCorners(TArray<TSharedPtr<FTopologicalVertex>>& OutCorners, TArray<int32>& OutStartSideIndex) const;
+		void FindBreaks(TArray<TSharedPtr<FTopologicalVertex>>& Ruptures, TArray<int32>& OutStartSideIndex, TArray<double>& RuptureValues) const;
+
+		void ComputeBoundaryProperties(const TArray<int32>& StartSideIndex, TArray<FEdge2DProperties>& OutSideProperties) const;
+
+		void EnsureLogicalClosing();
+	};
+}
diff --git a/Engine/Source/Runtime/Datasmith/CADKernel/Public/CADKernel/Topo/TopologicalVertex.h b/Engine/Source/Runtime/Datasmith/CADKernel/Public/CADKernel/Topo/TopologicalVertex.h
new file mode 100644
--- /dev/null
+++ b/Engine/Source/Runtime/Datasmith/CADKernel/Public/CADKernel/Topo/TopologicalVertex.h
@@ -0,0 +1,167 @@
+// Copyright Epic Games, Inc. All Rights Reserved.
+#pragma once
+
+#include "CADKernel/Math/Point.h"
+#include "CADKernel/Topo/Linkable.h"
+#include "CADKernel/Topo/TopologicalLink.h"
+
+namespace CADKernel
+{
+	class FModelMesh;
+	class FTopologicalEdge;
+	class FVertexMesh;
+
+	class CADKERNEL_API FTopologicalVertex : public TLinkable<FTopologicalVertex, FVertexLink>
+	{
+		friend class FEntity;
+		friend class FTopologicalEdge;
+		friend class FVertexLink;
+
+	protected:
+
+		TArray<TWeakPtr<FTopologicalEdge>> ConnectedEdges;
+		FPoint Coordinates;
+		TSharedPtr<FVertexMesh> Mesh;
+
+		FTopologicalVertex(const FPoint& InCoordinates)
+			: Coordinates(InCoordinates)
+			, Mesh(TSharedPtr<FVertexMesh>())
+		{
+		}
+
+		FTopologicalVertex(FCADKernelArchive& Archive)
+			: Mesh(TSharedPtr<FVertexMesh>())
+		{
+			Serialize(Archive);
+		}
+
+	public:
+
+		virtual void Serialize(FCADKernelArchive& Ar) override
+		{
+			TLinkable<FTopologicalVertex, FVertexLink>::Serialize(Ar);
+			Ar << Coordinates;
+			SerializeIdents(Ar, (TArray<TSharedPtr<FEntity>>&) ConnectedEdges);
+		}
+
+		virtual void SpawnIdent(FDatabase& Database) override;
+		//{
+		//	if (!FEntity::SetId(Database))
+		//	{
+		//		return;
+		//	}
+		//	TopologicalLink->SpawnIdent(Database);
+		//	((TSharedPtr<FEntity>) Mesh)->SpawnIdent(Database);
+		//}
+
+		virtual void ResetMarkersRecursively() override
+		{
+			ResetMarkers();
+		}
+
+
+#ifdef CADKERNEL_DEV
+		virtual FInfoEntity& GetInfo(FInfoEntity&) const override;
+#endif
+
+		virtual EEntity GetEntityType() const override
+		{
+			return EEntity::TopologicalVertex;
+		}
+
+		/**
+		 * @return the 3d coordinate of the barycenter of the twin vertices
+		 */
+		inline const FPoint& GetBarycenter() const
+		{
+			if (TopologicalLink.IsValid() && TopologicalLink->GetTwinsEntitieNum() > 1)
+			{
+				return TopologicalLink->GetBarycenter();
+			}
+			return Coordinates;
+		}
+
+		/**
+		 * @return the 3d coordinates of the vertex (prefere GetBarycenter())
+		 */
+		const FPoint& GetCoordinates() const
+		{
+			return Coordinates;
+		}
+
+		void SetCoordinates(const FPoint& NewCoordinates)
+		{
+			if (GetLink()->GetTwinsEntities().Num() > 1)
+			{
+				// Update barycenter
+				FPoint BaryCenter = GetLink()->GetBarycenter() * (double)GetLink()->GetTwinsEntities().Num();
+				BaryCenter -= Coordinates;
+				BaryCenter += NewCoordinates;
+				BaryCenter /= (double)GetLink()->GetTwinsEntities().Num();
+				GetLink()->SetBarycenter(BaryCenter);
+			}
+			else
+			{
+				GetLink()->SetBarycenter(NewCoordinates);
+			}
+			Coordinates = NewCoordinates;
+		}
+
+		double Distance(const TSharedRef<FTopologicalVertex>& OtherVertex) const
+		{
+			return Coordinates.Distance((*OtherVertex).Coordinates);
+		}
+
+		double SquareDistance(const TSharedRef<FTopologicalVertex>& OtherVertex) const
+		{
+			return Coordinates.SquareDistance((*OtherVertex).Coordinates);
+		}
+
+		double SquareDistance(const FPoint& Point) const
+		{
+			return Coordinates.SquareDistance(Point);
+		}
+
+		virtual TSharedPtr<FEntityGeom> ApplyMatrix(const FMatrixH& InMatrix) const override;
+
+		TSharedRef<FVertexMesh> GetOrCreateMesh(TSharedRef<FModelMesh>& MeshModel);
+
+		const TSharedRef<FVertexMesh> GetMesh() const
+		{
+			if (!IsActiveEntity())
+			{
+				return GetLinkActiveEntity()->GetMesh();
+			}
+			return Mesh.ToSharedRef();
+		}
+
+		void Link(TSharedRef<FTopologicalVertex> InEntity);
+
+		void DeleteIfIsolated()
+		{
+			if (ConnectedEdges.Num() == 0)
+			{
+				if (TopologicalLink.IsValid())
+				{
+					TopologicalLink->RemoveEntity(StaticCastSharedRef<FTopologicalVertex>(AsShared()));
+					TopologicalLink->ComputeBarycenter();
+					TopologicalLink.Reset();
+				}
+			}
+			SetDeleted();
+		}
+
+		bool IsBorderVertex();
+
+		void AddConnectedEdge(TSharedRef<FTopologicalEdge> Edge);
+		void RemoveConnectedEdge(TSharedRef<FTopologicalEdge> Edge);
+
+		const TArray<TWeakPtr<FTopologicalEdge>>& GetConnectedEdges() const
+		{
+			return ConnectedEdges;
+		}
+
+		void GetConnectedEdges(TSharedPtr<FTopologicalVertex> OtherVertex, TArray<TSharedPtr<FTopologicalEdge>>& Edges) const;
+	};
+
+} // namespace CADKernel
diff --git a/Engine/Source/Runtime/Datasmith/CADKernel/Public/CADKernel/UI/Console.h b/Engine/Source/Runtime/Datasmith/CADKernel/Public/CADKernel/UI/Console.h
new file mode 100644
--- /dev/null
+++ b/Engine/Source/Runtime/Datasmith/CADKernel/Public/CADKernel/UI/Console.h
@@ -0,0 +1,24 @@
+// Copyright Epic Games, Inc. All Rights Reserved.
+#pragma once
+
+#include "CADKernel/Core/Types.h"
+
+#include "CADKernel/UI/Message.h"
+
+namespace CADKernel
+{
+	class CADKERNEL_API FConsole
+	{
+	public:
+
+		virtual ~FConsole()
+		{
+		}
+
+		virtual void Print(const TCHAR* Texte, EVerboseLevel VerboseLevel = EVerboseLevel::NoVerbose)
+		{}
+
+	};
+
+} // namespace CADKernel
+
diff --git a/Engine/Source/Runtime/Datasmith/CADKernel/Public/CADKernel/UI/Display.h b/Engine/Source/Runtime/Datasmith/CADKernel/Public/CADKernel/UI/Display.h
new file mode 100644
--- /dev/null
+++ b/Engine/Source/Runtime/Datasmith/CADKernel/Public/CADKernel/UI/Display.h
@@ -0,0 +1,212 @@
+// Copyright Epic Games, Inc. All Rights Reserved.
+#pragma once
+
+#include "CADKernel/Core/Types.h"
+#include "CADKernel/Core/System.h"
+#include "CADKernel/Geo/GeoEnum.h"
+#include "CADKernel/Math/Point.h"
+#include "CADKernel/UI/Visu.h"
+
+class FString;
+namespace CADKernel
+{
+	class FAABB;
+	class FAABB2D;
+	class FBody;
+	class FCurve;
+	class FRestrictionCurve;
+	class FTopologicalEdge;
+	class FEdgeMesh;
+	class FEntity;
+	class FGroup;
+	class FMesh;
+
+	class FVertexMesh;
+	class FEdgeMesh;
+	class FFaceMesh;
+	class FModelMesh;
+
+	class FModel;
+	class FShell;
+	class FSurface;
+	class FFaceMesh;
+	class FTopologicalFace;
+	class FTopologicalLoop;
+	class FTopologicalVertex;
+	class FVertexMesh;
+
+	struct FLinearBoundary;
+
+	void Wait(bool bMakeWait = true);
+
+	void Open3DDebugSession(FString name, const TArray<FIdent>& idList = TArray<FIdent>());
+	void Close3DDebugSession();
+
+	class CADKERNEL_API F3DDebugSession
+	{
+	public :
+		F3DDebugSession(FString name, const TArray<FIdent>& Idents = TArray<FIdent>())
+		{
+			Open3DDebugSession(name, Idents);
+		}
+
+		~F3DDebugSession()
+		{
+			Close3DDebugSession();
+		}
+	};
+
+	CADKERNEL_API void Open3DDebugSegment(FIdent Ident);
+	CADKERNEL_API void Close3DDebugSegment();
+
+	class CADKERNEL_API F3DDebugSegment
+	{
+	public:
+		F3DDebugSegment(FIdent Ident)
+		{
+			Open3DDebugSegment(Ident);
+		}
+
+		~F3DDebugSegment()
+		{
+			Close3DDebugSegment();
+		}
+	};
+
+	void CADKERNEL_API FlushVisu();
+
+	template<typename TPoint>
+	void DrawPoint(const TPoint& InPoint, EVisuProperty Property = EVisuProperty::BluePoint)
+	{
+		FSystem::Get().GetVisu()->DrawPoint(InPoint, Property);
+	}
+
+	/**
+	 * Draw a mesh element i.e. Element of dimension 1 is an edge, of dimension 2 is a triangle or quadrangle, of dimension 3 is tetrahedron, pyramid, hexahedron, ...
+	 */
+	CADKERNEL_API void DrawElement(int32 Dimension, TArray<FPoint>& Points, EVisuProperty Property = EVisuProperty::Element);
+
+	template<typename TPoint>
+	void Draw(const TArray<TPoint>& Points, EVisuProperty Property = EVisuProperty::BlueCurve)
+	{
+		FSystem::Get().GetVisu()->DrawPolyline(Points, Property);
+	}
+
+	CADKERNEL_API void DrawMesh(const TSharedPtr<FMesh>& mesh);
+
+	CADKERNEL_API void DisplayEdgeCriteriaGrid(int32 EdgeId, const TArray<FPoint>& Points3D);
+
+	template<typename TPoint>
+	void DisplayPoint(const TPoint& Point, FIdent Ident)
+	{
+		F3DDebugSegment G(Ident);
+		DrawPoint(Point);
+	}
+
+	template<typename TPoint>
+	void DisplayPoint(const TPoint& Point, EVisuProperty Property = EVisuProperty::BluePoint)
+	{
+		DrawPoint(Point, Property);
+	}
+
+	template<typename TPoint>
+	void DisplayPoint(const TPoint& Point, EVisuProperty Property, FIdent Ident)
+	{
+		F3DDebugSegment G(Ident);
+		DrawPoint(Point, Property);
+	}
+
+	CADKERNEL_API void DisplayAABB(const FAABB& aabb, FIdent Ident = 0);
+	CADKERNEL_API void DisplayAABB2D(const FAABB2D& aabb, FIdent Ident = 0);
+
+	CADKERNEL_API void DisplayEntity(const TSharedPtr<FEntity>& Entity);
+	CADKERNEL_API void DisplayEntity2D(const TSharedPtr<FEntity>& Entity);
+
+	CADKERNEL_API void DisplayLoop(const TSharedPtr<FTopologicalFace>& Entity);
+	CADKERNEL_API void DisplayIsoCurve(const TSharedPtr<FSurface>& CarrierSurface, double Coordinate, EIso Type);
+
+	CADKERNEL_API void Display(const FPlane& plane, FIdent Ident = 0);
+
+	CADKERNEL_API void Display(const TSharedPtr<FCurve>& Curve);
+	CADKERNEL_API void Display(const TSharedPtr<FSurface>& CarrierSurface);
+
+	CADKERNEL_API void Display(const TSharedPtr<FGroup>& Group);
+	CADKERNEL_API void Display(const TSharedPtr<FModel>& Model);
+	CADKERNEL_API void Display(const TSharedPtr<FBody>& Body);
+	CADKERNEL_API void Display(const TSharedPtr<FShell>& Shell);
+	CADKERNEL_API void Display(const TSharedPtr<FTopologicalEdge>& Edge, EVisuProperty Property = EVisuProperty::BlueCurve);
+	CADKERNEL_API void Display(const TSharedPtr<FTopologicalFace>& Face);
+	CADKERNEL_API void Display(const TSharedPtr<FTopologicalLoop>& Loop);
+	CADKERNEL_API void Display(const TSharedPtr<FTopologicalVertex>& Vertex);
+
+	CADKERNEL_API void Display2D(const TSharedPtr<FTopologicalEdge>& Edge, EVisuProperty Property = EVisuProperty::BlueCurve);
+	CADKERNEL_API void Display2D(const TSharedPtr<FTopologicalFace>& Face);
+	CADKERNEL_API void Display2D(const TSharedPtr<FTopologicalLoop>& Loop);
+	CADKERNEL_API void Display2D(const TSharedPtr<FSurface>& CarrierSurface);
+
+	CADKERNEL_API void DisplayMesh(const TSharedPtr<FMesh>& Mesh);
+	CADKERNEL_API void DisplayMesh(const TSharedRef<FFaceMesh>& Mesh);
+	CADKERNEL_API void DisplayMesh(const TSharedRef<FEdgeMesh>& Mesh);
+	CADKERNEL_API void DisplayMesh(const TSharedRef<FVertexMesh>& Mesh);
+
+	CADKERNEL_API void Display(const TSharedPtr<FModelMesh>& MeshModel);
+
+	CADKERNEL_API void DisplayControlPolygon(const TSharedPtr<FCurve>& Entity);
+	CADKERNEL_API void DisplayControlPolygon(const TSharedPtr<FSurface>& Entity);
+
+	template<typename TPoint>
+	void DisplaySegment(const TPoint& Point1, const TPoint& Point2, FIdent Ident = 0, EVisuProperty Property = EVisuProperty::Element, bool bWithOrientation = false)
+	{
+		F3DDebugSegment G(Ident);
+		if (bWithOrientation)
+		{
+			DrawSegmentOrientation(Point1, Point2, Property);
+		}
+		DrawSegment(Point1, Point2, Property);
+	};
+
+	template<typename TPoint>
+	void DisplayPolyline(const TArray<TPoint>& Points, EVisuProperty Property)
+	{
+		Open3DDebugSegment(0);
+		Draw(Points, Property);
+		Close3DDebugSegment();
+	}
+
+	CADKERNEL_API void DrawQuadripode(double Height, double Base, FPoint& Centre, FPoint& Direction);
+
+	CADKERNEL_API void Draw(const TSharedPtr<FTopologicalEdge>& Edge, EVisuProperty Property = EVisuProperty::BlueCurve);
+	CADKERNEL_API void Draw(const TSharedPtr<FTopologicalFace>& Face);
+	CADKERNEL_API void Draw2D(const TSharedPtr<FTopologicalFace>& Face);
+	CADKERNEL_API void Draw(const TSharedPtr<FShell>& Shell);
+
+	CADKERNEL_API void Draw(const TSharedPtr<FCurve>& Curve, EVisuProperty Property = EVisuProperty::BlueCurve);
+	CADKERNEL_API void Draw(const TSharedPtr<FCurve>& Curve, const FLinearBoundary& Boundary, EVisuProperty Property = EVisuProperty::BlueCurve);
+	CADKERNEL_API void Draw(const FLinearBoundary& Boundary, const TSharedPtr<FRestrictionCurve>& Curve, EVisuProperty Property = EVisuProperty::BlueCurve);
+
+	template<typename TPoint>
+	void DrawSegment(const TPoint& Point1, const TPoint& Point2, EVisuProperty Property = EVisuProperty::Element)
+	{
+		TArray<FPoint> Points;
+		Points.Add(Point1);
+		Points.Add(Point2);
+		Draw(Points, Property);
+	}
+
+	template<typename TPoint>
+	void DrawSegmentOrientation(const TPoint& Point1, const TPoint& Point2, EVisuProperty Property = EVisuProperty::Element)
+	{
+		double Length = Point1.Distance(Point2);
+		double Height = Length / 10.0;
+		double Base = Height / 2;
+
+		FPoint Middle = (Point1 + Point2) / 2.;
+		FPoint Tangent = Point2 - Point1;
+		DrawQuadripode(Height, Base, Middle, Tangent);
+	}
+
+	CADKERNEL_API void DrawIsoCurves(const TSharedPtr<FTopologicalFace>& Face);
+
+
+} // namespace CADKernel
+
diff --git a/Engine/Source/Runtime/Datasmith/CADKernel/Public/CADKernel/UI/Message.h b/Engine/Source/Runtime/Datasmith/CADKernel/Public/CADKernel/UI/Message.h
new file mode 100644
--- /dev/null
+++ b/Engine/Source/Runtime/Datasmith/CADKernel/Public/CADKernel/UI/Message.h
@@ -0,0 +1,81 @@
+// Copyright Epic Games, Inc. All Rights Reserved.
+#pragma once
+
+#include "CADKernel/Core/Types.h"
+
+#include "CADKernel/UI/Progress.h"
+
+namespace CADKernel
+{
+	extern const char* VerboseLevelConstNames[];
+	extern const char* VerboseConstDescHelp[];
+
+
+	class CADKERNEL_API FMessage
+	{
+		friend class FProgressManager;
+
+	private:
+		static int32 NumberOfIndentation;
+		static int32 OldPercent;
+
+		static void VPrintf(EVerboseLevel Level, const TCHAR* Text, ...);
+
+	public:
+
+		template <typename FmtType, typename... Types>
+		static void Printf(EVerboseLevel Level, const FmtType& Text, Types... Args)
+		{
+			static_assert(TAnd<TIsValidVariadicFunctionArg<Types>...>::Value, "Invalid argument(s) passed to FString::Printf");
+			VPrintf(Level, Text, Args...);
+		}
+
+		template <typename FmtType, typename... Types>
+		static void Error(const FmtType& Text, Types... Args)
+		{
+			static_assert(TAnd<TIsValidVariadicFunctionArg<Types>...>::Value, "Invalid argument(s) passed to FString::Printf");
+			FString CompleteText = TEXT("ERROR: ");
+			CompleteText += Text;
+			VPrintf(Log, *CompleteText, Args...);
+		}
+
+		template <typename FmtType, typename... Types>
+		static void Warning(const FmtType& Text, Types... Args)
+		{
+			static_assert(TAnd<TIsValidVariadicFunctionArg<Types>...>::Value, "Invalid argument(s) passed to FString::Printf");
+			FString CompleteText = TEXT("WARNING: ");
+			CompleteText += Text;
+			VPrintf(Log, *CompleteText, Args...);
+		}
+
+		static void Indent(int32 InNumberOfIndent = 1)
+		{
+			NumberOfIndentation += InNumberOfIndent;
+		}
+
+		static void Deindent(int32 InNumberOfIndent = 1)
+		{
+			NumberOfIndentation -= InNumberOfIndent;
+			if (NumberOfIndentation < 0)
+			{
+				NumberOfIndentation = 0;
+			}
+		}
+
+	};
+} // namespace CADKernel
+
+#define ERROR_FUNCTION_CALL_NOT_EXPECTED \
+{ \
+	FMessage::Printf(Log, L"CALL of \" %s \" NOT EXPECTED at line %d of the file %s\n", __func__, __LINE__, __FILE__); \
+}
+
+#define ERROR_NOT_EXPECTED \
+{ \
+	FMessage::Printf(Log, L"Error not expected in \" %S \" at line %d of the file %S\n", __func__, __LINE__, __FILE__); \
+}
+
+#define NOT_IMPLEMENTED \
+{ \
+	FMessage::Printf(Log, L"The function \" %s \" at line %d of the file %s is not implemented", __func__, __LINE__, __FILE__); \
+}
diff --git a/Engine/Source/Runtime/Datasmith/CADKernel/Public/CADKernel/UI/Progress.h b/Engine/Source/Runtime/Datasmith/CADKernel/Public/CADKernel/UI/Progress.h
new file mode 100644
--- /dev/null
+++ b/Engine/Source/Runtime/Datasmith/CADKernel/Public/CADKernel/UI/Progress.h
@@ -0,0 +1,121 @@
+// Copyright Epic Games, Inc. All Rights Reserved.
+#pragma once
+
+#include "CADKernel/Core/Types.h"
+
+namespace CADKernel
+{
+	class FProgressBase;
+	
+	class CADKERNEL_API FProgressManager
+	{
+		friend class FProgress;
+
+	protected:
+
+		FProgress* RootProgress;
+		FProgress* CurrentProgress;
+
+		FProgress* GetRoot() const
+		{
+			return RootProgress;
+		}
+
+		virtual void ActivateProgressBar(bool bActive)
+		{}
+
+		/**
+		 * Update the progress bar 
+		 */
+		virtual void Update()
+		{}
+
+		double GetProgression() const;
+
+		/**
+		 * @return a string describing the current step hierarchy 
+		 */
+		virtual FString GetCurrentStep() const;
+
+		virtual void SetCurrentProgress(FProgress* InProgress);
+
+		FProgress* GetCurrentProgress() const
+		{
+			return CurrentProgress;
+		}
+
+	public:
+		FProgressManager()
+			: RootProgress(nullptr)
+			, CurrentProgress(nullptr)
+		{}
+	};
+
+	class CADKERNEL_API FProgress
+	{
+		friend class FProgressManager;
+	private:
+		FString Name;
+
+		FProgress* Parent = nullptr;
+		FProgress* UnderlyingProgress = nullptr;
+
+		int32 StepCount;
+		int32 Progression = 0;
+
+	public:
+		FProgress(int32 InStepCount, const FString& InStepName = FString());
+
+		FProgress(const FString& StepName = FString())
+			: FProgress(1, StepName)
+		{
+		}
+
+		virtual ~FProgress();
+
+		const FString& GetName()
+		{
+			return Name;
+		}
+
+		FProgress* GetParent() const
+		{
+			return Parent;
+		}
+
+		/**
+		 * Increase the advancement of of one step
+		 */
+		void Increase(int32 StepSize = 1);
+
+	protected:
+
+		bool IsRoot() const
+		{
+			return Parent == nullptr;
+		}
+
+		void AddUnderlying(FProgress* Progress)
+		{
+			UnderlyingProgress = Progress;
+		}
+
+		void UnderlyingFinished(FProgress* Progress)
+		{
+			Progression++;
+			UnderlyingProgress = nullptr;
+		}
+
+		double GetProgression() const
+		{
+			double Ratio = ((double)Progression / (double)StepCount);
+
+			if (UnderlyingProgress)
+			{
+				Ratio += UnderlyingProgress->GetProgression() / (double)StepCount;
+			}
+			return Ratio;
+		}
+	};
+}
+
diff --git a/Engine/Source/Runtime/Datasmith/CADKernel/Public/CADKernel/UI/Visu.h b/Engine/Source/Runtime/Datasmith/CADKernel/Public/CADKernel/UI/Visu.h
new file mode 100644
--- /dev/null
+++ b/Engine/Source/Runtime/Datasmith/CADKernel/Public/CADKernel/UI/Visu.h
@@ -0,0 +1,165 @@
+// Copyright Epic Games, Inc. All Rights Reserved.
+#pragma once
+
+#include "CADKernel/Core/Types.h"
+
+#include "CADKernel/Core/Parameters.h"
+#include "CADKernel/Math/Point.h"
+
+namespace CADKernel
+{
+	enum EVisuProperty
+	{
+		BBox = 0,          // Green see CADKernelUI::properties    
+		Point,             // Blue
+		Element,           // Brown
+		Curve,             // Blue
+		Iso,               // Red
+		ControlLine,       // Purple
+		EdgeMesh,          // blue
+		NodeMesh,          // blue
+		BorderEdge,        // Green
+		NonManifoldEdge,   // Red          
+		BorderTopo,        // Dark brown
+		ControlPoint,      // Purple		 
+		YellowPoint,			
+		YellowCurve,			
+		BluePoint,				
+		BlueCurve,				
+		RedPoint,				 
+		RedCurve,				 
+		PurplePoint,			
+		PurpleCurve,			
+		GreenPoint,				
+		GreenCurve,				
+		PinkPoint,
+		PinkCurve,
+		Last,
+	} ;
+
+	extern const TCHAR* VisuPropertyNames[];
+
+	class FMesh;
+
+	class CADKERNEL_API FVisuParameters : public FParameters
+	{
+	public:
+
+		FParameter ChordError;
+
+		/**
+		 *	Number of isos along U to display on a surface or topological face
+		 */ 
+		FParameter IsoUNumber;
+
+		/**
+		 *	Number of isos along V to display on a surface or topological face
+		 */
+		FParameter IsoVNumber;
+
+		FParameter bDisplayCADOrient;
+		FParameter bDisplayMeshOrient;
+		FParameter bDisplayAxis;
+		FParameter bDisplayNormals;
+
+		FParameter NormalLength;
+
+		FVisuParameters()
+			: FParameters(10)
+			, ChordError(TEXT("DiscretizationError"), 0.02, *this)
+			, IsoUNumber(TEXT("IsoUNumber"), 3, *this)
+			, IsoVNumber(TEXT("IsoVNumber"), 3, *this)
+			, bDisplayCADOrient(TEXT("DisplayCadOrient"), false, *this)
+			, bDisplayMeshOrient(TEXT("DisplayMeshOrient"), true, *this)
+			, bDisplayAxis(TEXT("DisplayAxis"), false, *this)
+			, bDisplayNormals(TEXT("DisplayNormals"), false, *this)
+			, NormalLength(TEXT("NormalLength"), 10., *this)
+		{}
+	};
+
+	class CADKERNEL_API FVisu
+	{
+	protected:
+		FVisuParameters Parameters;
+
+		/**
+		 * Number of graphic session opened
+		 */ 
+		int32 SessionNum;
+
+		void NewSession()
+		{
+			ensureCADKernel(SessionNum >= 0);
+			++SessionNum;
+		}
+
+		void EndSession()
+		{
+			--SessionNum;
+			if (SessionNum < 0)
+			{
+				SessionNum = 0;
+			}
+		}
+
+	public:
+
+		FVisu()
+		{
+			SessionNum = 0;
+		}
+
+		virtual ~FVisu() = default;
+
+		int32 GetPropertyCount() 
+		{
+			return EVisuProperty::Last;
+		}
+
+		FVisuParameters* GetParameters() 
+		{ 
+			return &Parameters;
+		}
+
+		virtual void NewDB(const TCHAR* InName) 
+		{}
+
+		virtual void Open3DDebugSession(const TCHAR* SessionName, const TArray<FIdent>& Ids) 
+		{}
+
+		virtual void Close3DDebugSession() 
+		{}
+
+		virtual void Open3DDebugSegment(FIdent Ident) 
+		{}
+
+		virtual void Close3DDebugSegment() 
+		{}
+
+		virtual void UpdateViewer() 
+		{}
+
+		virtual void DrawPoint(const FPoint& Point, EVisuProperty InProperty = EVisuProperty::BluePoint)
+		{}
+
+		virtual void DrawPoint(const FPoint2D& Point, EVisuProperty InProperty = EVisuProperty::BluePoint)
+		{}
+
+		virtual void DrawPoint(const FPointH& Point, EVisuProperty InProperty = EVisuProperty::BluePoint)
+		{}
+
+		virtual void DrawElement(int32 Dimension, TArray<FPoint>& Points, EVisuProperty InProperty = EVisuProperty::Element)
+		{}
+
+		virtual void DrawPolyline(const TArray<FPoint>& Points, EVisuProperty InProperty = EVisuProperty::BlueCurve)
+		{}
+
+		virtual void DrawPolyline(const TArray<FPoint2D>& Points, EVisuProperty InProperty = EVisuProperty::BlueCurve)
+		{}
+
+		virtual void DrawMesh(FIdent MeshId) 
+		{}
+	};
+
+} // namespace CADKernel
+
