Fix Incremental... Issue #65835 -> #65843

diff --git a/Engine/Source/Runtime/Datasmith/CADKernel/Private/CADKernel/Core/Database.cpp b/Engine/Source/Runtime/Datasmith/CADKernel/Private/CADKernel/Core/Database.cpp
--- a/Engine/Source/Runtime/Datasmith/CADKernel/Private/CADKernel/Core/Database.cpp
+++ b/Engine/Source/Runtime/Datasmith/CADKernel/Private/CADKernel/Core/Database.cpp
@@ -347,7 +347,7 @@ void FDatabase::ExpandSelection(TSharedPtr<FEntity> Entity, const TSet<EEntity>&
 	{
 		ensureCADKernel(false);
 		TSharedPtr<FGroup>Group = StaticCastSharedPtr<FGroup>(Entity);
-		for (TSharedPtr<FEntity> GroupEntity : Group->GetEntities())
+		for (const TSharedPtr<FEntity>& GroupEntity : Group->GetEntities())
 		{
 			ExpandSelection(GroupEntity, Filter, Selection);
 		}
@@ -357,12 +357,12 @@ void FDatabase::ExpandSelection(TSharedPtr<FEntity> Entity, const TSet<EEntity>&
 	case EEntity::Model:
 	{
 		TSharedPtr<FModel> InModel = StaticCastSharedPtr<FModel>(Entity);
-		for (TSharedPtr<FBody> Body : InModel->GetBodyList())
+		for (const TSharedPtr<FBody>& Body : InModel->GetBodyList())
 		{
 			ExpandSelection(Body, Filter, Selection);
 		}
 
-		for (TSharedPtr<FTopologicalFace> Domain : InModel->GetFaces())
+		for (const TSharedPtr<FTopologicalFace>& Domain : InModel->GetFaces())
 		{
 			ExpandSelection(Domain, Filter, Selection);
 		}
@@ -372,7 +372,7 @@ void FDatabase::ExpandSelection(TSharedPtr<FEntity> Entity, const TSet<EEntity>&
 	case EEntity::Body:
 	{
 		TSharedPtr<FBody> Body = StaticCastSharedPtr<FBody>(Entity);
-		for (const TSharedPtr<FShell> Shell : Body->GetShells())
+		for (const TSharedPtr<FShell>& Shell : Body->GetShells())
 		{
 			for (const FOrientedFace& Face : Shell->GetFaces())
 			{
@@ -385,7 +385,7 @@ void FDatabase::ExpandSelection(TSharedPtr<FEntity> Entity, const TSet<EEntity>&
 	case EEntity::MeshModel:
 	{
 		TSharedPtr<FModelMesh> MeshModel = StaticCastSharedPtr<FModelMesh>(Entity);
-		for (TSharedPtr<FMesh> Mesh : MeshModel->GetMeshes())
+		for (const TSharedPtr<FMesh>& Mesh : MeshModel->GetMeshes())
 		{
 			ExpandSelection(Mesh, Filter, Selection);
 		}
diff --git a/Engine/Source/Runtime/Datasmith/CADKernel/Private/CADKernel/Core/Parameters.cpp b/Engine/Source/Runtime/Datasmith/CADKernel/Private/CADKernel/Core/Parameters.cpp
--- a/Engine/Source/Runtime/Datasmith/CADKernel/Private/CADKernel/Core/Parameters.cpp
+++ b/Engine/Source/Runtime/Datasmith/CADKernel/Private/CADKernel/Core/Parameters.cpp
@@ -90,14 +90,14 @@ void FParameter::SetFromString(const FString& String)
 	{
 	case EValue::Double:
 		{
-			double TmpValue = _wtof(*String);
+			double TmpValue = FCString::Atod(*String);
 			SetValue(TmpValue);
 		}
 		break;
 
 	case EValue::Integer:
 		{
-			int32 TmpValue = _wtoi(*String);
+			int32 TmpValue = FCString::Atoi(*String);
 			SetValue(TmpValue);
 		}
 		break;
diff --git a/Engine/Source/Runtime/Datasmith/CADKernel/Private/CADKernel/Core/System.cpp b/Engine/Source/Runtime/Datasmith/CADKernel/Private/CADKernel/Core/System.cpp
--- a/Engine/Source/Runtime/Datasmith/CADKernel/Private/CADKernel/Core/System.cpp
+++ b/Engine/Source/Runtime/Datasmith/CADKernel/Private/CADKernel/Core/System.cpp
@@ -15,8 +15,6 @@
 #include <signal.h>
 #endif 
 
-#pragma warning(disable:4996) // unsafe sprintf
-
 using namespace CADKernel;
 
 TUniquePtr<FSystem> FSystem::Instance = nullptr;
@@ -39,22 +37,13 @@ void FSystem::Initialize(bool bIsDll, const FString& LogFilePath, const FString&
 {
 	SetVerboseLevel(Log);
 
-	TCHAR* DebugFolderPath = _wgetenv(TEXT("DS_DBG"));
-	if (DebugFolderPath) 
+	if (LogFilePath.Len() > 0)
 	{
-		DefineLogFile(FString(DebugFolderPath) + TEXT("\\") + ProductName + TEXT(".log"));
-		DefineSpyFile(FString(DebugFolderPath) + TEXT("\\") + ProductName + TEXT(".py"));
+		DefineLogFile(LogFilePath);
 	}
-	else 
+	if (SpyFilePath.Len() > 0)
 	{
-		if (LogFilePath.Len() > 0)
-		{
-			DefineLogFile(LogFilePath);
-		}
-		if (SpyFilePath.Len() > 0)
-		{
-			DefineSpyFile(SpyFilePath);
-		}
+		DefineSpyFile(SpyFilePath);
 	}
 
 	PrintHeader();
@@ -139,11 +128,3 @@ void FSystem::PrintHeader()
 	FMessage::Printf(Log, TEXT("_______________________________________________________________________________\n"));
 	FMessage::Printf(Log, TEXT("\n"));
 }
-
-extern "C"
-{
-#ifdef COMPILER_MSC
-#pragma warning( disable : 4297)
-#endif
-}
-
diff --git a/Engine/Source/Runtime/Datasmith/CADKernel/Private/CADKernel/Geo/Curves/Curve.cpp b/Engine/Source/Runtime/Datasmith/CADKernel/Private/CADKernel/Geo/Curves/Curve.cpp
--- a/Engine/Source/Runtime/Datasmith/CADKernel/Private/CADKernel/Geo/Curves/Curve.cpp
+++ b/Engine/Source/Runtime/Datasmith/CADKernel/Private/CADKernel/Geo/Curves/Curve.cpp
@@ -130,18 +130,6 @@ namespace CADKernel
 		return FEntity::MakeShared<FBoundedCurve>(Tolerance, StaticCastSharedRef<FCurve>(AsShared()), NewBoundary, Dimension);
 	}
 
-	//double FCurve::ComputeLength(const FLinearBoundary& InBoundary) const
-	//{
-	//	if (GetDimension() == 3)
-	//	{
-	//		return ComputeLength3D(InBoundary);
-	//	}
-	//	else
-	//	{
-	//		return ComputeLength2D(InBoundary);
-	//	}
-	//}
-
 	double FCurve::ComputeLength(const FLinearBoundary& InBoundary) const
 	{
 		FPolyline3D Polyline;
@@ -150,28 +138,6 @@ namespace CADKernel
 		return Polyline.GetLength(Boundary);
 	}
 
-	//double FCurve::ComputeLength2D(const FLinearBoundary& InBoundary) const
-	//{
-	//	double Gap = Boundary.Length() / 30;
-	//	int32 PointCount = (int32) (InBoundary.Length() / Gap + 1);
-	//	Gap = InBoundary.Length() / PointCount;
-	//	TArray<double> Coordinates;
-	//	PointCount++;
-	//	Coordinates.Reserve(PointCount);
-	//	double Coordinate = InBoundary.GetMin();
-	//	for (int32 Index = 0; Index < PointCount; ++Index)
-	//	{
-	//		Coordinates.Add(Coordinate);
-	//		Coordinate += Gap;
-	//	}
-
-	//	TArray<FPoint2D> Polyline;
-	//	Evaluate2DPoints(Coordinates, Polyline);
-
-	//	TPolylineApproximator<FPoint2D> Approximator(Coordinates, Polyline);
-	//	return Approximator.ComputeLength();
-	//}
-
 	void FCurve::Presample(const FLinearBoundary& InBoundary, TArray<double>& OutSampling) const
 	{
 		FPolyline3D Presampling;
diff --git a/Engine/Source/Runtime/Datasmith/CADKernel/Private/CADKernel/Geo/Sampler/SamplerOnChord.h b/Engine/Source/Runtime/Datasmith/CADKernel/Private/CADKernel/Geo/Sampler/SamplerOnChord.h
--- a/Engine/Source/Runtime/Datasmith/CADKernel/Private/CADKernel/Geo/Sampler/SamplerOnChord.h
+++ b/Engine/Source/Runtime/Datasmith/CADKernel/Private/CADKernel/Geo/Sampler/SamplerOnChord.h
@@ -169,7 +169,6 @@ namespace CADKernel
 
 		virtual void EvaluatesNewCandidatePoints() override
 		{
-			CandidatePoints;
 			Curve.Evaluate2DPoints(CandidatePoints.Coordinates, CandidatePoints.Points);
 		};
 
diff --git a/Engine/Source/Runtime/Datasmith/CADKernel/Private/CADKernel/Geo/Surfaces/CompositeSurface.cpp b/Engine/Source/Runtime/Datasmith/CADKernel/Private/CADKernel/Geo/Surfaces/CompositeSurface.cpp
--- a/Engine/Source/Runtime/Datasmith/CADKernel/Private/CADKernel/Geo/Surfaces/CompositeSurface.cpp
+++ b/Engine/Source/Runtime/Datasmith/CADKernel/Private/CADKernel/Geo/Surfaces/CompositeSurface.cpp
@@ -45,7 +45,7 @@ void FCompositeSurface::LinesNotDerivables(const FSurfacicBoundary& Bounds, int3
 	}
 
 	int32 SurfaceIndex = 0;
-	for (const TSharedPtr<FSurface> Surface : Surfaces)
+	for (const TSharedPtr<FSurface>& Surface : Surfaces)
 	{
 		FCoordinateGrid SurfaceNotDerivableCoordinates;
 		Surface->LinesNotDerivables(InDerivativeOrder, SurfaceNotDerivableCoordinates);
diff --git a/Engine/Source/Runtime/Datasmith/CADKernel/Private/CADKernel/Geo/Surfaces/CoonsSurface.cpp b/Engine/Source/Runtime/Datasmith/CADKernel/Private/CADKernel/Geo/Surfaces/CoonsSurface.cpp
--- a/Engine/Source/Runtime/Datasmith/CADKernel/Private/CADKernel/Geo/Surfaces/CoonsSurface.cpp
+++ b/Engine/Source/Runtime/Datasmith/CADKernel/Private/CADKernel/Geo/Surfaces/CoonsSurface.cpp
@@ -58,29 +58,29 @@ void FCoonsSurface::LinesNotDerivables(const FSurfacicBoundary& Bounds, int32 In
 void FCoonsSurface::EvaluatePoint(const FPoint2D& InPoint2D, FSurfacicPoint& OutPoint3D, int32 InDerivativeOrder) const
 {
 	TFunction<void(const TSharedPtr<FCurve>*, const FPoint2D&, EIso, FPoint&, FCurvePoint[])> \
-	ComputePointOnRuledSurface = [&](const TSharedPtr<FCurve>* Curves, const FPoint2D& Coord, EIso Iso, FPoint& OutPoint, FCurvePoint OutCurvePoints[2])
+	ComputePointOnRuledSurface = [&](const TSharedPtr<FCurve>* InCurves, const FPoint2D& Coord, EIso Iso, FPoint& OutPoint, FCurvePoint OutCurvePoints[2])
 	{
 		const int32 IndexCurve0 = Iso * 2;
 		const int32 IndexCurve1 = Iso * 2 + 1;
 
 		EIso OtherIso = Iso == EIso::IsoU ? EIso::IsoV : EIso::IsoU;
 
-		double CoordinateCurve0 = Curves[IndexCurve0]->GetUMin() + Coord[Iso] * (Curves[IndexCurve0]->GetUMax() - Curves[IndexCurve0]->GetUMin());
-		double CoordinateCurve1 = Curves[IndexCurve1]->GetUMin() + Coord[Iso] * (Curves[IndexCurve1]->GetUMax() - Curves[IndexCurve1]->GetUMin());
+		double CoordinateCurve0 = InCurves[IndexCurve0]->GetUMin() + Coord[Iso] * (InCurves[IndexCurve0]->GetUMax() - InCurves[IndexCurve0]->GetUMin());
+		double CoordinateCurve1 = InCurves[IndexCurve1]->GetUMin() + Coord[Iso] * (InCurves[IndexCurve1]->GetUMax() - InCurves[IndexCurve1]->GetUMin());
 
-		Curves[IndexCurve0]->EvaluatePoint(CoordinateCurve0, OutCurvePoints[0], InDerivativeOrder);
-		Curves[IndexCurve1]->EvaluatePoint(CoordinateCurve1, OutCurvePoints[1], InDerivativeOrder);
+		InCurves[IndexCurve0]->EvaluatePoint(CoordinateCurve0, OutCurvePoints[0], InDerivativeOrder);
+		InCurves[IndexCurve1]->EvaluatePoint(CoordinateCurve1, OutCurvePoints[1], InDerivativeOrder);
 
 		OutPoint = OutCurvePoints[0].Point + (OutCurvePoints[1].Point - OutCurvePoints[0].Point) * Coord[OtherIso];
 
 		if (InDerivativeOrder > 0)
 		{
-			OutCurvePoints[0].Gradient *= (Curves[IndexCurve0]->GetUMax() - Curves[IndexCurve0]->GetUMin());
-			OutCurvePoints[1].Gradient *= (Curves[IndexCurve1]->GetUMax() - Curves[IndexCurve1]->GetUMin());
+			OutCurvePoints[0].Gradient *= (InCurves[IndexCurve0]->GetUMax() - InCurves[IndexCurve0]->GetUMin());
+			OutCurvePoints[1].Gradient *= (InCurves[IndexCurve1]->GetUMax() - InCurves[IndexCurve1]->GetUMin());
 			if (InDerivativeOrder > 1)
 			{
-				OutCurvePoints[0].Laplacian *= FMath::Square(Curves[IndexCurve0]->GetUMax() - Curves[IndexCurve0]->GetUMin());
-				OutCurvePoints[1].Laplacian *= FMath::Square(Curves[IndexCurve1]->GetUMax() - Curves[IndexCurve1]->GetUMin());
+				OutCurvePoints[0].Laplacian *= FMath::Square(InCurves[IndexCurve0]->GetUMax() - InCurves[IndexCurve0]->GetUMin());
+				OutCurvePoints[1].Laplacian *= FMath::Square(InCurves[IndexCurve1]->GetUMax() - Curves[IndexCurve1]->GetUMin());
 			}
 		}
 	};
diff --git a/Engine/Source/Runtime/Datasmith/CADKernel/Private/CADKernel/Mesh/Structure/Grid.h b/Engine/Source/Runtime/Datasmith/CADKernel/Private/CADKernel/Mesh/Structure/Grid.h
--- a/Engine/Source/Runtime/Datasmith/CADKernel/Private/CADKernel/Mesh/Structure/Grid.h
+++ b/Engine/Source/Runtime/Datasmith/CADKernel/Private/CADKernel/Mesh/Structure/Grid.h
@@ -224,11 +224,16 @@ namespace CADKernel
 		void FindCoordinateIndex(const TArray<double>& InCoordinates, double Coordinate, int32& OutIndex) const
 		{
 			ensureCADKernel(InCoordinates.IsValidIndex(OutIndex));
+
 			while (Coordinate < InCoordinates[OutIndex] && OutIndex > 0)
 			{
 				OutIndex--;
 			}
-			for (; OutIndex + 2 < InCoordinates.Num() && Coordinate > InCoordinates[OutIndex + 1]; ++OutIndex);
+
+			for (; OutIndex + 2 < InCoordinates.Num() && Coordinate > InCoordinates[OutIndex + 1]; ++OutIndex)
+			{
+			}
+
 			ensureCADKernel(InCoordinates.IsValidIndex(OutIndex));
 		};
 
diff --git a/Engine/Source/Runtime/Datasmith/CADKernel/Public/CADKernel/Geo/Curves/PolylineCurve.h b/Engine/Source/Runtime/Datasmith/CADKernel/Public/CADKernel/Geo/Curves/PolylineCurve.h
--- a/Engine/Source/Runtime/Datasmith/CADKernel/Public/CADKernel/Geo/Curves/PolylineCurve.h
+++ b/Engine/Source/Runtime/Datasmith/CADKernel/Public/CADKernel/Geo/Curves/PolylineCurve.h
@@ -65,7 +65,7 @@ namespace CADKernel
 
 		virtual void EvaluateCurvesPoint(double InCoordinate, PointCurveType& OutPoint, int32 InDerivativeOrder = 0) const
 		{
-			Approximator.ApproximatePoint<PointCurveType>(InCoordinate, OutPoint, InDerivativeOrder);
+			template Approximator.ApproximatePoints<PointCurveType>(InCoordinate, OutPoint, InDerivativeOrder);
 		}
 
 		PointType EvaluatePointAt(double InCoordinate) const
@@ -75,7 +75,7 @@ namespace CADKernel
 
 		virtual void EvaluateCurvesPoints(const TArray<double>& InCoordinates, TArray<PointCurveType>& OutPoints, int32 InDerivativeOrder = 0) const
 		{
-			Approximator.ApproximatePoints<PointCurveType>(InCoordinates, OutPoints, InDerivativeOrder);
+			template Approximator.ApproximatePoints<PointCurveType>(InCoordinates, OutPoints, InDerivativeOrder);
 		}
 
 		virtual double ComputeSubLength(const FLinearBoundary& InBoundary) const
diff --git a/Engine/Source/Runtime/Datasmith/CADKernel/Public/CADKernel/Geo/Sampling/PolylineTools.h b/Engine/Source/Runtime/Datasmith/CADKernel/Public/CADKernel/Geo/Sampling/PolylineTools.h
--- a/Engine/Source/Runtime/Datasmith/CADKernel/Public/CADKernel/Geo/Sampling/PolylineTools.h
+++ b/Engine/Source/Runtime/Datasmith/CADKernel/Public/CADKernel/Geo/Sampling/PolylineTools.h
@@ -535,12 +535,12 @@ namespace CADKernel
 			OutCoordinates.Empty(EndIndex - StartIndex + 2);
 			OutPoints.Empty(EndIndex - StartIndex + 2);
 
-			if (FMath::IsNearlyEqual(PolylineCoordinates[StartIndex + 1], InBoundary.Min, KINDA_SMALL_NUMBER))
+			if (FMath::IsNearlyEqual(PolylineCoordinates[StartIndex + 1], InBoundary.Min, (double) KINDA_SMALL_NUMBER))
 			{
 				StartIndex++;
 			}
 
-			if (!FMath::IsNearlyEqual(PolylineCoordinates[EndIndex], InBoundary.Max, KINDA_SMALL_NUMBER))
+			if (!FMath::IsNearlyEqual(PolylineCoordinates[EndIndex], InBoundary.Max, (double) KINDA_SMALL_NUMBER))
 			{
 				EndIndex--;
 			}
diff --git a/Engine/Source/Runtime/Datasmith/CADKernel/Public/CADKernel/Math/Geometry.h b/Engine/Source/Runtime/Datasmith/CADKernel/Public/CADKernel/Math/Geometry.h
--- a/Engine/Source/Runtime/Datasmith/CADKernel/Public/CADKernel/Math/Geometry.h
+++ b/Engine/Source/Runtime/Datasmith/CADKernel/Public/CADKernel/Math/Geometry.h
@@ -243,9 +243,9 @@ namespace CADKernel
 				ensureCADKernel(false);
 			}
 
-			// CenterU  = 1/D * (Cv.|B|² - By.|C|²) = 1/D * CBv ^ SquareNorms 
-			// CenterV  = 1/D * (Bu.|B|² - Cu.|C|²) = -1/D * SquareNorms ^ CBu 
-			// with CBu = (Cu, Bu), CBv = (Cv, Bv), SquareNorms = (| B |², | C |²)
+			// CenterU  = 1/D * (Cv.|B|.|B| - By.|C|.|C|) = 1/D * CBv ^ SquareNorms 
+			// CenterV  = 1/D * (Bu.|B|.|B| - Cu.|C|.|C|) = -1/D * SquareNorms ^ CBu 
+			// with CBu = (Cu, Bu), CBv = (Cv, Bv), SquareNorms = (|B|.|B|, |C|.|C|)
 			FPoint2D CBu(C.U, B.U);
 			FPoint2D CBv(C.V, B.V);
 			FPoint2D SquareNorms(B.SquareLength(), C.SquareLength());
@@ -273,9 +273,9 @@ namespace CADKernel
 				return FPoint2D::ZeroPoint;
 			}
 
-			// CenterU  = 1/D * (Cv.|B|² - Bv.|C|²) = 1/D * CBv ^ SquareNorms 
-			// CenterV  = 1/D * (Bu.|C|² - Cu.|B|²) = 1/D * SquareNorms ^ CBu  
-			// with CBu = (Cu, Bu), CBv = (Cv, Bv), SquareNorms = (| B |², | C |²)
+			// CenterU  = 1/D * (Cv.|B|.|B| - Bv.|C|.|C|) = 1/D * CBv ^ SquareNorms 
+			// CenterV  = 1/D * (Bu.|C|.|C| - Cu.|B|.|B|) = 1/D * SquareNorms ^ CBu  
+			// with CBu = (Cu, Bu), CBv = (Cv, Bv), SquareNorms = (|B|.|B|, |C|.|C|)
 			FPoint2D CBu(C.U, B.U);
 			FPoint2D CBv(C.V, B.V);
 			FPoint2D SquareNorms(C.SquareLength(), B.SquareLength());
@@ -444,11 +444,9 @@ namespace CADKernel
 	template<class PointType>
 	inline double CoordinateOfProjectedPointOnSegment(const PointType& Point, const PointType& InSegmentA, const PointType& InSegmentB, bool bRestrictCoodinateToInside = true)
 	{
-		double SquareLength;
-
 		PointType Segment = InSegmentB - InSegmentA;
 
-		double SquareLength = Segment * Segment;
+		double SquareLength = FMath::Square(Segment);
 
 		if (SquareLength <= 0.0)
 		{
diff --git a/Engine/Source/Runtime/Datasmith/CADKernel/Public/CADKernel/Math/MatrixH.h b/Engine/Source/Runtime/Datasmith/CADKernel/Public/CADKernel/Math/MatrixH.h
--- a/Engine/Source/Runtime/Datasmith/CADKernel/Public/CADKernel/Math/MatrixH.h
+++ b/Engine/Source/Runtime/Datasmith/CADKernel/Public/CADKernel/Math/MatrixH.h
@@ -48,9 +48,9 @@ namespace CADKernel
 			BuildChangeOfCoordinateSystemMatrix(Ox, Oy, Oz, Origin);
 		}
 
-		friend FArchive& operator<<(FArchive& Ar, FMatrixH& Matrix)
+		friend FArchive& operator<<(FArchive& Ar, FMatrixH& InMatrix)
 		{
-			Ar.Serialize(Matrix.Matrix, 16 * sizeof(double));
+			Ar.Serialize(InMatrix.Matrix, 16 * sizeof(double));
 			return Ar;
 		}
 
diff --git a/Engine/Source/Runtime/Datasmith/CADKernel/Public/CADKernel/UI/Progress.h b/Engine/Source/Runtime/Datasmith/CADKernel/Public/CADKernel/UI/Progress.h
--- a/Engine/Source/Runtime/Datasmith/CADKernel/Public/CADKernel/UI/Progress.h
+++ b/Engine/Source/Runtime/Datasmith/CADKernel/Public/CADKernel/UI/Progress.h
@@ -6,7 +6,8 @@
 namespace CADKernel
 {
 	class FProgressBase;
-	
+	class FProgress;
+
 	class CADKERNEL_API FProgressManager
 	{
 		friend class FProgress;
