Fix Issue 65876 Incremental build

diff --git a/Engine/Source/Runtime/Datasmith/CADKernel/Private/CADKernel/Geo/Curves/RestrictionCurve.cpp b/Engine/Source/Runtime/Datasmith/CADKernel/Private/CADKernel/Geo/Curves/RestrictionCurve.cpp
--- a/Engine/Source/Runtime/Datasmith/CADKernel/Private/CADKernel/Geo/Curves/RestrictionCurve.cpp
+++ b/Engine/Source/Runtime/Datasmith/CADKernel/Private/CADKernel/Geo/Curves/RestrictionCurve.cpp
@@ -2,7 +2,7 @@
 
 #include "CADKernel/Geo/Curves/RestrictionCurve.h"
 
-#include "CADKernel/math/SlopeUtils.h"
+#include "CADKernel/Math/SlopeUtils.h"
 
 using namespace CADKernel;
 
diff --git a/Engine/Source/Runtime/Datasmith/CADKernel/Private/CADKernel/Mesh/Criteria/CriteriaGrid.cpp b/Engine/Source/Runtime/Datasmith/CADKernel/Private/CADKernel/Mesh/Criteria/CriteriaGrid.cpp
--- a/Engine/Source/Runtime/Datasmith/CADKernel/Private/CADKernel/Mesh/Criteria/CriteriaGrid.cpp
+++ b/Engine/Source/Runtime/Datasmith/CADKernel/Private/CADKernel/Mesh/Criteria/CriteriaGrid.cpp
@@ -84,7 +84,7 @@ void FCriteriaGrid::ApplyCriteria(const TArray<TSharedPtr<FCriterion>>& Criteria
 			double LengthUV;
 			double SagUV = FCriterion::EvaluateSag(Point, PointUV, PointUVMid, LengthUV);
 
-			for (const TSharedPtr< FCriterion> Criterion : Criteria)
+			for (const TSharedPtr<FCriterion>& Criterion : Criteria)
 			{
 				Criterion->UpdateDelta((GetCoordinate(EIso::IsoU, IndexU + 1) - GetCoordinate(EIso::IsoU, IndexU)), SagU, SagUV, SagV, LengthU, LengthUV, DeltaUMaxArray[IndexU], DeltaUMiniArray[IndexU], SurfaceCurvature[EIso::IsoU]);
 				Criterion->UpdateDelta((GetCoordinate(EIso::IsoV, IndexV + 1) - GetCoordinate(EIso::IsoV, IndexV)), SagV, SagUV, SagU, LengthV, LengthUV, DeltaVMaxArray[IndexV], DeltaVMinArray[IndexV], SurfaceCurvature[EIso::IsoV]);
@@ -151,14 +151,14 @@ void FCriteriaGrid::Display()
 	Close3DDebugSession();
 
 	Open3DDebugSession(TEXT("Loop 3D"));
-	for (const TSharedPtr<FTopologicalLoop> Loop : Surface->GetLoops())
+	for (const TSharedPtr<FTopologicalLoop>& Loop : Surface->GetLoops())
 	{
 		::Display(Loop);
 	}
 	Close3DDebugSession();
 
 	Open3DDebugSession(TEXT("Loop 2D"));
-	for (const TSharedPtr<FTopologicalLoop> Loop : Surface->GetLoops())
+	for (const TSharedPtr<FTopologicalLoop>& Loop : Surface->GetLoops())
 	{
 		::Display2D(Loop);
 	}
diff --git a/Engine/Source/Runtime/Datasmith/CADKernel/Private/CADKernel/Mesh/Meshers/IsoTriangulator.cpp b/Engine/Source/Runtime/Datasmith/CADKernel/Private/CADKernel/Mesh/Meshers/IsoTriangulator.cpp
--- a/Engine/Source/Runtime/Datasmith/CADKernel/Private/CADKernel/Mesh/Meshers/IsoTriangulator.cpp
+++ b/Engine/Source/Runtime/Datasmith/CADKernel/Private/CADKernel/Mesh/Meshers/IsoTriangulator.cpp
@@ -295,11 +295,12 @@ void FIsoTriangulator::FillMeshNodes()
 namespace
 {
 	const double MaxSlopeToBeIso = 0.125;
-	constexpr double LimitValueMin(double Slope)
+	const constexpr double LimitValueMin(double Slope)
 	{
 		return Slope - MaxSlopeToBeIso;
 	}
-	constexpr double LimitValueMax(double Slope)
+
+	const constexpr double LimitValueMax(double Slope)
 	{
 		return Slope + MaxSlopeToBeIso;
 	}
@@ -368,19 +369,19 @@ bool FIsoTriangulator::BuildLoopSegments()
 		{
 			Segment->SetAsIsoU();
 		}
-		if (SegmentSlop < LimitValueMax(2) && SegmentSlop > LimitValueMin(2))
+		if (SegmentSlop < LimitValueMax(2.) && SegmentSlop > LimitValueMin(2.))
 		{
 			Segment->SetAsIsoV();
 		}
-		if (SegmentSlop < LimitValueMax(4) && SegmentSlop > LimitValueMin(4))
+		if (SegmentSlop < LimitValueMax(4.) && SegmentSlop > LimitValueMin(4.))
 		{
 			Segment->SetAsIsoU();
 		}
-		if (SegmentSlop < LimitValueMax(6) && SegmentSlop > LimitValueMin(6))
+		if (SegmentSlop < LimitValueMax(6.) && SegmentSlop > LimitValueMin(6.))
 		{
 			Segment->SetAsIsoV();
 		}
-		if (SegmentSlop > LimitValueMin(8))
+		if (SegmentSlop > LimitValueMin(8.))
 		{
 			Segment->SetAsIsoU();
 		}
diff --git a/Engine/Source/Runtime/Datasmith/CADKernel/Private/CADKernel/Mesh/Meshers/IsoTriangulator/IsoSegment.h b/Engine/Source/Runtime/Datasmith/CADKernel/Private/CADKernel/Mesh/Meshers/IsoTriangulator/IsoSegment.h
--- a/Engine/Source/Runtime/Datasmith/CADKernel/Private/CADKernel/Mesh/Meshers/IsoTriangulator/IsoSegment.h
+++ b/Engine/Source/Runtime/Datasmith/CADKernel/Private/CADKernel/Mesh/Meshers/IsoTriangulator/IsoSegment.h
@@ -3,7 +3,7 @@
 
 #include "CADKernel/Core/Types.h"
 #include "CADKernel/Mesh/Meshers/IsoTriangulator/IsoNode.h"
-#include "CADKernel/Ui/Display.h"
+#include "CADKernel/UI/Display.h"
 
 namespace CADKernel
 {
diff --git a/Engine/Source/Runtime/Datasmith/CADKernel/Public/CADKernel/Core/Parameters.h b/Engine/Source/Runtime/Datasmith/CADKernel/Public/CADKernel/Core/Parameters.h
--- a/Engine/Source/Runtime/Datasmith/CADKernel/Public/CADKernel/Core/Parameters.h
+++ b/Engine/Source/Runtime/Datasmith/CADKernel/Public/CADKernel/Core/Parameters.h
@@ -71,7 +71,9 @@ namespace CADKernel
 		}
 
 		friend FArchive& operator<<(FArchive& Ar, FParameters& Database)
-		{}
+		{
+			return Ar;
+		}
 
 	};
 } 
diff --git a/Engine/Source/Runtime/Datasmith/CADKernel/Public/CADKernel/Geo/Curves/PolylineCurve.h b/Engine/Source/Runtime/Datasmith/CADKernel/Public/CADKernel/Geo/Curves/PolylineCurve.h
--- a/Engine/Source/Runtime/Datasmith/CADKernel/Public/CADKernel/Geo/Curves/PolylineCurve.h
+++ b/Engine/Source/Runtime/Datasmith/CADKernel/Public/CADKernel/Geo/Curves/PolylineCurve.h
@@ -65,7 +65,13 @@ namespace CADKernel
 
 		virtual void EvaluateCurvesPoint(double InCoordinate, PointCurveType& OutPoint, int32 InDerivativeOrder = 0) const
 		{
+>>>> ORIGINAL //UE5/Main/Engine/Source/Runtime/Datasmith/CADKernel/Public/CADKernel/Geo/Curves/PolylineCurve.h#1
+			Approximator.ApproximatePoint<PointCurveType>(InCoordinate, OutPoint, InDerivativeOrder);
+==== THEIRS //UE5/Main/Engine/Source/Runtime/Datasmith/CADKernel/Public/CADKernel/Geo/Curves/PolylineCurve.h#2
 			template Approximator.ApproximatePoints<PointCurveType>(InCoordinate, OutPoint, InDerivativeOrder);
+==== YOURS //David.Lesage_YUL-UE5-Main/Engine/Source/Runtime/Datasmith/CADKernel/Public/CADKernel/Geo/Curves/PolylineCurve.h
+			Approximator.ApproximatePoint(InCoordinate, OutPoint, InDerivativeOrder);
+<<<<
 		}
 
 		PointType EvaluatePointAt(double InCoordinate) const
@@ -75,7 +81,13 @@ namespace CADKernel
 
 		virtual void EvaluateCurvesPoints(const TArray<double>& InCoordinates, TArray<PointCurveType>& OutPoints, int32 InDerivativeOrder = 0) const
 		{
+>>>> ORIGINAL //UE5/Main/Engine/Source/Runtime/Datasmith/CADKernel/Public/CADKernel/Geo/Curves/PolylineCurve.h#1
+			Approximator.ApproximatePoints<PointCurveType>(InCoordinates, OutPoints, InDerivativeOrder);
+==== THEIRS //UE5/Main/Engine/Source/Runtime/Datasmith/CADKernel/Public/CADKernel/Geo/Curves/PolylineCurve.h#2
 			template Approximator.ApproximatePoints<PointCurveType>(InCoordinates, OutPoints, InDerivativeOrder);
+==== YOURS //David.Lesage_YUL-UE5-Main/Engine/Source/Runtime/Datasmith/CADKernel/Public/CADKernel/Geo/Curves/PolylineCurve.h
+			Approximator.ApproximatePoints(InCoordinates, OutPoints, InDerivativeOrder);
+<<<<
 		}
 
 		virtual double ComputeSubLength(const FLinearBoundary& InBoundary) const
