Fix #issue 65895 Incremental Build

diff --git a/Engine/Source/Runtime/Datasmith/CADKernel/Public/CADKernel/Geo/Curves/PolylineCurve.h b/Engine/Source/Runtime/Datasmith/CADKernel/Public/CADKernel/Geo/Curves/PolylineCurve.h
--- a/Engine/Source/Runtime/Datasmith/CADKernel/Public/CADKernel/Geo/Curves/PolylineCurve.h
+++ b/Engine/Source/Runtime/Datasmith/CADKernel/Public/CADKernel/Geo/Curves/PolylineCurve.h
@@ -65,13 +65,7 @@ namespace CADKernel
 
 		virtual void EvaluateCurvesPoint(double InCoordinate, PointCurveType& OutPoint, int32 InDerivativeOrder = 0) const
 		{
->>>> ORIGINAL //UE5/Main/Engine/Source/Runtime/Datasmith/CADKernel/Public/CADKernel/Geo/Curves/PolylineCurve.h#1
-			Approximator.ApproximatePoint<PointCurveType>(InCoordinate, OutPoint, InDerivativeOrder);
-==== THEIRS //UE5/Main/Engine/Source/Runtime/Datasmith/CADKernel/Public/CADKernel/Geo/Curves/PolylineCurve.h#2
-			template Approximator.ApproximatePoints<PointCurveType>(InCoordinate, OutPoint, InDerivativeOrder);
-==== YOURS //David.Lesage_YUL-UE5-Main/Engine/Source/Runtime/Datasmith/CADKernel/Public/CADKernel/Geo/Curves/PolylineCurve.h
 			Approximator.ApproximatePoint(InCoordinate, OutPoint, InDerivativeOrder);
-<<<<
 		}
 
 		PointType EvaluatePointAt(double InCoordinate) const
@@ -81,13 +75,7 @@ namespace CADKernel
 
 		virtual void EvaluateCurvesPoints(const TArray<double>& InCoordinates, TArray<PointCurveType>& OutPoints, int32 InDerivativeOrder = 0) const
 		{
->>>> ORIGINAL //UE5/Main/Engine/Source/Runtime/Datasmith/CADKernel/Public/CADKernel/Geo/Curves/PolylineCurve.h#1
-			Approximator.ApproximatePoints<PointCurveType>(InCoordinates, OutPoints, InDerivativeOrder);
-==== THEIRS //UE5/Main/Engine/Source/Runtime/Datasmith/CADKernel/Public/CADKernel/Geo/Curves/PolylineCurve.h#2
-			template Approximator.ApproximatePoints<PointCurveType>(InCoordinates, OutPoints, InDerivativeOrder);
-==== YOURS //David.Lesage_YUL-UE5-Main/Engine/Source/Runtime/Datasmith/CADKernel/Public/CADKernel/Geo/Curves/PolylineCurve.h
 			Approximator.ApproximatePoints(InCoordinates, OutPoints, InDerivativeOrder);
-<<<<
 		}
 
 		virtual double ComputeSubLength(const FLinearBoundary& InBoundary) const
