Fix issues raised by the linux compiler

diff --git a/Engine/Source/Runtime/Datasmith/CADKernel/Private/CADKernel/Core/Database.cpp b/Engine/Source/Runtime/Datasmith/CADKernel/Private/CADKernel/Core/Database.cpp
--- a/Engine/Source/Runtime/Datasmith/CADKernel/Private/CADKernel/Core/Database.cpp
+++ b/Engine/Source/Runtime/Datasmith/CADKernel/Private/CADKernel/Core/Database.cpp
@@ -12,7 +12,7 @@
 #include "CADKernel/Topo/Model.h"
 #include "CADKernel/Topo/Shell.h"
 #include "CADKernel/Topo/TopologicalVertex.h"
-#include "CADKernel/Ui/Message.h"
+#include "CADKernel/UI/Message.h"
 
 #include "HAL/FileManager.h"
 
diff --git a/Engine/Source/Runtime/Datasmith/CADKernel/Private/CADKernel/Core/Version.h b/Engine/Source/Runtime/Datasmith/CADKernel/Private/CADKernel/Core/Version.h
--- a/Engine/Source/Runtime/Datasmith/CADKernel/Private/CADKernel/Core/Version.h
+++ b/Engine/Source/Runtime/Datasmith/CADKernel/Private/CADKernel/Core/Version.h
@@ -3,11 +3,11 @@
 
 //** ** Version du Toolkit** **//
 #define TOOLKIT_VERSION 7, 0, 4
-#define TOOLKIT_VERSION_ASCII L"7.0.4"
+#define TOOLKIT_VERSION_ASCII TEXT("7.0.4")
 
 //** ** Numero du TAG** **//
 #define TAG_VERSION 2000
-#define TAG_VERSION_ASCII L".2000" 
+#define TAG_VERSION_ASCII TEXT(".2000")
 
 #define RELEASE_DATE	 __DATE__
 
diff --git a/Engine/Source/Runtime/Datasmith/CADKernel/Private/CADKernel/Mesh/Meshers/IsoTriangulator.cpp b/Engine/Source/Runtime/Datasmith/CADKernel/Private/CADKernel/Mesh/Meshers/IsoTriangulator.cpp
--- a/Engine/Source/Runtime/Datasmith/CADKernel/Private/CADKernel/Mesh/Meshers/IsoTriangulator.cpp
+++ b/Engine/Source/Runtime/Datasmith/CADKernel/Private/CADKernel/Mesh/Meshers/IsoTriangulator.cpp
@@ -295,12 +295,12 @@ void FIsoTriangulator::FillMeshNodes()
 namespace
 {
 	const double MaxSlopeToBeIso = 0.125;
-	const constexpr double LimitValueMin(double Slope)
+	const double LimitValueMin(double Slope)
 	{
 		return Slope - MaxSlopeToBeIso;
 	}
 
-	const constexpr double LimitValueMax(double Slope)
+	const double LimitValueMax(double Slope)
 	{
 		return Slope + MaxSlopeToBeIso;
 	}
@@ -441,7 +441,7 @@ void FIsoTriangulator::BuildThinZoneSegments()
 		ThinZoneSegments.Add(&Segment);
 	};
 
-	for (const TSharedPtr<FTopologicalLoop> Loop : Grid.GetFace()->GetLoops())
+	for (const TSharedPtr<FTopologicalLoop>& Loop : Grid.GetFace()->GetLoops())
 	{
 		for (const FOrientedEdge& OrientedEdge : Loop->GetEdges())
 		{
diff --git a/Engine/Source/Runtime/Datasmith/CADKernel/Private/CADKernel/Mesh/Meshers/ParametricMesher.cpp b/Engine/Source/Runtime/Datasmith/CADKernel/Private/CADKernel/Mesh/Meshers/ParametricMesher.cpp
--- a/Engine/Source/Runtime/Datasmith/CADKernel/Private/CADKernel/Mesh/Meshers/ParametricMesher.cpp
+++ b/Engine/Source/Runtime/Datasmith/CADKernel/Private/CADKernel/Mesh/Meshers/ParametricMesher.cpp
@@ -205,7 +205,7 @@ void FParametricMesher::ApplyEdgeCriteria(TSharedRef<FTopologicalEdge> Edge)
 	ActiveEdge->EvaluatePoints(Coordinates, 0, Points3D);
 
 	const TArray<TSharedPtr<FCriterion>>& Criteria = GetMeshModel()->GetCriteria();
-	for (const TSharedPtr<FCriterion> Criterion : Criteria)
+	for (const TSharedPtr<FCriterion>& Criterion : Criteria)
 	{
 		Criterion->ApplyOnEdgeParameters(ActiveEdge, CrossingPointUs, Points3D);
 	}
diff --git a/Engine/Source/Runtime/Datasmith/CADKernel/Private/CADKernel/Mesh/Structure/FaceMesh.cpp b/Engine/Source/Runtime/Datasmith/CADKernel/Private/CADKernel/Mesh/Structure/FaceMesh.cpp
--- a/Engine/Source/Runtime/Datasmith/CADKernel/Private/CADKernel/Mesh/Structure/FaceMesh.cpp
+++ b/Engine/Source/Runtime/Datasmith/CADKernel/Private/CADKernel/Mesh/Structure/FaceMesh.cpp
@@ -14,7 +14,7 @@ using namespace CADKernel;
 void FFaceMesh::GetNodeIdToCoordinates(TMap<int32, const FPoint*>& NodeIdToCoordinates) const
 {
 	const TSharedRef<FTopologicalFace> Face = StaticCastSharedRef<FTopologicalFace>(GetGeometricEntity());
-	for (const TSharedPtr<FTopologicalLoop> Loop : Face->GetLoops())
+	for (const TSharedPtr<FTopologicalLoop>& Loop : Face->GetLoops())
 	{
 		const TArray<FOrientedEdge>& Edges = Loop->GetEdges();
 
diff --git a/Engine/Source/Runtime/Datasmith/CADKernel/Private/CADKernel/Mesh/Structure/Grid.cpp b/Engine/Source/Runtime/Datasmith/CADKernel/Private/CADKernel/Mesh/Structure/Grid.cpp
--- a/Engine/Source/Runtime/Datasmith/CADKernel/Private/CADKernel/Mesh/Structure/Grid.cpp
+++ b/Engine/Source/Runtime/Datasmith/CADKernel/Private/CADKernel/Mesh/Structure/Grid.cpp
@@ -1080,7 +1080,7 @@ namespace CADKernel
 		F3DDebugSession _(TEXT("GetLoopMesh"));
 #endif
 
-		for (const TSharedPtr<FTopologicalLoop> Loop : Face->GetLoops())
+		for (const TSharedPtr<FTopologicalLoop>& Loop : Face->GetLoops())
 		{
 			int32 LoopNodeCount = 0;
 
diff --git a/Engine/Source/Runtime/Datasmith/CADKernel/Private/CADKernel/Topo/TopologicalFace.cpp b/Engine/Source/Runtime/Datasmith/CADKernel/Private/CADKernel/Topo/TopologicalFace.cpp
--- a/Engine/Source/Runtime/Datasmith/CADKernel/Private/CADKernel/Topo/TopologicalFace.cpp
+++ b/Engine/Source/Runtime/Datasmith/CADKernel/Private/CADKernel/Topo/TopologicalFace.cpp
@@ -238,7 +238,7 @@ const void FTopologicalFace::Get2DLoopSampling(TArray<TArray<FPoint2D>>& LoopSam
 {
 	LoopSamplings.Empty(GetLoops().Num());
 
-	for (const TSharedPtr<FTopologicalLoop> Loop : GetLoops())
+	for (const TSharedPtr<FTopologicalLoop>& Loop : GetLoops())
 	{
 		TArray<FPoint2D>& LoopSampling2D = LoopSamplings.Emplace_GetRef();
 		Loop->Get2DSampling(LoopSampling2D);
diff --git a/Engine/Source/Runtime/Datasmith/CADKernel/Private/CADKernel/UI/Display.cpp b/Engine/Source/Runtime/Datasmith/CADKernel/Private/CADKernel/UI/Display.cpp
--- a/Engine/Source/Runtime/Datasmith/CADKernel/Private/CADKernel/UI/Display.cpp
+++ b/Engine/Source/Runtime/Datasmith/CADKernel/Private/CADKernel/UI/Display.cpp
@@ -14,7 +14,7 @@
 #include "CADKernel/Geo/Surfaces/BezierSurface.h"
 #include "CADKernel/Geo/Surfaces/NURBSSurface.h"
 #include "CADKernel/Geo/Surfaces/Surface.h"
-#include "CADKernel/Math/AABB.h"
+#include "CADKernel/Math/Aabb.h"
 #include "CADKernel/Math/Boundary.h"
 #include "CADKernel/Math/Plane.h"
 #include "CADKernel/Mesh/Structure/EdgeMesh.h"
@@ -1019,7 +1019,7 @@ namespace CADKernel
 
 	void DisplayLoop(const TSharedPtr<FTopologicalFace>& Surface)
 	{
-		for (const TSharedPtr<FTopologicalLoop> Loop : Surface->GetLoops())
+		for (const TSharedPtr<FTopologicalLoop>& Loop : Surface->GetLoops())
 		{
 			Display(Loop);
 		}
diff --git a/Engine/Source/Runtime/Datasmith/CADKernel/Public/CADKernel/Core/System.h b/Engine/Source/Runtime/Datasmith/CADKernel/Public/CADKernel/Core/System.h
--- a/Engine/Source/Runtime/Datasmith/CADKernel/Public/CADKernel/Core/System.h
+++ b/Engine/Source/Runtime/Datasmith/CADKernel/Public/CADKernel/Core/System.h
@@ -19,7 +19,7 @@ namespace CADKernel
 
 	protected:
 
-		FString ProductName = L"CADKernel";
+		FString ProductName = FString(TEXT("CADKernel"));
 
 		TSharedRef<FKernelParameters> Parameters;
 
diff --git a/Engine/Source/Runtime/Datasmith/CADKernel/Public/CADKernel/Geo/Curves/PolylineCurve.h b/Engine/Source/Runtime/Datasmith/CADKernel/Public/CADKernel/Geo/Curves/PolylineCurve.h
--- a/Engine/Source/Runtime/Datasmith/CADKernel/Public/CADKernel/Geo/Curves/PolylineCurve.h
+++ b/Engine/Source/Runtime/Datasmith/CADKernel/Public/CADKernel/Geo/Curves/PolylineCurve.h
@@ -21,7 +21,7 @@ namespace CADKernel
 		TPolylineApproximator<PointType> Approximator;
 
 		TPolylineCurve(const double InTolerance, const TArray<PointType>& InPoints, const TArray<double>& InCoordinates, int8 InDimension)
-			: FCurve(InTolerance, Dimension)
+			: FCurve(InTolerance, InDimension)
 			, Approximator(this->Coordinates, this->Points)
 		{
 			this->Coordinates = InCoordinates;
@@ -31,7 +31,7 @@ namespace CADKernel
 		}
 
 		TPolylineCurve(const double InTolerance, const TArray<PointType>& InPoints, int8 InDimension)
-			: FCurve(InTolerance, Dimension)
+			: FCurve(InTolerance, InDimension)
 			, TPolyline<PointType>(InPoints)
 			, Approximator(this->Coordinates, this->Points)
 		{
diff --git a/Engine/Source/Runtime/Datasmith/CADKernel/Public/CADKernel/Math/Geometry.h b/Engine/Source/Runtime/Datasmith/CADKernel/Public/CADKernel/Math/Geometry.h
--- a/Engine/Source/Runtime/Datasmith/CADKernel/Public/CADKernel/Math/Geometry.h
+++ b/Engine/Source/Runtime/Datasmith/CADKernel/Public/CADKernel/Math/Geometry.h
@@ -4,7 +4,7 @@
 #include "CADKernel/Core/Types.h"
 
 #include "CADKernel/Geo/GeoEnum.h"
-#include "CADKernel/Math/AABB.h"
+#include "CADKernel/Math/Aabb.h"
 #include "CADKernel/Math/MathConst.h"
 #include "CADKernel/Math/MatrixH.h"
 #include "CADKernel/Math/Point.h"
diff --git a/Engine/Source/Runtime/Datasmith/CADKernel/Public/CADKernel/Topo/TopologicalFace.h b/Engine/Source/Runtime/Datasmith/CADKernel/Public/CADKernel/Topo/TopologicalFace.h
--- a/Engine/Source/Runtime/Datasmith/CADKernel/Public/CADKernel/Topo/TopologicalFace.h
+++ b/Engine/Source/Runtime/Datasmith/CADKernel/Public/CADKernel/Topo/TopologicalFace.h
@@ -4,7 +4,7 @@
 #include "CADKernel/Core/MetadataDictionary.h"
 #include "CADKernel/Geo/GeoPoint.h"
 #include "CADKernel/Geo/Surfaces/Surface.h"
-#include "CADKernel/Math/AABB.h"
+#include "CADKernel/Math/Aabb.h"
 #include "CADKernel/Math/Curvature.h"
 #include "CADKernel/Topo/TopologicalEntity.h"
 #include "CADKernel/Topo/TopologicalLoop.h"
diff --git a/Engine/Source/Runtime/Datasmith/CADKernel/Public/CADKernel/UI/Message.h b/Engine/Source/Runtime/Datasmith/CADKernel/Public/CADKernel/UI/Message.h
--- a/Engine/Source/Runtime/Datasmith/CADKernel/Public/CADKernel/UI/Message.h
+++ b/Engine/Source/Runtime/Datasmith/CADKernel/Public/CADKernel/UI/Message.h
@@ -67,15 +67,15 @@ namespace CADKernel
 
 #define ERROR_FUNCTION_CALL_NOT_EXPECTED \
 { \
-	FMessage::Printf(Log, L"CALL of \" %s \" NOT EXPECTED at line %d of the file %s\n", __func__, __LINE__, __FILE__); \
+	FMessage::Printf(Log, TEXT("CALL of \" %s \" NOT EXPECTED at line %d of the file %s\n"), __func__, __LINE__, __FILE__); \
 }
 
 #define ERROR_NOT_EXPECTED \
 { \
-	FMessage::Printf(Log, L"Error not expected in \" %S \" at line %d of the file %S\n", __func__, __LINE__, __FILE__); \
+	FMessage::Printf(Log, TEXT("Error not expected in \" %S \" at line %d of the file %S\n"), __func__, __LINE__, __FILE__); \
 }
 
 #define NOT_IMPLEMENTED \
 { \
-	FMessage::Printf(Log, L"The function \" %s \" at line %d of the file %s is not implemented", __func__, __LINE__, __FILE__); \
+	FMessage::Printf(Log, TEXT("The function \" %s \" at line %d of the file %s is not implemented"), __func__, __LINE__, __FILE__); \
 }
