Build fix
Remove pragma warnings that don't seem needed

diff --git a/Engine/Source/Runtime/Datasmith/CADKernel/Private/CADKernel/Math/BSpline.cpp b/Engine/Source/Runtime/Datasmith/CADKernel/Private/CADKernel/Math/BSpline.cpp
--- a/Engine/Source/Runtime/Datasmith/CADKernel/Private/CADKernel/Math/BSpline.cpp
+++ b/Engine/Source/Runtime/Datasmith/CADKernel/Private/CADKernel/Math/BSpline.cpp
@@ -1,9 +1,5 @@
 // Copyright Epic Games, Inc. All Rights Reserved.
 
-#ifdef _WINDOWS
-#pragma warning(disable:4996)
-#endif
-
 #include "CADKernel/Math/BSpline.h"
 
 #include "CADKernel/Geo/Curves/NURBSCurve.h"
diff --git a/Engine/Source/Runtime/Datasmith/CADKernel/Private/CADKernel/UI/Message.cpp b/Engine/Source/Runtime/Datasmith/CADKernel/Private/CADKernel/UI/Message.cpp
--- a/Engine/Source/Runtime/Datasmith/CADKernel/Private/CADKernel/UI/Message.cpp
+++ b/Engine/Source/Runtime/Datasmith/CADKernel/Private/CADKernel/UI/Message.cpp
@@ -9,9 +9,6 @@
 
 #include "Misc/VarArgs.h"
 
-#pragma warning(disable:4996) // unsafe sprintf
-
-
 namespace CADKernel
 {
 	const char* VerboseLevelConstNames[] = {
