CADKernel Source Code v0.01: fix Issue #66196: Compile warnings

diff --git a/Engine/Source/Runtime/Datasmith/CADKernel/Private/CADKernel/Core/Group.h b/Engine/Source/Runtime/Datasmith/CADKernel/Private/CADKernel/Core/Group.h
--- a/Engine/Source/Runtime/Datasmith/CADKernel/Private/CADKernel/Core/Group.h
+++ b/Engine/Source/Runtime/Datasmith/CADKernel/Private/CADKernel/Core/Group.h
@@ -31,11 +31,11 @@ namespace CADKernel
 		{
 		}
 
-		FGroup(TArray<TSharedPtr<FEntity>>& Entities)
+		FGroup(TArray<TSharedPtr<FEntity>>& InEntities)
 			: FEntity()
 			, Origin(EGroupOrigin::Unknown)
 		{
-			Entities.Append(Entities);
+			Entities.Append(InEntities);
 		}
 
 		FGroup(FCADKernelArchive& Archive)
diff --git a/Engine/Source/Runtime/Datasmith/CADKernel/Private/CADKernel/Math/BSpline.cpp b/Engine/Source/Runtime/Datasmith/CADKernel/Private/CADKernel/Math/BSpline.cpp
--- a/Engine/Source/Runtime/Datasmith/CADKernel/Private/CADKernel/Math/BSpline.cpp
+++ b/Engine/Source/Runtime/Datasmith/CADKernel/Private/CADKernel/Math/BSpline.cpp
@@ -478,7 +478,7 @@ namespace CADKernel
 			isoU ==> on cree nb_pt_u iso_u d'abord puis calcul de nb_pt_v points
 			Choisir le nombre Minimal d'evaluations*/
 
-			if ((1 + derivee) * (nb_pt_v * PoleUNum * VDegre * (VDegre + 1)) + (1 + 2 * derivee) * (nb_pt_u * nb_pt_v * UDegre * (UDegre + 1)) <(1 + derivee) * (nb_pt_u * PoleVNum * UDegre * (UDegre + 1)) + (1 + 2 * derivee) * (nb_pt_u * nb_pt_v * VDegre * (VDegre + 1)))
+			if (((1 + derivee) * (nb_pt_v * PoleUNum * VDegre * (VDegre + 1)) + (1 + 2 * derivee) * (nb_pt_u * nb_pt_v * UDegre * (UDegre + 1))) < ((1 + derivee) * (nb_pt_u * PoleVNum * UDegre * (UDegre + 1)) + (1 + 2 * derivee) * (nb_pt_u * nb_pt_v * VDegre * (VDegre + 1))))
 			{
 				/* choisir la strategie isoV*/
 				nb_pt1 = nb_pt_v;
@@ -516,7 +516,7 @@ namespace CADKernel
 				std::copy(tab_v, tab_v + nb_pt1, tab_t1.GetData());
 				std::copy(tab_u, tab_u + nb_pt2, tab_t2.GetData());
 			}
-			else if (IsoType == EIso::IsoU)
+			else // (IsoType == EIso::IsoU)
 			{
 				std::copy(tab_u, tab_u + nb_pt1, tab_t1.GetData());
 				std::copy(tab_v, tab_v + nb_pt2, tab_t2.GetData());
@@ -536,7 +536,7 @@ namespace CADKernel
 					}
 				}
 			}
-			else if (IsoType == EIso::IsoU)
+			else // (IsoType == EIso::IsoU)
 			{
 				for (IndexV = 0; IndexV <nb_pole1; IndexV++)
 				{
@@ -667,7 +667,7 @@ namespace CADKernel
 						}
 					}
 				}
-				else if (IsoType == EIso::IsoU)
+				else // (IsoType == EIso::IsoU)
 				{
 					for (IndexV = 0; IndexV <nb_pt2; IndexV++)
 					{
@@ -699,7 +699,7 @@ namespace CADKernel
 						}
 					}
 				}
-				else if (IsoType == EIso::IsoU)
+				else // (IsoType == EIso::IsoU)
 				{
 					for (IndexV = 0; IndexV <nb_pt2; IndexV++)
 					{
diff --git a/Engine/Source/Runtime/Datasmith/CADKernel/Private/CADKernel/Mesh/Meshers/ParametricMesher.cpp b/Engine/Source/Runtime/Datasmith/CADKernel/Private/CADKernel/Mesh/Meshers/ParametricMesher.cpp
--- a/Engine/Source/Runtime/Datasmith/CADKernel/Private/CADKernel/Mesh/Meshers/ParametricMesher.cpp
+++ b/Engine/Source/Runtime/Datasmith/CADKernel/Private/CADKernel/Mesh/Meshers/ParametricMesher.cpp
@@ -725,10 +725,10 @@ void FParametricMesher::IsolateQuadFace(TArray<FCostToFace>& QuadSurfaces, TArra
 			double LocalMinCurvature;
 			double LocalMaxCurvature;
 			Sort(Face->GetCurvature(EIso::IsoU).Max, Face->GetCurvature(EIso::IsoV).Max, LocalMinCurvature, LocalMaxCurvature);
-			if (LocalMaxCurvature > LocalMinCurvature)
+			if (LocalMaxCurvature > ConstMinCurvature)
 			{
 				QuadSurfaces.Emplace(LocalMaxCurvature, Face.ToSharedRef());
-				if (LocalMinCurvature > LocalMinCurvature)
+				if (LocalMinCurvature > ConstMinCurvature)
 				{
 					QuadSurfaces.Emplace(LocalMinCurvature, Face.ToSharedRef());
 				}
@@ -1502,7 +1502,10 @@ void FParametricMesher::MeshThinZoneSide(const FThinZoneSide& Side)
 			const FEdgeSegment* ClosedSegment = EdgeSegment->GetClosedSegment();
 			if (ClosedSegment == nullptr)
 			{
+#ifdef CADKERNEL_DEV
 				Wait();
+#endif
+				continue;
 			}
 			double OppositeCuttingPointSegmentU;
 			FPoint OppositeCuttingPoint3D = ClosedSegment->ProjectPoint(CuttingPoint3D, OppositeCuttingPointSegmentU);
diff --git a/Engine/Source/Runtime/Datasmith/CADKernel/Private/CADKernel/Mesh/Structure/Grid.cpp b/Engine/Source/Runtime/Datasmith/CADKernel/Private/CADKernel/Mesh/Structure/Grid.cpp
--- a/Engine/Source/Runtime/Datasmith/CADKernel/Private/CADKernel/Mesh/Structure/Grid.cpp
+++ b/Engine/Source/Runtime/Datasmith/CADKernel/Private/CADKernel/Mesh/Structure/Grid.cpp
@@ -1040,7 +1040,7 @@ namespace CADKernel
 		}
 		else
 		{
-			SwapAndScale(EIso::IsoU);
+			SwapAndScale(EIso::IsoV);
 		}
 	}
 
diff --git a/Engine/Source/Runtime/Datasmith/CADKernel/Private/CADKernel/Mesh/Structure/ThinZone2D.cpp b/Engine/Source/Runtime/Datasmith/CADKernel/Private/CADKernel/Mesh/Structure/ThinZone2D.cpp
--- a/Engine/Source/Runtime/Datasmith/CADKernel/Private/CADKernel/Mesh/Structure/ThinZone2D.cpp
+++ b/Engine/Source/Runtime/Datasmith/CADKernel/Private/CADKernel/Mesh/Structure/ThinZone2D.cpp
@@ -746,8 +746,8 @@ void FThinZone2DFinder::BuildLoopSegments()
 	SortedLoopSegments = LoopSegments;
 }
 
-FThinZoneSide::FThinZoneSide(FThinZoneSide& InFrontSide, const TArray<FEdgeSegment*>& InSegments, bool bInIsFirstSide)
-	: FrontSide(InFrontSide)
+FThinZoneSide::FThinZoneSide(FThinZoneSide* InFrontSide, const TArray<FEdgeSegment*>& InSegments, bool bInIsFirstSide)
+	: FrontSide(*InFrontSide)
 {
 	Segments = InSegments;
 	Length = 0;
diff --git a/Engine/Source/Runtime/Datasmith/CADKernel/Private/CADKernel/Mesh/Structure/ThinZone2D.h b/Engine/Source/Runtime/Datasmith/CADKernel/Private/CADKernel/Mesh/Structure/ThinZone2D.h
--- a/Engine/Source/Runtime/Datasmith/CADKernel/Private/CADKernel/Mesh/Structure/ThinZone2D.h
+++ b/Engine/Source/Runtime/Datasmith/CADKernel/Private/CADKernel/Mesh/Structure/ThinZone2D.h
@@ -53,7 +53,7 @@ namespace CADKernel
 
 	public:
 
-		FThinZoneSide(FThinZoneSide& InFrontSide, const TArray<FEdgeSegment*>& InSegments, bool bInIsFirstSide);
+		FThinZoneSide(FThinZoneSide* InFrontSide, const TArray<FEdgeSegment*>& InSegments, bool bInIsFirstSide);
 		virtual ~FThinZoneSide() = default;
 
 		bool IsClosed() const
@@ -131,8 +131,8 @@ namespace CADKernel
 	public:
 
 		FThinZone2D(const TArray<FEdgeSegment*>& InFirstSideSegments, bool bInIsClosed1, const TArray<FEdgeSegment*>& InSecondSideSegments, bool bInIsClosed2, double ZoneThickness)
-			: FirstSide(SecondSide, InFirstSideSegments, true)
-			, SecondSide(FirstSide, InSecondSideSegments, false)
+			: FirstSide(&SecondSide, InFirstSideSegments, true)
+			, SecondSide(&FirstSide, InSecondSideSegments, false)
 			, Category(EThinZone2DType::Undefined)
 			, Thickness(ZoneThickness)
 		{
@@ -239,7 +239,7 @@ namespace CADKernel
 
 		bool HasClosedSide() const 
 		{
-			constexpr EHaveStates ATLeastOneClosed = EHaveStates::FirstSideClosed | EHaveStates::FirstSideClosed;
+			constexpr EHaveStates ATLeastOneClosed = EHaveStates::FirstSideClosed | EHaveStates::SecondSideClosed;
 			return EnumHasAnyFlags(States, ATLeastOneClosed);
 		}
 
diff --git a/Engine/Source/Runtime/Datasmith/CADKernel/Private/CADKernel/Topo/Joiner.cpp b/Engine/Source/Runtime/Datasmith/CADKernel/Private/CADKernel/Topo/Joiner.cpp
--- a/Engine/Source/Runtime/Datasmith/CADKernel/Private/CADKernel/Topo/Joiner.cpp
+++ b/Engine/Source/Runtime/Datasmith/CADKernel/Private/CADKernel/Topo/Joiner.cpp
@@ -21,6 +21,11 @@ FJoiner::FJoiner(TArray<TSharedPtr<FTopologicalFace>>& InSurfaces, double InTole
 {
 }
 
+void FJoiner::JoinSurfaces()
+{
+
+}
+
 void FJoiner::JoinSurfaces(bool bProcessOnlyBorderEdges, bool bProcessOnlyNonManifoldEdges)
 {
 	if (!(bProcessOnlyBorderEdges || bProcessOnlyNonManifoldEdges))
diff --git a/Engine/Source/Runtime/Datasmith/CADKernel/Private/CADKernel/Topo/Joiner.h b/Engine/Source/Runtime/Datasmith/CADKernel/Private/CADKernel/Topo/Joiner.h
--- a/Engine/Source/Runtime/Datasmith/CADKernel/Private/CADKernel/Topo/Joiner.h
+++ b/Engine/Source/Runtime/Datasmith/CADKernel/Private/CADKernel/Topo/Joiner.h
@@ -25,7 +25,8 @@ namespace CADKernel
 
 		FJoiner(TArray<TSharedPtr<FTopologicalFace>>& surfaces, double Tolerance);
 
-		void JoinSurfaces(bool bProcessOnlyBorderEdges = false, bool bProcessOnlyNonManifoldEdges = false);
+		void JoinSurfaces();
+		void JoinSurfaces(bool bProcessOnlyBorderEdges, bool bProcessOnlyNonManifoldEdges);
 
 		void FixCollapsedEdges();
 
diff --git a/Engine/Source/Runtime/Datasmith/CADKernel/Private/CADKernel/Topo/Model.cpp b/Engine/Source/Runtime/Datasmith/CADKernel/Private/CADKernel/Topo/Model.cpp
--- a/Engine/Source/Runtime/Datasmith/CADKernel/Private/CADKernel/Topo/Model.cpp
+++ b/Engine/Source/Runtime/Datasmith/CADKernel/Private/CADKernel/Topo/Model.cpp
@@ -7,6 +7,7 @@
 #include "CADKernel/Topo/TopologicalEdge.h"
 #include "CADKernel/Topo/TopologicalVertex.h"
 #include "CADKernel/Topo/TopologicalFace.h"
+#include "CADKernel/Topo/Joiner.h"
 
 
 using namespace CADKernel;
@@ -139,3 +140,6 @@ void FModel::Add(TSharedPtr<FBody> Body)
 {
 	Bodies.Add(Body);
 }
+
+
+
diff --git a/Engine/Source/Runtime/Datasmith/CADKernel/Private/CADKernel/Topo/Shell.cpp b/Engine/Source/Runtime/Datasmith/CADKernel/Private/CADKernel/Topo/Shell.cpp
--- a/Engine/Source/Runtime/Datasmith/CADKernel/Private/CADKernel/Topo/Shell.cpp
+++ b/Engine/Source/Runtime/Datasmith/CADKernel/Private/CADKernel/Topo/Shell.cpp
@@ -98,3 +98,116 @@ void FShell::SpreadBodyOrientation()
 	}
 }
 
+void FShell::CheckTopology(TArray<FFaceSubset>& SubShells)
+{
+	// Processed1 : Surfaces added in CandidateSurfacesForMesh
+
+	int32 TopologicalFaceCount = FaceCount();
+	// Is closed ?
+	// Is one shell ?
+
+	int32 ProcessFaceCount = 0;
+
+	TArray<TSharedPtr<FTopologicalFace>> Front;
+	TFunction<void(const TSharedPtr<FTopologicalFace>&, FFaceSubset&)> GetNeighboringFaces = [&](const TSharedPtr<FTopologicalFace>& Face, FFaceSubset& Shell)
+	{
+		for (const TSharedPtr<FTopologicalLoop>& Loop : Face->GetLoops())
+		{
+			for (const FOrientedEdge& OrientedEdge : Loop->GetEdges())
+			{
+				const TSharedPtr<FTopologicalEdge>& Edge = OrientedEdge.Entity;
+				if (Edge->HasMarker1())
+				{
+					continue;
+				}
+				Edge->SetMarker1();
+
+				if (Edge->GetTwinsEntityCount() == 1)
+				{
+					if (!Edge->IsDegenerated())
+					{
+						Shell.BorderEdgeCount++;
+					}
+					continue;
+				}
+
+				if (Edge->GetTwinsEntityCount() > 2)
+				{
+					Shell.NonManifoldEdgeCount++;
+				}
+
+				for (TWeakPtr<FTopologicalEdge> WeakEdge : Edge->GetTwinsEntities())
+				{
+					TSharedPtr<FTopologicalEdge> NextEdge = WeakEdge.Pin();
+					if (NextEdge->HasMarker1())
+					{
+						continue;
+					}
+					NextEdge->SetMarker1();
+
+					TSharedPtr<FTopologicalFace> NextFace = NextEdge->GetFace();
+					if (!NextFace.IsValid())
+					{
+						continue;
+					}
+
+					if (NextFace->HasMarker1())
+					{
+						continue;
+					}
+					NextFace->SetMarker1();
+					Front.Add(NextFace);
+				}
+			}
+		}
+	};
+	
+	TFunction<void(FFaceSubset&)> SpreadFront = [&](FFaceSubset& Shell)
+	{
+		while (Front.Num())
+		{
+			TSharedPtr<FTopologicalFace> Face = Front.Pop();
+			Shell.Faces.Add(Face);
+			GetNeighboringFaces(Face, Shell);
+		}
+	};
+
+	for (const FOrientedFace& OrientedFace : GetFaces())
+	{
+		if (OrientedFace.Entity->HasMarker1())
+		{
+			continue;
+		}
+	
+		FFaceSubset& Shell = SubShells.Emplace_GetRef();
+		Shell.Faces.Reserve(TopologicalFaceCount - ProcessFaceCount);
+		Front.Empty(TopologicalFaceCount);
+
+		const TSharedPtr<FTopologicalFace>& Face = OrientedFace.Entity;
+
+		Front.Empty(TopologicalFaceCount);
+		Face->SetMarker1();
+		Front.Add(Face);
+		SpreadFront(Shell);
+		ProcessFaceCount += Shell.Faces.Num();
+
+		if (ProcessFaceCount == TopologicalFaceCount)
+		{
+			break;
+		}
+	}
+
+	// reset Marker
+	for (const FOrientedFace& OrientedFace : GetFaces())
+	{
+		const TSharedPtr<FTopologicalFace>& Face = OrientedFace.Entity;
+		Face->ResetMarkers();
+		for (const TSharedPtr<FTopologicalLoop>& Loop : Face->GetLoops())
+		{
+			for (const FOrientedEdge& OrientedEdge : Loop->GetEdges())
+			{
+				OrientedEdge.Entity->ResetMarker1();
+			}
+		}
+	}
+}
diff --git a/Engine/Source/Runtime/Datasmith/CADKernel/Public/CADKernel/Core/Metadata.h b/Engine/Source/Runtime/Datasmith/CADKernel/Public/CADKernel/Core/Metadata.h
deleted file mode 100644
diff --git a/Engine/Source/Runtime/Datasmith/CADKernel/Public/CADKernel/Core/MetadataDictionary.h b/Engine/Source/Runtime/Datasmith/CADKernel/Public/CADKernel/Core/MetadataDictionary.h
--- a/Engine/Source/Runtime/Datasmith/CADKernel/Public/CADKernel/Core/MetadataDictionary.h
+++ b/Engine/Source/Runtime/Datasmith/CADKernel/Public/CADKernel/Core/MetadataDictionary.h
@@ -2,8 +2,7 @@
 #pragma once
 
 #include "CADKernel/Core/Types.h"
-
-#include "CADKernel/Core/Metadata.h"
+#include "CADKernel/Math/MathConst.h"
 
 namespace CADKernel
 {
@@ -11,63 +10,41 @@ namespace CADKernel
 	{
 
 	private:
-		TMap<FString, TSharedPtr<FMetadata>> MetadataMap;
 
-		void AddValue(const FString& Name, TSharedPtr<FMetadata>& Attribute)
-		{
-			MetadataMap.Emplace(Name, Attribute);
-		}
+		FString Name;
+		int32 HostId = 0;
+		uint32 ColorId = 0;
+		uint32 MaterialId = 0;
+		int32 LayerId = -1; 
 
 	public:
 
 		FMetadataDictionary() = default;
 
-		template<typename ValueType>
-		void Add(const FString& Name, const ValueType& Value)
+		void SetHostId(const int32 InHostId)
 		{
-			TSharedRef<FMetadata> NewAttribute = FMetadata::Create(Value);
-			SetOrAdd(Name, NewAttribute);
+			HostId = InHostId;
 		}
 
-		void AddLayer(const FString& Name, const int32 LayerId, const FString& LayerName, const int32 LayerFlag)
+		void SetLayer(const int32 InLayerId)
 		{
-			TSharedRef<FMetadata> NewAttribute = FMetadata::CreateLayer(LayerId, LayerName, LayerFlag);
-			SetOrAdd(Name, NewAttribute);
+			LayerId = InLayerId;
 		}
 
-		void AddRGBColor(const FString& Name, const double RedValue, const double GreenValue, const double BlueValue)
+		void SetName(const FString& InName)
 		{
-			TSharedRef<FMetadata> NewAttribute = FMetadata::CreateRGBColor(RedValue, GreenValue, BlueValue);
-			SetOrAdd(Name, NewAttribute);
+			Name = InName;
 		}
 
-		const TSharedPtr<FMetadata> Get(const FString& Name)
+		void SetColorId(const uint32& InColorId)
 		{
-			TSharedPtr<FMetadata>* Attribute = MetadataMap.Find(Name);
-			return Attribute == nullptr ? TSharedPtr<FMetadata>() : *Attribute;
+			ColorId = InColorId;
 		}
 
-		int32 Count() const
+		void SetMaterialId(const uint32& InMaterialId)
 		{
-			return MetadataMap.Num();
+			MaterialId = InMaterialId;
 		}
-
-	private:
-
-		void SetOrAdd(const FString Name, const TSharedPtr<FMetadata> NewAttribute)
-		{
-			TSharedPtr<FMetadata>* OldAttribute = MetadataMap.Find(Name);
-			if (OldAttribute != nullptr)
-			{
-				(*OldAttribute).Reset();
-				*OldAttribute = NewAttribute;
-			}
-			else
-			{
-				MetadataMap.Add(Name, NewAttribute);
-			}
-		}
-
 	};
 }
 
diff --git a/Engine/Source/Runtime/Datasmith/CADKernel/Public/CADKernel/Math/Aabb.h b/Engine/Source/Runtime/Datasmith/CADKernel/Public/CADKernel/Math/Aabb.h
--- a/Engine/Source/Runtime/Datasmith/CADKernel/Public/CADKernel/Math/Aabb.h
+++ b/Engine/Source/Runtime/Datasmith/CADKernel/Public/CADKernel/Math/Aabb.h
@@ -116,6 +116,22 @@ namespace CADKernel
 			);
 		}
 
+		FAABB& operator+= (const double* Point)
+		{
+			for (int32 Index = 0; Index < 3; Index++)
+			{
+				if (Point[Index] < MinCorner[Index])
+				{
+					MinCorner[Index] = Point[Index];
+				}
+				if (Point[Index] > MaxCorner[Index])
+				{
+					MaxCorner[Index] = Point[Index];
+				}
+			}
+			return *this;
+		}
+
 		FAABB& operator+= (const FPoint& Point)
 		{
 			for (int32 Index = 0; Index < 3; Index++)
diff --git a/Engine/Source/Runtime/Datasmith/CADKernel/Public/CADKernel/Mesh/Meshers/ParametricMesher.h b/Engine/Source/Runtime/Datasmith/CADKernel/Public/CADKernel/Mesh/Meshers/ParametricMesher.h
--- a/Engine/Source/Runtime/Datasmith/CADKernel/Public/CADKernel/Mesh/Meshers/ParametricMesher.h
+++ b/Engine/Source/Runtime/Datasmith/CADKernel/Public/CADKernel/Mesh/Meshers/ParametricMesher.h
@@ -103,7 +103,7 @@ namespace CADKernel
 		/**
 		 * Limit of flatness of quad face
 		 */
-		const double MinCurvature = 0.001;
+		const double ConstMinCurvature = 0.001;
 
 		TSharedRef<FModelMesh> MeshModel;
 		TSharedRef<FMesherParameters> Parameters;
diff --git a/Engine/Source/Runtime/Datasmith/CADKernel/Public/CADKernel/Topo/Shell.h b/Engine/Source/Runtime/Datasmith/CADKernel/Public/CADKernel/Topo/Shell.h
--- a/Engine/Source/Runtime/Datasmith/CADKernel/Public/CADKernel/Topo/Shell.h
+++ b/Engine/Source/Runtime/Datasmith/CADKernel/Public/CADKernel/Topo/Shell.h
@@ -14,6 +14,8 @@ namespace CADKernel
 	class FBody;
 	class FTopologicalFace;
 
+	struct FFaceSubset;
+
 	class CADKERNEL_API FOrientedFace : public TOrientedEntity<FTopologicalFace>
 	{
 	public:
@@ -116,6 +118,8 @@ namespace CADKernel
 
 		virtual void SpreadBodyOrientation() override;
 
+		void CheckTopology(TArray<FFaceSubset>& SubShells);
+
 		bool IsInner() const
 		{
 			return ((States & EHaveStates::IsInner) == EHaveStates::IsInner);
diff --git a/Engine/Source/Runtime/Datasmith/CADKernel/Public/CADKernel/Topo/TopologicalEdge.h b/Engine/Source/Runtime/Datasmith/CADKernel/Public/CADKernel/Topo/TopologicalEdge.h
--- a/Engine/Source/Runtime/Datasmith/CADKernel/Public/CADKernel/Topo/TopologicalEdge.h
+++ b/Engine/Source/Runtime/Datasmith/CADKernel/Public/CADKernel/Topo/TopologicalEdge.h
@@ -564,10 +564,13 @@ namespace CADKernel
 		// ======   State Functions   ======
 
 		/**
-		 * A Degenerated Edge is used to close 2D boundary in case of degenerated surface to ensureCADKernel a closed boundary
+		 * Important note: A Degenerated Edge is used to close 2D boundary in case of degenerated surface to ensureCADKernel a closed boundary
 		 * Specific process is done for the mesh of this kind of surface
 		 */
-		//virtual bool IsDegenerated() const override;
+		virtual bool IsDegenerated() const override
+		{
+			return FHaveStates::IsDegenerated();
+		}
 
 		bool IsThinPeak() const
 		{
diff --git a/Engine/Source/Runtime/Datasmith/CADKernel/Public/CADKernel/Topo/TopologicalFace.h b/Engine/Source/Runtime/Datasmith/CADKernel/Public/CADKernel/Topo/TopologicalFace.h
--- a/Engine/Source/Runtime/Datasmith/CADKernel/Public/CADKernel/Topo/TopologicalFace.h
+++ b/Engine/Source/Runtime/Datasmith/CADKernel/Public/CADKernel/Topo/TopologicalFace.h
@@ -491,4 +491,12 @@ namespace CADKernel
 
 	};
 
+	struct FFaceSubset
+	{
+		TArray<TSharedPtr<FTopologicalFace>> Faces;
+		int32 BorderEdgeCount = 0;
+		int32 NonManifoldEdgeCount = 0;
+	};
+
+
 } // namespace CADKernel
diff --git a/Engine/Source/Runtime/Datasmith/CADKernel/Public/CADKernel/Utils/IndexOfCoordinateFinder.h b/Engine/Source/Runtime/Datasmith/CADKernel/Public/CADKernel/Utils/IndexOfCoordinateFinder.h
--- a/Engine/Source/Runtime/Datasmith/CADKernel/Public/CADKernel/Utils/IndexOfCoordinateFinder.h
+++ b/Engine/Source/Runtime/Datasmith/CADKernel/Public/CADKernel/Utils/IndexOfCoordinateFinder.h
@@ -104,6 +104,11 @@ namespace CADKernel
 
 		virtual int32 Find(const double InCoordinate) override
 		{
+			if (StartUpper < 0)
+			{
+				StartUpper = Coordinates.Num() - 2;
+			}
+
 			if ((StartUpper - StartLower < 1) || (InCoordinate < Coordinates[StartLower + 1]))
 			{
 				return StartLower;
@@ -116,7 +121,7 @@ namespace CADKernel
 			}
 
 			int32 LowerBound = StartLower;
-			int32 UpperBound = StartUpper < 0 ? Coordinates.Num() - 1 : StartUpper;
+			int32 UpperBound = StartUpper;
 
 			while (LowerBound <= UpperBound)
 			{
