CADKernel Source Code v0.02

diff --git a/Engine/Source/Runtime/Datasmith/CADKernel/Private/CADKernel/Mesh/Meshers/IsoTriangulator.cpp b/Engine/Source/Runtime/Datasmith/CADKernel/Private/CADKernel/Mesh/Meshers/IsoTriangulator.cpp
--- a/Engine/Source/Runtime/Datasmith/CADKernel/Private/CADKernel/Mesh/Meshers/IsoTriangulator.cpp
+++ b/Engine/Source/Runtime/Datasmith/CADKernel/Private/CADKernel/Mesh/Meshers/IsoTriangulator.cpp
@@ -4,6 +4,10 @@
 
 #include "CADKernel/Core/Chrono.h"
 #include "CADKernel/Math/MathConst.h"
+#include "CADKernel/Math/Point.h"
+#include "CADKernel/Mesh/Meshers/BowyerWatsonTriangulator.h"
+#include "CADKernel/Mesh/Meshers/IsoTriangulator/IsoCell.h"
+#include "CADKernel/Mesh/Meshers/IsoTriangulator/IsoNode.h"
 #include "CADKernel/Mesh/Meshers/IsoTriangulator/IsoSegment.h"
 #include "CADKernel/Mesh/Structure/Grid.h"
 #include "CADKernel/Mesh/Structure/EdgeMesh.h"
@@ -12,11 +16,12 @@
 #include "CADKernel/Topo/TopologicalFace.h"
 #include "CADKernel/Topo/TopologicalFace.h"
 #include "CADKernel/UI/Display.h"
+#include "CADKernel/Utils/ArrayUtils.h"
 
 using namespace CADKernel;
 
 //#define ADD_TRIANGLE_2D
-#define DEBUG_ISOTRIANGULATOR
+//#define DEBUG_ISOTRIANGULATOR
 
 FIsoTriangulator::FIsoTriangulator(const FGrid& InGrid, TSharedRef<FFaceMesh> EntityMesh)
 	: Grid(InGrid)
@@ -24,12 +29,15 @@
 	, LoopSegmentsIntersectionTool(InGrid)
 	, InnerSegmentsIntersectionTool(InGrid)
 	, InnerToLoopSegmentsIntersectionTool(InGrid)
+	, InnerToOuterSegmentsIntersectionTool(InGrid)
 {
 	FinalInnerSegments.Reserve(3 * Grid.InnerNodesCount());
 	IndexOfLowerLeftInnerNodeSurroundingALoop.Reserve(Grid.GetLoopCount());
 
 #ifdef DEBUG_ISOTRIANGULATOR
-	if (Grid.GetFace()->GetId() == SURFACE_TO_DEBUG)
+#ifdef DEBUG_ONLY_SURFACE_TO_DEBUG
+	if (InGrid.GetFace()->GetId() == FaceToDebug)
+#endif
 	{
 		bDisplay = true;
 	}
@@ -38,12 +46,12 @@
 
 bool FIsoTriangulator::Triangulate()
 {
-	EGridSpace DisplaySpace = EGridSpace::Default2D;
+	EGridSpace DisplaySpace = EGridSpace::UniformScaled;
 #ifdef DEBUG_ISOTRIANGULATOR
 	if (bDisplay)
 	{
 		Open3DDebugSession(FString::Printf(TEXT("Triangulate %d"), Grid.GetFace()->GetId()));
-		Wait(Grid.GetFace()->GetId() == SurfaceToDebug);
+		//Wait(Grid.GetFace()->GetId() == FaceToDebug);
 	}
 #endif
 
@@ -67,7 +75,7 @@
 	}
 
 #ifdef CADKERNEL_DEV
-	Display(DisplaySpace, TEXT("FIsoTrianguler::LoopSegments"), LoopSegments, false, false, EVisuProperty::GreenCurve);
+	Display(DisplaySpace, TEXT("FIsoTrianguler::LoopSegments"), LoopSegments, false, false, EVisuProperty::OrangeCurve);
 #endif
 
 	BuildThinZoneSegments();
@@ -79,6 +87,7 @@
 	BuildInnerSegments();
 #ifdef CADKERNEL_DEV
 	Display(DisplaySpace, TEXT("FIsoTrianguler::FinalInnerSegments"), FinalInnerSegments, false, false, EVisuProperty::BlueCurve);
+	InnerToOuterSegmentsIntersectionTool.Display(TEXT("FIsoTrianguler::IntersectionTool InnerToOutter"));
 	Chronos.TriangulateDuration1 = FChrono::Elapse(StartTime);
 #endif
 
@@ -90,15 +99,42 @@
 	InnerSegmentsIntersectionTool.Display(TEXT("FIsoTrianguler::IntersectionTool Inner"));
 #endif
 
-	SortLoopNodes();
 #ifdef CADKERNEL_DEV
 	Chronos.TriangulateDuration2 = FChrono::Elapse(StartTime);
 #endif
 
 	// =============================================================================================================
+	// 	   For each cell
+	// 	      - Connect loops together and to cell vertices
+	// 	           - Find subset of node of each loop
+	// 	           - build Delaunay connection
+	// 	           - find the shortest segment to connect each connected loop by Delaunay
+	// =============================================================================================================
+	ConnectCellLoops();
+#ifdef CADKERNEL_DEV
+	Display(DisplaySpace, TEXT("FIsoTrianguler::Final Iso ToLink Inner To Loops"), FinalToLoops, false, false, EVisuProperty::YellowCurve);
+#endif
+/*
+	return true;
+
+
+
+
+
+
+
+
+
+
+
+
+	// =============================================================================================================
 	// 	   Find first segments of the final mesh
 	// =============================================================================================================
 
+	//FindCandidateSegmentsToLinkInnerAndLoop();
+
+	SortLoopNodes();
 	// use of SortedLoopNodes but specific order
 	FindIsoSegmentToLinkLoopToLoop();
 #ifdef CADKERNEL_DEV
@@ -123,7 +159,7 @@
 	// Delaunay
 	if (Grid.GetLoops3D().Num())
 	{
-		Delaunay();
+		//ConnectCellLoopsByNeighborhood();
 #ifdef CADKERNEL_DEV
 		Display(DisplaySpace, TEXT("FIsoTrianguler::Candidate Segments by Delaunay"), CandidateLoopToLoopSegments, false);
 #endif
@@ -156,7 +192,7 @@
 #ifdef CADKERNEL_DEV
 	Chronos.TriangulateDuration4 = FChrono::Elapse(StartTime);
 #endif
-
+*/
 	// =============================================================================================================
 	// Make the final tessellation 
 	// =============================================================================================================
@@ -175,11 +211,11 @@
 		Close3DDebugSession();
 		Close3DDebugSession();
 	}
-	Wait(bDisplay);
+	//Wait(bDisplay);
 #endif
 
 #ifdef CADKERNEL_DEV
-	Chronos.PrintTimeElapse();
+	//Chronos.PrintTimeElapse();
 #endif
 
 	return true;
@@ -189,20 +225,20 @@
 {
 	FTimePoint StartTime = FChrono::Now();
 
-	int32 NbNodes = 0;
+	LoopNodeCount = 0;
 	for (const TArray<FPoint2D>& Loop : Grid.GetLoops2D(EGridSpace::Default2D))
 	{
-		NbNodes += (int32)Loop.Num();
+		LoopNodeCount += (int32)Loop.Num();
 	}
 	LoopStartIndex.Reserve(Grid.GetLoops2D(EGridSpace::Default2D).Num());
-	LoopNodes.Reserve((int32)(NbNodes * 1.2 + 5)); // reserve more in case it need to create complementary nodes
+	LoopNodes.Reserve((int32)(LoopNodeCount * 1.2 + 5)); // reserve more in case it need to create complementary nodes
 
 	// Loop nodes
 	int32 FaceIndex = 0;
 	int32 LoopIndex = 0;
 	for (const TArray<FPoint2D>& Loop : Grid.GetLoops2D(EGridSpace::Default2D))
 	{
-		LoopStartIndex.Add(LoopNodes.Num());
+		LoopStartIndex.Add(LoopNodeCount);
 		const TArray<int32>& LoopIds = Grid.GetNodeIdsOfFaceLoops()[LoopIndex];
 		FLoopNode* NextNode = nullptr;
 		FLoopNode* FirstNode = &LoopNodes.Emplace_GetRef(LoopIndex, 0, FaceIndex++, LoopIds[0]);
@@ -240,8 +276,8 @@
 
 void FIsoTriangulator::FillMeshNodes()
 {
-	int32 TriangleNum = 50 + (int32)((2 * InnerNodes.Num() + LoopNodes.Num()) * 1.1);
-	Mesh->Init(TriangleNum, InnerNodes.Num() + LoopNodes.Num());
+	int32 TriangleNum = 50 + (int32)((2 * InnerNodeCount + LoopNodeCount) * 1.1);
+	Mesh->Init(TriangleNum, InnerNodeCount + LoopNodeCount);
 
 	TArray<FPoint>& InnerNodeCoordinates = Mesh->GetNodeCoordinates();
 	InnerNodeCoordinates.Reserve(InnerNodeCount);
@@ -259,7 +295,7 @@
 		Node.OffsetId(StartId);
 	}
 
-	Mesh->VerticesGlobalIndex.SetNum(InnerNodes.Num() + LoopNodes.Num());
+	Mesh->VerticesGlobalIndex.SetNum(InnerNodeCount + LoopNodeCount);
 	int32 Index = 0;
 	for (FLoopNode& Node : LoopNodes)
 	{
@@ -310,7 +346,6 @@
 {
 	FTimePoint StartTime = FChrono::Now();
 
-	int32 LoopNodeCount = LoopNodes.Num();
 	LoopSegments.Reserve(LoopNodeCount);
 
 	int32 LoopIndex = 0;
@@ -320,7 +355,7 @@
 	{
 		FTimePoint StartCheckIntersectionTime = FChrono::Now();
 		FIntersectionSegmentTool IntersectionTool(Grid);
-		IntersectionTool.AddSegment(Segments);
+		IntersectionTool.AddSegments(Segments);
 		IntersectionTool.Sort();
 		for (const FIsoSegment* Segment : Segments)
 		{
@@ -399,7 +434,7 @@
 		}
 	}
 
-	LoopSegmentsIntersectionTool.AddSegment(LoopSegments);
+	LoopSegmentsIntersectionTool.AddSegments(LoopSegments);
 	LoopSegmentsIntersectionTool.Sort();
 
 #ifdef CADKERNEL_DEV
@@ -410,7 +445,7 @@
 
 void FIsoTriangulator::BuildThinZoneSegments()
 {
-	ThinZoneSegments.Reserve((int32)(0.6 * (double)LoopNodes.Num()));
+	ThinZoneSegments.Reserve((int32)(0.6 * (double) LoopNodeCount));
 
 	TMap<int32, FLoopNode*> IndexToNode;
 	for (FLoopNode& Node : LoopNodes)
@@ -467,7 +502,7 @@
 		}
 	}
 
-	LoopSegmentsIntersectionTool.AddSegment(ThinZoneSegments);
+	LoopSegmentsIntersectionTool.AddSegments(ThinZoneSegments);
 	LoopSegmentsIntersectionTool.Sort();
 }
 
@@ -482,24 +517,21 @@
 	int32 NumU = Grid.GetCuttingCount(EIso::IsoU);
 	int32 NumV = Grid.GetCuttingCount(EIso::IsoV);
 
-	LoopSegmentsIntersectionTool.Reserve(InnerSegmentsIntersectionTool.Num());
+	LoopSegmentsIntersectionTool.Reserve(InnerSegmentsIntersectionTool.Count());
 
 	TFunction<void(const int32, const int32, const ESegmentType)> BuildSegmentIfValid = [this](const int32 IndexNode1, const int32 IndexNode2, const ESegmentType InType)
 	{
-		if (!Grid.IsNodeInsideFace(IndexNode1))
+		if (!Grid.IsNodeInsideFace(IndexNode1) || !Grid.IsNodeInsideFace(IndexNode2))
 		{
+			InnerToOuterSegmentsIntersectionTool.AddSegment(Grid.GetInner2DPoint(EGridSpace::UniformScaled, IndexNode1), Grid.GetInner2DPoint(EGridSpace::UniformScaled, IndexNode2));
 			return;
 		}
 
-		if (!Grid.IsNodeInsideFace(IndexNode2))
-		{
-			return;
-		}
-
 		if (Grid.IsNodeCloseToLoop(IndexNode1) && Grid.IsNodeCloseToLoop(IndexNode2))
 		{
-			if (LoopSegmentsIntersectionTool.DoesIntersect(Grid.GetInner2DPoint(EGridSpace::Default2D, IndexNode1), Grid.GetInner2DPoint(EGridSpace::Default2D, IndexNode2)))
+			if (LoopSegmentsIntersectionTool.DoesIntersect(Grid.GetInner2DPoint(EGridSpace::UniformScaled, IndexNode1), Grid.GetInner2DPoint(EGridSpace::UniformScaled, IndexNode2)))
 			{
+				InnerToOuterSegmentsIntersectionTool.AddSegment(Grid.GetInner2DPoint(EGridSpace::UniformScaled, IndexNode1), Grid.GetInner2DPoint(EGridSpace::UniformScaled, IndexNode2));
 				return;
 			}
 		}
@@ -653,11 +685,11 @@
 	// So a fraction of the strip's width defined a good tolerance around a given point. 
 	constexpr double ToleranceCoefficent = 1. / 12.; // Why 12 ? ;o)
 
-	const TArray<double>& IsoUCoordinates = Grid.GetCuttingCoordinatesAlongIso(EIso::IsoU);
-	const TArray<double>& IsoVCoordinates = Grid.GetCuttingCoordinatesAlongIso(EIso::IsoV);
+	const TArray<double>& IsoUCoordinates = Grid.GetUniformCuttingCoordinatesAlongIso(EIso::IsoU);
+	const TArray<double>& IsoVCoordinates = Grid.GetUniformCuttingCoordinatesAlongIso(EIso::IsoV);
 
 
-	// Warning, Min delta is compute in EGridSpace::Default2D
+	// Warning, Min delta is compute in EGridSpace::Uniform
 	TFunction<double(const TArray<double>&)> GetMinDelta = [](const TArray<double>& IsoCoordinates)
 	{
 		double MinDelta = HUGE_VAL;
@@ -685,25 +717,27 @@
 			return;
 		}
 
-		if (InnerSegmentsIntersectionTool.DoesIntersect(Coordinate1, Coordinate2))
-		{
-			return;
-		}
-
-		if (LoopSegmentsIntersectionTool.DoesIntersect(*Node1, *Node2))
-		{
-			return;
-		}
-
 		// Is Outside and not too flat at Node1
+		ensureCADKernel(Node1->GetLoopIndex() > 0);
 		const double FlatAngle = 0.1;
-		if (!IsPointPIsInsideSectorABC(Node1->GetPreviousNode().Get2DPoint(EGridSpace::Default2D, Grid), Coordinate1, Node1->GetNextNode().Get2DPoint(EGridSpace::Default2D, Grid), Coordinate2, FlatAngle))
+		if(Node1->IsSegmentBeInsideFace(Coordinate2, Grid, FlatAngle))
 		{
 			return;
 		}
 
 		// Is Outside and not too flat at Node2
-		if (!IsPointPIsInsideSectorABC(Node2->GetPreviousNode().Get2DPoint(EGridSpace::Default2D, Grid), Coordinate2, Node2->GetNextNode().Get2DPoint(EGridSpace::Default2D, Grid), Coordinate1, FlatAngle))
+		ensureCADKernel(Node2->GetLoopIndex() > 0);
+		if (Node2->IsSegmentBeInsideFace(Coordinate1, Grid, FlatAngle))
 		{
 			return;
 		}
+
+		if (InnerSegmentsIntersectionTool.DoesIntersect(Coordinate1, Coordinate2))
+		{
+			return;
+		}
+
+		if (LoopSegmentsIntersectionTool.DoesIntersect(*Node1, *Node2))
+		{
+			return;
+		}
@@ -763,7 +797,7 @@
 				continue;
 			}
 
-			const FPoint2D& LoopPoint = LoopNode->Get2DPoint(EGridSpace::Default2D, Grid);
+			const FPoint2D& LoopPoint = LoopNode->Get2DPoint(EGridSpace::UniformScaled, Grid);
 
 #ifdef DEBUG_FIND_ISOSEGMENT_TO_LINK_LOOP_TO_LOOP
 			F3DDebugSession G(FString::Printf(TEXT("Iso %d Index %d"), InAxe, Index));
@@ -777,7 +811,7 @@
 #ifdef DEBUG_FIND_ISOSEGMENT_TO_LINK_LOOP_TO_LOOP
 			Display(EGridSpace::UniformScaled, *NextNode, 0, EVisuProperty::RedPoint);
 #endif
-			const FPoint2D& NextPoint = NextNode->Get2DPoint(EGridSpace::Default2D, Grid);
+			const FPoint2D& NextPoint = NextNode->Get2DPoint(EGridSpace::UniformScaled, Grid);
 			if (FMath::IsNearlyEqual(NextPoint[InAxe], LoopPoint[InAxe], ToleranceU))
 			{
 				// the nodes are nearly iso aligned, are they nearly in the same V Stip
@@ -830,7 +864,7 @@
 		}
 	};
 
-	int32 InitNum = InnerSegmentsIntersectionTool.Num() + LoopSegmentsIntersectionTool.Num();
+	int32 InitNum = InnerSegmentsIntersectionTool.Count() + LoopSegmentsIntersectionTool.Count();
 	FinalToLoops.Reserve(InitNum);
 	InnerToLoopSegmentsIntersectionTool.Reserve(InitNum);
 
@@ -845,21 +879,21 @@
 		double DeltaFactor = FMath::Min((GetMinDelta(IsoUCoordinates) / 1000.), (GetMinDelta(IsoVCoordinates) / 1000.));
 
 		// Bounds and GetMinDelta are defined in EGridSpace::Default2D,
-		const FSurfacicBoundary& Bounds = Grid.GetFace()->GetBoundary();
-		double UMin = Bounds.UVBoundaries[EIso::IsoU].Min;
-		double VMin = Bounds.UVBoundaries[EIso::IsoV].Min;
+		//const FSurfacicBoundary& Bounds = Grid.GetFace()->GetBoundary();
+		double UMin = Grid.GetUniformCuttingCoordinates()[EIso::IsoU][0]; // Bounds.UVBoundaries[EIso::IsoU].Min;
+		double VMin = Grid.GetUniformCuttingCoordinates()[EIso::IsoV][0]; // Bounds.UVBoundaries[EIso::IsoV].Min;
 		Algo::Sort(SortedLoopNodesAlong, [&](const FLoopNode* LoopNode1, const FLoopNode* LoopNode2)
 			{
-				const FPoint2D& Node1Coordinates = LoopNode1->Get2DPoint(EGridSpace::Default2D, Grid);
-				const FPoint2D& Node2Coordinates = LoopNode2->Get2DPoint(EGridSpace::Default2D, Grid);
+				const FPoint2D& Node1Coordinates = LoopNode1->Get2DPoint(EGridSpace::UniformScaled, Grid);
+				const FPoint2D& Node2Coordinates = LoopNode2->Get2DPoint(EGridSpace::UniformScaled, Grid);
 				return (Node1Coordinates.U + (Node1Coordinates.V - VMin) * DeltaFactor) < (Node2Coordinates.U + (Node2Coordinates.V - VMin) * DeltaFactor);
 			});
 		FindIsoSegmentAlong(IsoU, IsoUCoordinates, IsoVCoordinates);
 
 		Algo::Sort(SortedLoopNodesAlong, [&](const FLoopNode* LoopNode1, const FLoopNode* LoopNode2)
 			{
-				const FPoint2D& Node1Coordinates = LoopNode1->Get2DPoint(EGridSpace::Default2D, Grid);
-				const FPoint2D& Node2Coordinates = LoopNode2->Get2DPoint(EGridSpace::Default2D, Grid);
+				const FPoint2D& Node1Coordinates = LoopNode1->Get2DPoint(EGridSpace::UniformScaled, Grid);
+				const FPoint2D& Node2Coordinates = LoopNode2->Get2DPoint(EGridSpace::UniformScaled, Grid);
 				return (Node1Coordinates.V + (Node1Coordinates.U - UMin) * DeltaFactor) < (Node2Coordinates.V + (Node2Coordinates.U - UMin) * DeltaFactor);
 			});
 		FindIsoSegmentAlong(IsoV, IsoVCoordinates, IsoUCoordinates);
@@ -869,6 +903,148 @@
 #endif
 }
 
+void FIsoTriangulator::FindSegmentToLinkOuterLoopNodes(FCell& Cell)
+{
+	//F3DDebugSession _(TEXT("FindSegmentToLinkOuterLoopNodes"));
+
+	TArray<FLoopNode*>& OuterLoop = Cell.Loops[0];
+	int32 NodeCount = OuterLoop.Num();
+
+	int SubdivisionCount =  Cell.OuterLoopSubdivision.Num();
+	for (int32 Andex = 0; Andex < SubdivisionCount - 1; ++Andex)
+	{
+		TArray<FLoopNode*>& SubLoopA = Cell.OuterLoopSubdivision[Andex];
+		for (int32 Bndex = Andex + 1; Bndex < SubdivisionCount; ++Bndex)
+		{
+			TArray<FLoopNode*>& SubLoopB = Cell.OuterLoopSubdivision[Bndex];
+			TryToConnectTwoLoopsWithShortestSegment(Cell, SubLoopA, SubLoopB);
+		}
+	}
+	Cell.SelectSegmentInCandidateSegments(IsoSegmentFactory);
+}
+
+void FIsoTriangulator::AddSementToLinkOuterLoopExtremities(FCell& Cell)
+{
+	//F3DDebugSession _(TEXT("AddSementToLinkOuterLoopExtremities"));
+
+	if(!Cell.bHasOuterLoop)
+	{
+		return;
+	}
+	FLoopNode* NodeA = Cell.Loops[0][0];
+	const FPoint2D& ACoordinates = NodeA->Get2DPoint(EGridSpace::UniformScaled, Grid);
+	FLoopNode* NodeB = Cell.Loops[0].Last();
+	const FPoint2D& BCoordinates = NodeB->Get2DPoint(EGridSpace::UniformScaled, Grid);
+	TryToCreateSegment(Cell, NodeA, ACoordinates, NodeB, BCoordinates, 0.1);
+	Cell.SelectSegmentInCandidateSegments(IsoSegmentFactory);
+}
+
+void FIsoTriangulator::FindIsoSegmentToLinkOuterLoopNodes(FCell& Cell)
+{
+	//F3DDebugSession _(TEXT("FindIsoSegmentToLinkOuterLoopNodes"));
+
+	TArray<FLoopNode*>& OuterLoop = Cell.Loops[0];
+	int32 NodeCount = OuterLoop.Num();
+
+	int SubdivisionCount = Cell.OuterLoopSubdivision.Num();
+	for (int32 Andex = 0; Andex < SubdivisionCount - 1; ++Andex)
+	{
+		TArray<FLoopNode*>& SubLoopA = Cell.OuterLoopSubdivision[Andex];
+		for (int32 Bndex = Andex + 1; Bndex < SubdivisionCount; ++Bndex)
+		{
+			TArray<FLoopNode*>& SubLoopB = Cell.OuterLoopSubdivision[Bndex];
+			TryToConnectTwoLoopsWithTheMostIsoSegment(Cell, SubLoopA, SubLoopB);
+		}
+	}
+	Cell.SelectSegmentInCandidateSegments(IsoSegmentFactory);
+}
+
+void FIsoTriangulator::FindIsoSegmentToLinkOuterLoopNodes2(FCell& Cell)
+{
+	//F3DDebugSession _(TEXT("FindIsoSegmentToLinkOuterLoopNode"));
+
+	// This coefficient is to defined the tolerance on coordinates according the iso strip...
+	// With some surface, the parameterization speed can vary enormously depending on the point of the surface.
+	// A good information is the width of iso strip around a point. Indeed, strips have the optimal width to respect meshing criteria.
+	// So a fraction of the strip's width defined a good tolerance around a given point. 
+	constexpr double ToleranceCoefficent = 1. / 10.;
+
+	const TArray<double>& IsoUCoordinates = Grid.GetUniformCuttingCoordinatesAlongIso(EIso::IsoU);
+	const TArray<double>& IsoVCoordinates = Grid.GetUniformCuttingCoordinatesAlongIso(EIso::IsoV);
+
+	int32 UIndex = Cell.Id % Grid.GetCuttingCount(EIso::IsoU);
+	int32 VIndex = Cell.Id / Grid.GetCuttingCount(EIso::IsoU);
+	const double ToleranceU = (IsoUCoordinates[UIndex + 1] - IsoUCoordinates[UIndex]) * ToleranceCoefficent;
+	const double MinULength = 3 * ToleranceU;
+	const double ToleranceV = (IsoVCoordinates[VIndex + 1] - IsoVCoordinates[VIndex]) * ToleranceCoefficent;
+	const double MinVLength = 3 * ToleranceV;
+
+
+	TArray<FLoopNode*>& OuterLoop = Cell.Loops[0];
+	int32 NodeCount = OuterLoop.Num();
+	for (int32 Index = 0; Index < NodeCount - 1; ++Index)
+	{
+		FLoopNode* StartNode = OuterLoop[Index];
+		const FPoint2D& StartPoint = StartNode->Get2DPoint(EGridSpace::UniformScaled, Grid); 
+		for (int32 Endex = Index + 1; Endex < NodeCount; ++Endex)
+		{
+			FLoopNode* EndNode = OuterLoop[Endex];
+			if(&StartNode->GetNextNode() == EndNode || &StartNode->GetPreviousNode() == EndNode)
+			{
+				continue;
+			}
+			const FPoint2D& EndPoint = EndNode->Get2DPoint(EGridSpace::UniformScaled, Grid);
+
+			FPoint2D SegmentLength = Abs(EndPoint - StartPoint);
+
+			if(SegmentLength.U < ToleranceU && SegmentLength.V > MinVLength)
+			{
+				//DisplaySegment(StartPoint, EndPoint, 0, EVisuProperty::BlueCurve);
+				TryToCreateSegment(Cell, StartNode, StartPoint, EndNode, EndPoint, 1);
+			}
+			else if (SegmentLength.V < ToleranceV && SegmentLength.U > MinULength)
+			{
+				//DisplaySegment(StartPoint, EndPoint, 0, EVisuProperty::GreenCurve);
+				TryToCreateSegment(Cell, StartNode, StartPoint, EndNode, EndPoint, 1);
+			}
+		}
+	}
+}
+
+void FIsoTriangulator::LastChanceToCreateSegmentInCell(FCell& Cell)
+{
+	if (Cell.LoopNodeCount > 5 && Cell.CandidateSegments.Num() > Cell.Loops.Num())
+	{
+		return;
+	}
+
+	//F3DDebugSession _(TEXT("LastChanceToCreateSegmentInCell"));
+
+	TArray<FLoopNode*> Nodes;
+	Nodes.Reserve(Cell.LoopNodeCount);
+
+	for(const TArray<FLoopNode*>& Loop : Cell.Loops)
+	{
+		Nodes.Append(Loop);
+	}
+
+	for (int32 Index = 0; Index < Cell.LoopNodeCount - 1; ++Index)
+	{
+		FLoopNode* StartNode = Nodes[Index];
+		const FPoint2D& StartPoint = StartNode->Get2DPoint(EGridSpace::UniformScaled, Grid);
+		for (int32 Endex = Index + 1; Endex < Cell.LoopNodeCount; ++Endex)
+		{
+			FLoopNode* EndNode = Nodes[Endex];
+			if (&StartNode->GetNextNode() == EndNode || &StartNode->GetPreviousNode() == EndNode)
+			{
+				continue;
+			}
+			const FPoint2D& EndPoint = EndNode->Get2DPoint(EGridSpace::UniformScaled, Grid);
+			TryToCreateSegment(Cell, StartNode, StartPoint, EndNode, EndPoint, 1);
+		}
+	}
+}
+
 void FIsoTriangulator::FindIsoSegmentToLinkInnerToLoop()
 {
 	FTimePoint StartTime = FChrono::Now();
@@ -876,10 +1052,10 @@
 	int32 IndexU = 0;
 	int32 IndexV = 0;
 
-	const TArray<double>& IsoUCoordinates = Grid.GetCuttingCoordinatesAlongIso(EIso::IsoU);
-	const TArray<double>& IsoVCoordinates = Grid.GetCuttingCoordinatesAlongIso(EIso::IsoV);
+	const TArray<double>& IsoUCoordinates = Grid.GetUniformCuttingCoordinatesAlongIso(EIso::IsoU);
+	const TArray<double>& IsoVCoordinates = Grid.GetUniformCuttingCoordinatesAlongIso(EIso::IsoV);
 
-	int32 InitNum = InnerSegmentsIntersectionTool.Num() + LoopSegmentsIntersectionTool.Num();
+	int32 InitNum = InnerSegmentsIntersectionTool.Count() + LoopSegmentsIntersectionTool.Count();
 	FinalToLoops.Reserve(InitNum);
 	InnerToLoopSegmentsIntersectionTool.Reserve(InitNum);
 
@@ -942,7 +1118,7 @@
 			return IsoGlobalIndex;
 		}
 
-		if (LoopSegmentsIntersectionTool.DoesIntersect(IsoLoopNode, Grid.GetInner2DPoint(EGridSpace::Default2D, IsoGlobalIndex)))
+		if (LoopSegmentsIntersectionTool.DoesIntersect(IsoLoopNode, Grid.GetInner2DPoint(EGridSpace::UniformScaled, IsoGlobalIndex)))
 		{
 			return IsoGlobalIndex;
 		}
@@ -956,7 +1132,7 @@
 		return IsoGlobalIndex;
 	};
 
-	const TArray<TArray<FPoint2D>>& Loops = Grid.GetLoops2D(EGridSpace::Default2D);
+	const TArray<TArray<FPoint2D>>& Loops = Grid.GetLoops2D(EGridSpace::UniformScaled);
 
 	// Segment smaller than the size of the grid is build
 	for (FLoopNode& LoopNode : LoopNodes)
@@ -1009,9 +1185,289 @@
 #endif
 }
 
+// =============================================================================================================
+// 	   For each cell
+// 	      - Connect loops together and to cell vertices
+// 	           - Find subset of node of each loop
+// 	           - build Delaunay connection
+// 	           - find the shortest segment to connect each connected loop by Delaunay
+// =============================================================================================================
+void FIsoTriangulator::ConnectCellLoops()
+{
+	TArray<FCell> Cells;
+	FindCellContainingBoundaryNodes(Cells);
+	//DisplayCells(Cells);
+
+	FinalToLoops.Reserve(LoopNodeCount + InnerNodeCount);
+
+	//F3DDebugSession _(TEXT("Cells"));
+	for (FCell& Cell : Cells)
+	{
+		//F3DDebugSession _(FString::Printf(TEXT("Cell %d"), Cell.Id));
+		ConnectCellLoopsByNeighborhood(Cell);
+		if (Cell.bHasOuterLoop)
+		{
+			FindIsoSegmentToLinkOuterLoopNodes(Cell);
+			if(Cell.CandidateSegments.Num() == 0)
+			{
+				FindSegmentToLinkOuterLoopNodes(Cell);
+			}
+		}
+		ConnectCellCornerToInnerLoop(Cell);
+
+		FinalToLoops.Append(Cell.FinalSegments);
+	}
+}
+
+// find cell containing boundary nodes
+void FIsoTriangulator::FindCellContainingBoundaryNodes(TArray<FCell>& Cells)
+{
+	FTimePoint StartTime = FChrono::Now();
+
+	TArray<int32> NodeToCellIndices;
+	TArray<int32> SortedIndex;
+
+	const int32 CountU = Grid.GetCuttingCount(EIso::IsoU);
+	const int32 CountV = Grid.GetCuttingCount(EIso::IsoV);
+
+	const TArray<double>& IsoUCoordinates = Grid.GetUniformCuttingCoordinatesAlongIso(EIso::IsoU);
+	const TArray<double>& IsoVCoordinates = Grid.GetUniformCuttingCoordinatesAlongIso(EIso::IsoV);
+
+	NodeToCellIndices.Reserve(LoopNodeCount);
+	{
+		int32 IndexU = 0;
+		int32 IndexV = 0;
+		int32 Index = 0;
+		for (const FLoopNode& LoopPoint : LoopNodes)
+		{
+			const FPoint2D& Coordinate = LoopPoint.Get2DPoint(EGridSpace::UniformScaled, Grid);
+
+			ArrayUtils::FindCoordinateIndex(IsoUCoordinates, Coordinate.U, IndexU);
+			ArrayUtils::FindCoordinateIndex(IsoVCoordinates, Coordinate.V, IndexV);
+
+			NodeToCellIndices.Emplace(IndexV * CountU + IndexU);
+			SortedIndex.Emplace(Index++);
+		}
+	}
+
+	Algo::Sort(SortedIndex, [&](const int32& Index1, const int32& Index2)
+		{
+			return (NodeToCellIndices[Index1] < NodeToCellIndices[Index2]);
+		});
+
+	int32 CountOfCellsFilled = 1;
+	{
+		int32 CellIndex = NodeToCellIndices[0];
+		for (int32 Index : SortedIndex)
+		{
+			if (CellIndex != NodeToCellIndices[Index])
+			{
+				CellIndex = NodeToCellIndices[Index];
+				CountOfCellsFilled++;
+			}
+		}
+	}
+
+	// build Cells
+	{
+		Cells.Reserve(CountOfCellsFilled);
+		int32 CellIndex = NodeToCellIndices[SortedIndex[0]];
+		TArray<FLoopNode*> CellNodes;
+		CellNodes.Reserve(LoopNodeCount);
+		for (int32 Index : SortedIndex)
+		{
+			if (CellIndex != NodeToCellIndices[Index])
+			{
+				Cells.Emplace(CellIndex, CellNodes, Grid);
+
+				CellIndex = NodeToCellIndices[Index];
+				CellNodes.Empty(LoopNodeCount);
+			}
+
+			FLoopNode& LoopNode = LoopNodes[Index];
+			CellNodes.Add(&LoopNode);
+		}
+		Cells.Emplace(CellIndex, CellNodes, Grid);
+	}
+	FChrono::Elapse(StartTime);
+}
+
+void FIsoTriangulator::FindCandidateSegmentsToLinkInnerAndLoop()
+{
+	const double FlatAngle = 0.1;
+
+#ifdef CADKERNEL_DEV
+	FTimePoint StartTime = FChrono::Now();
+#endif
+	TFunction<void(FIsoInnerNode*, FLoopNode&)> CreateCandidateSegment = [&](FIsoNode* InnerNode, FLoopNode& LoopNode)
+	{
+		FIsoSegment& SegCandidate = IsoSegmentFactory.New();
+		SegCandidate.Init(*InnerNode, LoopNode, ESegmentType::InnerToLoop);
+		NewTestSegments.Add(&SegCandidate);
+	};
+
+	TFunction<void(FLoopNode&, FLoopNode&)> CreateCandidateBoundarySegment = [&](FLoopNode& StartNode, FLoopNode& EndNode)
+	{
+		FIsoSegment& SegCandidate = IsoSegmentFactory.New();
+		SegCandidate.Init(StartNode, EndNode, ESegmentType::LoopToLoop);
+		NewTestSegments.Add(&SegCandidate);
+	};
+
+	int32 CountU = Grid.GetCuttingCount(EIso::IsoU);
+	int32 CountV = Grid.GetCuttingCount(EIso::IsoV);
+
+	// find cell containing boundary nodes
+	TArray<int32> NodeToCellIndices;
+	TArray<int32> SortedIndex;
+	{
+		//F3DDebugSession _(TEXT("New Test"));
+
+		const TArray<double>& IsoUCoordinates = Grid.GetUniformCuttingCoordinatesAlongIso(EIso::IsoU);
+		const TArray<double>& IsoVCoordinates = Grid.GetUniformCuttingCoordinatesAlongIso(EIso::IsoV);
+
+		NodeToCellIndices.Reserve(LoopNodeCount);
+		int32 IndexU = 0;
+		int32 IndexV = 0;
+		int32 Index = 0;
+		for(const FLoopNode& LoopPoint : LoopNodes)
+		{
+			const FPoint2D& Coordinate = LoopPoint.Get2DPoint(EGridSpace::UniformScaled, Grid);
+
+			ArrayUtils::FindCoordinateIndex(IsoUCoordinates, Coordinate.U, IndexU);
+			ArrayUtils::FindCoordinateIndex(IsoVCoordinates, Coordinate.V, IndexV);
+
+			NodeToCellIndices.Emplace(IndexV * CountU + IndexU);
+			SortedIndex.Emplace(Index++);
+		}
+
+		Algo::Sort(SortedIndex, [&](const int32& Index1, const int32& Index2)
+			{
+				return (NodeToCellIndices[Index1] < NodeToCellIndices[Index2]);
+			});
+	}
+
+
+	int32 CellIndex = -1;
+
+	FIsoInnerNode* Cell[4];
+	TFunction<void(int32 CellIndex)> GetCellIsoNode = [&](int32 CellIndex)
+	{
+		int32 Index = CellIndex;
+		Cell[0] = GlobalIndexToIsoInnerNodes[Index++];
+		Cell[1] = GlobalIndexToIsoInnerNodes[Index];
+		Index += CountU;
+		Cell[2] = GlobalIndexToIsoInnerNodes[Index--];
+		Cell[3] = GlobalIndexToIsoInnerNodes[Index];
+	};
+
+	// create segment between a boundary node and a cell border
+	//Open3DDebugSession(FString::Printf(TEXT("Cell %d"), CellIndex));
+	for (int32 Index : SortedIndex)
+	{
+		if (CellIndex != NodeToCellIndices[Index])
+		{
+			CellIndex = NodeToCellIndices[Index];
+			GetCellIsoNode(CellIndex);
+
+			//Close3DDebugSession();
+			//Open3DDebugSession(FString::Printf(TEXT("Cell %d"), CellIndex));
+		}
+
+		FLoopNode& LoopPoint = LoopNodes[Index];
+
+		for(int32 ICell = 0; ICell < 4; ++ICell)
+		{
+			if (Cell[ICell])
+			{
+				if (LoopPoint.IsSegmentBeInsideFace(Cell[ICell]->Get2DPoint(EGridSpace::UniformScaled, Grid), Grid, FlatAngle))
+				{
+					continue;
+				}
+
+				if (LoopSegmentsIntersectionTool.DoesIntersect(*Cell[ICell], LoopPoint))
+				{
+					continue;
+				}
+
+				CreateCandidateSegment(Cell[ICell], LoopPoint);
+			}
+		}
+
+		//DisplayPoint(LoopPoint.Get2DPoint(EGridSpace::UniformScaled, Grid), EVisuProperty::YellowPoint);
+	}
+	//Close3DDebugSession();
+
+	// create segment between two boundary nodes
+	CellIndex = -1;
+	//Open3DDebugSession(TEXT("Find in Cell "));
+	//Open3DDebugSession(FString::Printf(TEXT("Cell %d"), 0));
+	for (int32 Index = 0; Index < SortedIndex.Num() - 1; ++Index)
+	{
+		int32 ISortedIndex = SortedIndex[Index];
+		FLoopNode& StartLoop = LoopNodes[ISortedIndex];
+		const FPoint2D& StartPoint = StartLoop.Get2DPoint(EGridSpace::UniformScaled, Grid);
+
+		//if (CellIndex != NodeToCellIndices[ISortedIndex])
+		//{
+		//	Close3DDebugSession();
+		//	Open3DDebugSession(FString::Printf(TEXT("Cell %d"), NodeToCellIndices[ISortedIndex]));
+		//	DisplayPoint(StartPoint, EVisuProperty::YellowPoint);
+		//}
+
+		CellIndex = NodeToCellIndices[ISortedIndex];
+
+		for (int32 Jndex = Index + 1; Jndex < SortedIndex.Num(); ++Jndex)
+		{
+			int32 JSortedIndex = SortedIndex[Jndex];
+			if (CellIndex != NodeToCellIndices[JSortedIndex])
+			{
+				break;
+			}
+			FLoopNode& EndLoop = LoopNodes[JSortedIndex];
+			//DisplayPoint(EndLoop.Get2DPoint(EGridSpace::UniformScaled, Grid), EVisuProperty::YellowPoint);
+
+			if (&EndLoop.GetPreviousNode() == &StartLoop || &EndLoop.GetNextNode() == &StartLoop)
+			{
+				continue;
+			}
+
+			const FPoint2D& EndPoint = EndLoop.Get2DPoint(EGridSpace::UniformScaled, Grid);
+
+			if (StartLoop.IsSegmentBeInsideFace(EndPoint, Grid, FlatAngle))
+			{
+				continue;
+			}
+
+			if (EndLoop.IsSegmentBeInsideFace(StartPoint, Grid, FlatAngle))
+			{
+				continue;
+			}
+
+			if (LoopSegmentsIntersectionTool.DoesIntersect(StartLoop, EndLoop))
+			{
+				continue;
+			}
+
+			CreateCandidateBoundarySegment(StartLoop, EndLoop);
+		}
+	}
+	//Close3DDebugSession();
+	//Close3DDebugSession();
+	//Wait();
+
+	//Display(EGridSpace::UniformScaled, TEXT("TEST"), NewTestSegments, false);
+	Wait(true);
+
+#ifdef CADKERNEL_DEV
+	Chronos.FindSegmentToLinkInnerToLoopDuration = FChrono::Elapse(StartTime);
+#endif
+}
+
 void FIsoTriangulator::FindCandidateSegmentsToLinkInnerToLoop()
 {
+#ifdef CADKERNEL_DEV
 	FTimePoint StartTime = FChrono::Now();
+#endif
 
 	TArray<FIsoInnerNode*> SortedInnerNode;
 	SortedInnerNode.Reserve(InnerNodes.Num());
@@ -1026,8 +1482,8 @@
 
 	Algo::Sort(SortedInnerNode, [this](const FIsoInnerNode* IsoNode1, const FIsoInnerNode* IsoNode2)
 		{
-			const FPoint2D& Node1Coordinates = IsoNode1->Get2DPoint(EGridSpace::Default2D, Grid);
-			const FPoint2D& Node2Coordinates = IsoNode2->Get2DPoint(EGridSpace::Default2D, Grid);
+			const FPoint2D& Node1Coordinates = IsoNode1->Get2DPoint(EGridSpace::UniformScaled, Grid);
+			const FPoint2D& Node2Coordinates = IsoNode2->Get2DPoint(EGridSpace::UniformScaled, Grid);
 			return (Node1Coordinates.U + Node1Coordinates.V) < (Node2Coordinates.U + Node2Coordinates.V);
 		});
 
@@ -1035,7 +1491,7 @@
 	Chronos.FindSegmentToLinkInnerToLoopDuration = FChrono::Elapse(StartTime);
 #endif
 
-	TFunction<void(FIsoInnerNode*, FLoopNode*)> CreateInnerToContoursCandidateSegment = [this](FIsoInnerNode* InnerNode, FLoopNode* LoopNode)
+	TFunction<void(FIsoInnerNode*, FLoopNode*)> CreateInnerToContoursCandidateSegment = [&](FIsoInnerNode* InnerNode, FLoopNode* LoopNode)
 	{
 		if (FIsoSegment::IsItAlreadyDefined(InnerNode, LoopNode))
 		{
@@ -1043,12 +1499,12 @@
 		}
 
 		// check intersection with
-		if (LoopSegmentsIntersectionTool.DoesIntersect(*LoopNode, InnerNode->Get2DPoint(EGridSpace::Default2D, Grid)))
+		if (LoopSegmentsIntersectionTool.DoesIntersect(*LoopNode, InnerNode->Get2DPoint(EGridSpace::UniformScaled, Grid)))
 		{
 			return;
 		}
 
-		if (InnerSegmentsIntersectionTool.DoesIntersect(*InnerNode, LoopNode->Get2DPoint(EGridSpace::Default2D, Grid)))
+		if (InnerSegmentsIntersectionTool.DoesIntersect(*InnerNode, LoopNode->Get2DPoint(EGridSpace::UniformScaled, Grid)))
 		{
 			return;
 		}
@@ -1114,9 +1570,9 @@
 		int32 StartInnerNodeIndex = 0;
 		for (FLoopNode* LoopNode : SortedLoopNodes)
 		{
-			const FPoint2D& PointLoop = LoopNode->Get2DPoint(EGridSpace::Default2D, Grid);
+			const FPoint2D& PointLoop = LoopNode->Get2DPoint(EGridSpace::UniformScaled, Grid);
 			{
-				const FPoint2D& PointInner = Grid.GetInner2DPoint(EGridSpace::Default2D, SortedInnerNode[StartInnerNodeIndex]->GetIndex());
+				const FPoint2D& PointInner = Grid.GetInner2DPoint(EGridSpace::UniformScaled, SortedInnerNode[StartInnerNodeIndex]->GetIndex());
 				double Diagonal = ComputeDiagonal(PointLoop, PointInner);
 				if (Diagonal > DiagonalMax)
 				{
@@ -1126,7 +1582,7 @@
 
 			for (; StartInnerNodeIndex < SortedInnerNode.Num(); ++StartInnerNodeIndex)
 			{
-				const FPoint2D& PointInner = Grid.GetInner2DPoint(EGridSpace::Default2D, SortedInnerNode[StartInnerNodeIndex]->GetIndex());
+				const FPoint2D& PointInner = Grid.GetInner2DPoint(EGridSpace::UniformScaled, SortedInnerNode[StartInnerNodeIndex]->GetIndex());
 
 				double Diagonal = ComputeDiagonal(PointLoop, PointInner);
 				if (Diagonal >= -DiagonalMax)
@@ -1141,7 +1597,7 @@
 
 			for (int32 InnerNodeIndex = StartInnerNodeIndex; InnerNodeIndex < SortedInnerNode.Num(); ++InnerNodeIndex)
 			{
-				const FPoint2D& PointInner = Grid.GetInner2DPoint(EGridSpace::Default2D, SortedInnerNode[InnerNodeIndex]->GetIndex());
+				const FPoint2D& PointInner = Grid.GetInner2DPoint(EGridSpace::UniformScaled, SortedInnerNode[InnerNodeIndex]->GetIndex());
 
 				FPoint2D DiagonalV = PointLoop - PointInner;
 				double Diagonal = FMath::Abs(DiagonalV.U + DiagonalV.V);
@@ -1151,9 +1607,9 @@
 					break;
 				}
 
-				double DeltaUMax = GetDeltaCoordinate(Grid.GetCuttingCoordinatesAlongIso(EIso::IsoU), IndexU, PointInner.U);
+				double DeltaUMax = GetDeltaCoordinate(Grid.GetUniformCuttingCoordinatesAlongIso(EIso::IsoU), IndexU, PointInner.U);
 				DeltaUMax *= 1.5;
-				double DeltaVMax = GetDeltaCoordinate(Grid.GetCuttingCoordinatesAlongIso(EIso::IsoV), IndexV, PointInner.V);
+				double DeltaVMax = GetDeltaCoordinate(Grid.GetUniformCuttingCoordinatesAlongIso(EIso::IsoV), IndexV, PointInner.V);
 				DeltaVMax *= 1.5;
 
 				if (FMath::Abs(DiagonalV.U) > DeltaUMax || FMath::Abs(DiagonalV.V) > DeltaVMax)
@@ -1277,7 +1733,7 @@
 			return (Segment1->Get3DLengthSquare(Grid)) < (Segment2->Get3DLengthSquare(Grid));
 		});
 
-	InnerSegmentsIntersectionTool.SetNum(InnerSegmentsIntersectionTool.Num() + (int32)CandidateSegments.Num());
+	InnerSegmentsIntersectionTool.SetNum(InnerSegmentsIntersectionTool.Count() + (int32)CandidateSegments.Num());
 
 	FinalToLoops.Reserve(CandidateSegments.Num());
 
@@ -1548,7 +2004,7 @@
 
 	FIntersectionSegmentTool CycleIntersectionTool(Grid);
 	CycleIntersectionTool.Reserve((int32)FMath::Square(NodeCycleNum));
-	CycleIntersectionTool.AddSegment(Cycle);
+	CycleIntersectionTool.AddSegments(Cycle);
 	CycleIntersectionTool.Sort();
 
 	// check if the cycle is self intersecting. 
@@ -1596,16 +2052,16 @@
 		// If the Segment has 2 adjacent triangles, the segment is a inner cycle segment
 		// It will have triangle in both side
 		//
-		//    X---------------X----------------X        X---------------X----------------X
-		//    |                                |   		|                                |  
-		//    |         X--------------------X |   		|                                |  
-		//    |         |                    | |   		|                                |  
-		//    X---------X  <- inner segment  | |	or  X---------X  <- inner segment    |
-		//    |         |                    | |   		|                                |  
-		//    |         X--------------------X |   		|                                |  
-		//    |                                |		|                                |
-		//    X---------------X----------------X		X---------------X----------------X
-		//												
+		//    X---------------X----------------X      X---------------X----------------X
+		//    |                                |      |                                |  
+		//    |         X--------------------X |      |                                |  
+		//    |         |                    | |      |                                |  
+		//    X---------X  <- inner segment  | |  or  X---------X  <- inner segment    |
+		//    |         |                    | |      |                                |  
+		//    |         X--------------------X |      |                                |  
+		//    |                                |      |                                |
+		//    X---------------X----------------X      X---------------X----------------X
+		//
 		for (FIsoSegment* Segment : Cycle)
 		{
 			if (Segment->HasTriangleOnRightAndLeft())
@@ -1757,7 +2213,7 @@
 				Display(EGridSpace::UniformScaled, *NextSegment, 0, EVisuProperty::EdgeMesh);
 				Display(EGridSpace::UniformScaled, EndNode, 0, EVisuProperty::RedPoint);
 				Display(EGridSpace::UniformScaled, NextNode);
-				Wait(Grid.GetFace()->GetId() == SurfaceToDebug);
+				Wait(Grid.GetFace()->GetId() == FaceToDebug);
 			}
 		}
 #endif
@@ -1916,7 +2372,7 @@
 #ifdef ADD_TRIANGLE_2D
 			if (bDisplay)
 			{
-				F3DDebugSession _(FString::Printf(TEXT("Triangle %d"), TriangleIndex));
+				F3DDebugSession _(FString::Printf(TEXT("Triangle")));
 				DisplayTriangle(EGridSpace::UniformScaled, EndNode, StartNode, *CandidatNode);
 				Wait(false);
 			}
@@ -1953,7 +2409,7 @@
 	}
 
 #ifdef ADD_TRIANGLE_2D
-	if (Grid.GetFace()->GetId() == SurfaceToDebug)
+	if (Grid.GetFace()->GetId() == FaceToDebug)
 	{
 		Close3DDebugSession();
 	}
@@ -2086,6 +2542,7 @@
 #endif
 }
 
+//#define FIND_CYCLE
 void FIsoTriangulator::TriangulateOverCycle(const EGridSpace Space)
 {
 	FTimePoint StartTime = FChrono::Now();
@@ -2097,7 +2554,7 @@
 
 	int32 CycleIndex = 0;
 #ifdef FIND_CYCLE
-	if (Grid.GetFace()->GetId() == SurfaceToDebug)
+	if (Grid.GetFace()->GetId() == FaceToDebug)
 	{
 		Open3DDebugSession(TEXT("Triangulate Over Cycle"));
 	}
@@ -2112,7 +2569,7 @@
 			bool bLeftSide = true;
 			FindCycle(Segment, bLeftSide, Cycle, CycleOrientation);
 #ifdef FIND_CYCLE
-			if (Grid.GetFace()->GetId() == SurfaceToDebug)
+			if (Grid.GetFace()->GetId() == FaceToDebug)
 			{
 				F3DDebugSession G(TEXT("Find & mesh cycles"));
 				FString Message = FString::Printf(TEXT("MeshCycle - cycle %d"), CycleIndex++);
@@ -2132,7 +2589,7 @@
 			bool bLeftSide = true;
 			FindCycle(Segment, bLeftSide, Cycle, CycleOrientation);
 #ifdef FIND_CYCLE
-			if (Grid.GetFace()->GetId() == SurfaceToDebug)
+			if (Grid.GetFace()->GetId() == FaceToDebug)
 			{
 				Open3DDebugSession(TEXT("Find & mesh cycles"));
 				FString Message = FString::Printf(TEXT("MeshCycle - cycle %d"), CycleIndex++);
@@ -2150,11 +2607,11 @@
 			bool bLeftSide = false;
 			FindCycle(Segment, bLeftSide, Cycle, CycleOrientation);
 #ifdef FIND_CYCLE
-			if (Grid.GetFace()->GetId() == SurfaceToDebug)
+			if (Grid.GetFace()->GetId() == FaceToDebug)
 			{
 				Open3DDebugSession(TEXT("Find & mesh cycles"));
 				FString Message = FString::Printf(TEXT("MeshCycle - cycle %d"), CycleIndex++);
-				Display(EGridSpace::UniformScaled, *Message, Cycle, false);
+				::Display(EGridSpace::UniformScaled, *Message, Cycle, false);
 				Close3DDebugSession();
 			}
 #endif
@@ -2181,7 +2638,7 @@
 #ifdef FIND_CYCLE
 	CycleIndex++;
 	//CycleId = CycleIndex;
-	if (Grid.GetFace()->GetId() == SurfaceToDebug)
+	if (Grid.GetFace()->GetId() == FaceToDebug)
 	{
 		CycleIndex = CycleId;
 		bDisplay = true;
@@ -2450,10 +2907,620 @@
 // =========================================================================================================================================================================================================
 // =========================================================================================================================================================================================================
 // =========================================================================================================================================================================================================
+//#define DEBUG_DELAUNAY
 
-void FIsoTriangulator::Delaunay()
+void FIsoTriangulator::ConnectCellLoopsByNeighborhood(FCell& Cell)
 {
 	FTimePoint StartTime = FChrono::Now();
+	//F3DDebugSession _(TEXT("ConnectCellLoopsByNeighborhood"));
+
+	int32 LoopCount = Cell.Loops.Num();
+
+	TArray<FPoint2D> LoopBarycenters;
+	// the four last points are the corners of Barycenter that defined the starting mesh for Bowyer & Watson algorithm
+	LoopBarycenters.Reserve(LoopCount + 4);
+
+	for (const TArray<FLoopNode*>& Nodes : Cell.Loops)
+	{
+		FPoint2D& BaryCenter = LoopBarycenters.Emplace_GetRef(FPoint2D::ZeroPoint);
+
+		// the external loop is not proceed 
+		if (Nodes[0]->GetLoopIndex() == 0)
+		{
+			continue;
+		}
+
+		for (const FLoopNode* Node : Nodes)
+		{
+			BaryCenter += Node->Get2DPoint(EGridSpace::UniformScaled, Grid);
+		}
+		BaryCenter /= (double)Nodes.Num();
+	}
+
+	TArray<int32> EdgeVertexIndices;
+	if (Cell.bHasOuterLoop && LoopCount < 5)
+	{
+		EdgeVertexIndices.Reserve(6);
+		Cell.BorderLoopIndices.Reserve(3);
+		if (LoopCount == 2)
+		{
+			Cell.BorderLoopIndices.Add(1);
+		}
+		else if (LoopCount == 3)
+		{
+			EdgeVertexIndices.Append({ 1, 2 });
+			Cell.BorderLoopIndices.Append({ 1, 2 });
+		}
+		else if (LoopCount == 4)
+		{
+			EdgeVertexIndices.Append({ 1, 2, 2, 3, 3, 1 });
+			Cell.BorderLoopIndices.Append({ 1, 2, 3 });
+		}
+	}
+	else if (LoopBarycenters.Num() < 4)
+	{
+		EdgeVertexIndices.Reserve(6);
+		Cell.BorderLoopIndices.Reserve(3);
+		if (LoopCount == 1)
+		{
+			Cell.BorderLoopIndices.Add(0);
+		}
+		else if (LoopCount == 2)
+		{
+			EdgeVertexIndices.Append({ 0, 1 });
+			Cell.BorderLoopIndices.Append({ 0, 1 });
+		}
+		else if (LoopCount == 3)
+		{
+			EdgeVertexIndices.Append({ 0, 1, 1, 2, 2, 0 });
+			Cell.BorderLoopIndices.Append({ 0, 1, 2 });
+		}
+	}
+	else
+	{
+		FBowyerWatsonTriangulator Triangulator(LoopBarycenters, EdgeVertexIndices);
+		Triangulator.Triangulate(Cell.bHasOuterLoop);
+		Triangulator.GetOuterVertices(Cell.BorderLoopIndices);
+	}
+
+	// Connect inner close loops 
+	// ==========================================================================================
+	{
+#ifdef DEBUG_DELAUNAY
+		//DisplayCell(Cell);
+		//F3DDebugSession _(TEXT("Build Segments"));
+#endif
+
+		for (int32 Index = 0; Index < EdgeVertexIndices.Num();)
+		{
+			//Open3DDebugSession(TEXT("Segment"));
+			int32 IndexLoopA = EdgeVertexIndices[Index++];
+			int32 IndexLoopB = EdgeVertexIndices[Index++];
+
+			//DisplaySegment(LoopBarycenters[IndexLoopA], LoopBarycenters[IndexLoopB]);
+
+			TryToConnectTwoLoopsWithShortestSegment(Cell, IndexLoopA, IndexLoopB);
+			//Close3DDebugSession();
+			//Wait();
+		}
+#ifdef DEBUG_DELAUNAY
+		//Wait();
+#endif
+	}
+
+	// With Outer loop
+	// ==========================================================================================
+	if (Cell.bHasOuterLoop && Cell.Loops.Num() > 1)
+	{
+		Cell.IntersectionTool.AddSegments(Cell.CandidateSegments);
+		Cell.IntersectionTool.Sort();
+
+#ifdef DEBUG_DELAUNAY
+		//F3DDebugSession _(TEXT("Build Segments"));
+#endif
+		for (TArray<FLoopNode*>& SubLoop : Cell.OuterLoopSubdivision)
+		{
+			for (int32 IndexLoopB : Cell.BorderLoopIndices)
+			{
+				TryToConnectTwoLoopsWithShortestSegment(Cell, SubLoop, IndexLoopB);
+			}
+		}
+	}
+
+	Cell.SelectSegmentInCandidateSegments(IsoSegmentFactory);
+
+#ifdef CADKERNEL_DEV
+	Chronos.FindSegmentToLinkLoopToLoopByDelaunayDuration += FChrono::Elapse(StartTime);
+#endif
+}
+
+void FIsoTriangulator::TryToConnectTwoLoopsWithShortestSegment(FCell& Cell, int32 IndexLoopA, int32 IndexLoopB)
+{
+	const TArray<FLoopNode*>& LoopA = Cell.Loops[IndexLoopA];
+	const TArray<FLoopNode*>& LoopB = Cell.Loops[IndexLoopB];
+	TryToConnectTwoLoopsWithShortestSegment(Cell, LoopA, LoopB);
+};
+
+void FIsoTriangulator::TryToConnectTwoLoopsWithShortestSegment(FCell& Cell, const TArray<FLoopNode*>& LoopA, int32 IndexLoopB)
+{
+	const TArray<FLoopNode*>& LoopB = Cell.Loops[IndexLoopB];
+	TryToConnectTwoLoopsWithShortestSegment(Cell, LoopA, LoopB);
+}
+
+void FIsoTriangulator::TryToConnectTwoLoopsWithShortestSegment(FCell& Cell, const TArray<FLoopNode*>& LoopA, const TArray<FLoopNode*>&LoopB)
+{
+	double MinDistanceSquare = HUGE_VALUE_SQUARE;
+	int32 MinIndexA = -1;
+	int32 MinIndexB = -1;
+
+	for (int32 IndexA = 0; IndexA < LoopA.Num(); ++IndexA)
+	{
+		const FLoopNode* NodeA = LoopA[IndexA];
+		const FPoint2D& ACoordinates = NodeA->Get2DPoint(EGridSpace::UniformScaled, Grid);
+
+		for (int32 IndexB = 0; IndexB < LoopB.Num(); ++IndexB)
+		{
+			const FLoopNode* NodeB = LoopB[IndexB];
+			const FPoint2D& BCoordinates = NodeB->Get2DPoint(EGridSpace::UniformScaled, Grid);
+
+			double SquareDistance = ACoordinates.SquareDistance(BCoordinates);
+			if (SquareDistance < MinDistanceSquare)
+			{
+				MinDistanceSquare = SquareDistance;
+				MinIndexA = IndexA;
+				MinIndexB = IndexB;
+			}
+		}
+	}
+
+	if (MinIndexA >= 0 && MinIndexB >= 0)
+	{
+		FLoopNode* NodeA = LoopA[MinIndexA];
+		const FPoint2D& ACoordinates = NodeA->Get2DPoint(EGridSpace::UniformScaled, Grid);
+		FLoopNode* NodeB = LoopB[MinIndexB];
+		const FPoint2D& BCoordinates = NodeB->Get2DPoint(EGridSpace::UniformScaled, Grid);
+
+		TryToCreateSegment(Cell, NodeA, ACoordinates, NodeB, BCoordinates, 0.1);
+	}
+
+};
+
+void FIsoTriangulator::TryToConnectTwoLoopsWithTheMostIsoSegment(FCell& Cell, const TArray<FLoopNode*>& LoopA, const TArray<FLoopNode*>& LoopB)
+{
+	double MinSlope = HUGE_VALUE_SQUARE;
+	int32 MinIndexA = -1;
+	int32 MinIndexB = -1;
+
+	for (int32 IndexA = 0; IndexA < LoopA.Num(); ++IndexA)
+	{
+		const FLoopNode* NodeA = LoopA[IndexA];
+		const FPoint2D& ACoordinates = NodeA->Get2DPoint(EGridSpace::UniformScaled, Grid);
+
+		for (int32 IndexB = 0; IndexB < LoopB.Num(); ++IndexB)
+		{
+			const FLoopNode* NodeB = LoopB[IndexB];
+			const FPoint2D& BCoordinates = NodeB->Get2DPoint(EGridSpace::UniformScaled, Grid);
+
+			double Slope = ComputeUnorientedSlope(ACoordinates, BCoordinates, 0);
+			if(Slope > 2)
+			{
+				Slope = 4 - Slope;
+			}
+
+			if (Slope < MinSlope)
+			{
+				MinSlope = Slope;
+				MinIndexA = IndexA;
+				MinIndexB = IndexB;
+			}
+		}
+	}
+
+	if (MinIndexA >= 0 && MinIndexB >= 0)
+	{
+		FLoopNode* NodeA = LoopA[MinIndexA];
+		const FPoint2D& ACoordinates = NodeA->Get2DPoint(EGridSpace::UniformScaled, Grid);
+		FLoopNode* NodeB = LoopB[MinIndexB];
+		const FPoint2D& BCoordinates = NodeB->Get2DPoint(EGridSpace::UniformScaled, Grid);
+
+		TryToCreateSegment(Cell, NodeA, ACoordinates, NodeB, BCoordinates, 0.1);
+	}
+
+};
+
+void FIsoTriangulator::TryToCreateSegment(FCell& Cell, FLoopNode* NodeA, const FPoint2D& ACoordinates, FIsoNode* NodeB, const FPoint2D& BCoordinates, const double FlatAngle)
+{
+	if (NodeA->GetSegmentConnectedTo(NodeB))
+	{
+		return;
+	}
+
+	if (InnerSegmentsIntersectionTool.DoesIntersect(*NodeA, *NodeB))
+	{
+		return;
+	}
+
+	if (InnerToLoopSegmentsIntersectionTool.DoesIntersect(*NodeA, *NodeB))
+	{
+		return;
+	}
+
+	if (Cell.IntersectionTool.DoesIntersect(*NodeA, *NodeB))
+	{
+		return;
+	}
+
+	if (LoopSegmentsIntersectionTool.DoesIntersect(*NodeA, *NodeB))
+	{
+		return;
+	}
+
+	// Is Outside and not to flat at NodeA
+	if (NodeA->IsSegmentBeInsideFace(BCoordinates, Grid, FlatAngle))
+	{
+		return;
+	}
+
+	// Is Outside and not to flat at NodeB
+	if (NodeB->IsALoopNode())
+	{
+		if (((FLoopNode*)NodeB)->IsSegmentBeInsideFace(ACoordinates, Grid, FlatAngle))
+		{
+			return;
+		}
+	}
+
+	FIsoSegment& Segment = IsoSegmentFactory.New();
+	Segment.Init(*NodeA, *NodeB, ESegmentType::LoopToLoop);
+	Segment.SetCandidate();
+	Cell.CandidateSegments.Add(&Segment);
+
+#ifdef DEBUG_DELAUNAY
+	//DisplaySegment(ACoordinates, BCoordinates, 0, EVisuProperty::OrangeCurve);
+#endif
+};
+
+void FIsoTriangulator::ConnectCellCornerToInnerLoop(FCell& Cell)
+{
+	FIsoInnerNode* CellNodes[4];
+	int32 Index = Cell.Id;
+	CellNodes[0] = GlobalIndexToIsoInnerNodes[Index++];
+	CellNodes[1] = GlobalIndexToIsoInnerNodes[Index];
+	Index += Grid.GetCuttingCount(EIso::IsoU);;
+	CellNodes[2] = GlobalIndexToIsoInnerNodes[Index--];
+	CellNodes[3] = GlobalIndexToIsoInnerNodes[Index];
+
+	{
+		int32 ICell = 0;
+		for (; ICell < 4; ++ICell)
+		{
+			if (CellNodes[ICell])
+			{
+				break;
+			}
+		}
+		if (ICell == 4)
+		{
+			// All Cell corners are not null
+			return;
+		}
+	}
+
+#ifdef DEBUG_DELAUNAY
+	//Open3DDebugSession(TEXT("With cell corners"));
+#endif
+
+	TFunction<void(int32, FIsoInnerNode*)> FindAndTryCreateCandidateSegmentToLinkLoopToCorner = [&](int32 IndexLoopA, FIsoInnerNode* InnerNode)
+	{
+
+		const FPoint2D& InnerCoordinates = InnerNode->Get2DPoint(EGridSpace::UniformScaled, Grid);
+
+		const TArray<FLoopNode*>& LoopA = Cell.Loops[IndexLoopA];
+
+		double MinDistanceSquare = HUGE_VALUE_SQUARE;
+		int32 MinIndexA = -1;
+		for (int32 IndexA = 0; IndexA < LoopA.Num(); ++IndexA)
+		{
+			const FLoopNode* NodeA = LoopA[IndexA];
+			const FPoint2D& ACoordinates = NodeA->Get2DPoint(EGridSpace::UniformScaled, Grid);
+
+			double SquareDistance = ACoordinates.SquareDistance(InnerCoordinates);
+			if (SquareDistance < MinDistanceSquare)
+			{
+				MinDistanceSquare = SquareDistance;
+				MinIndexA = IndexA;
+			}
+		}
+
+		if (MinIndexA >= 0)
+		{
+			FLoopNode* NodeA = LoopA[MinIndexA];
+			const FPoint2D& ACoordinates = NodeA->Get2DPoint(EGridSpace::UniformScaled, Grid);
+
+			TryToCreateSegment(Cell, NodeA, ACoordinates, InnerNode, InnerCoordinates, 0.1);
+		}
+
+	};
+
+	int32 IntersectionToolCount = Cell.IntersectionTool.Count();
+	int32 NewSegmentCount = Cell.CandidateSegments.Num() - IntersectionToolCount;
+	Cell.IntersectionTool.AddSegments(Cell.CandidateSegments.GetData() + IntersectionToolCount, NewSegmentCount);
+	Cell.IntersectionTool.Sort();
+
+	for (int32 ICell = 0; ICell < 4; ++ICell)
+	{
+		if (CellNodes[ICell])
+		{
+			for (int32 IndexLoopA : Cell.BorderLoopIndices)
+			{
+				DisplayPoint(CellNodes[ICell]->Get2DPoint(EGridSpace::UniformScaled, Grid), EVisuProperty::GreenPoint);
+				FindAndTryCreateCandidateSegmentToLinkLoopToCorner(IndexLoopA, CellNodes[ICell]);
+			}
+
+			if (Cell.bHasOuterLoop)
+			{
+				FindAndTryCreateCandidateSegmentToLinkLoopToCorner(0, CellNodes[ICell]);
+			}
+		}
+	}
+
+
+#ifdef DEBUG_DELAUNAY
+	Close3DDebugSession();
+#endif
+	Cell.SelectSegmentInCandidateSegments(IsoSegmentFactory);
+}
+
+
+
+// =========================================================================================================================================================================================================
+// =========================================================================================================================================================================================================
+// =========================================================================================================================================================================================================
+//
+//
+//                                                                            NOT YET REVIEWED
+//
+//
+// =========================================================================================================================================================================================================
+// =========================================================================================================================================================================================================
+// =========================================================================================================================================================================================================
+
+
+#ifdef NEED_TO_CHECK_USEFULNESS
+void FIsoTriangulator::CompleteIsoSegmentLoopToLoop()
+{
+	FIntersectionSegmentTool NewSegmentsIntersectionTool(Grid);
+	// similare to FindSegmentToLinkloopToloop::CreateSegment but different 
+	TFunction<void(FLoopNode&, FIsoNode&)> CreateComplementarySegment = [&](FLoopNode& Node1, FIsoNode& Node2)
+	{
+		if (Node1.GetSegmentConnectedTo(&Node2))
+		{
+			return;
+		}
+
+		if (NewSegmentsIntersectionTool.DoesIntersect(Node1, Node2))
+		{
+			return;
+		}
+
+		if (LoopSegmentsIntersectionTool.DoesIntersect(Node1, Node2))
+		{
+			return;
+		}
+
+		FIsoSegment& Segment = IsoSegmentFactory.New();
+		Segment.Init(Node1, Node2, ESegmentType::Unknown);
+		Segment.ConnectToNode();
+		FinalToLoops.Add(&Segment);
+		InnerToLoopSegmentsIntersectionTool.AddSegment(Segment);
+		NewSegmentsIntersectionTool.AddSegment(Segment);
+		NewSegmentsIntersectionTool.Sort();
+	};
+
+	for (FLoopNode& LoopNode : LoopNodes)
+	{
+		if (LoopNode.GetConnectedSegments().Num() == 2)
+		{
+			if (LoopNode.GetPreviousNode().GetConnectedSegments().Num() == 3)
+			{
+				if ((LoopNode.GetPreviousNode().GetConnectedSegments().Last()->GetType() == ESegmentType::IsoLoopToLoop) && (LoopNode.GetNextNode().GetConnectedSegments().Last()->GetType() == ESegmentType::IsoLoopToLoop))
+				{
+					FIsoNode& Node1 = LoopNode.GetPreviousNode().GetConnectedSegments().Last()->GetOtherNode(&LoopNode.GetPreviousNode());
+					DisplaySegment(Node1.Get2DPoint(EGridSpace::Default2D, Grid), LoopNode.Get2DPoint(EGridSpace::Default2D, Grid));
+					CreateComplementarySegment(LoopNode, Node1);
+					FIsoNode& Node2 = LoopNode.GetNextNode().GetConnectedSegments().Last()->GetOtherNode(&LoopNode.GetNextNode());
+					CreateComplementarySegment(LoopNode, Node2);
+					DisplaySegment(Node2.Get2DPoint(EGridSpace::Default2D, Grid), LoopNode.Get2DPoint(EGridSpace::Default2D, Grid));
+				}
+				else if (LoopNode.GetPreviousNode().GetConnectedSegments().Last()->GetType() == ESegmentType::IsoLoopToLoop)
+				{
+					FIsoNode& Node = LoopNode.GetPreviousNode().GetConnectedSegments().Last()->GetOtherNode(&LoopNode.GetPreviousNode());
+					CreateComplementarySegment(LoopNode, Node);
+					DisplaySegment(Node.Get2DPoint(EGridSpace::Default2D, Grid), LoopNode.Get2DPoint(EGridSpace::Default2D, Grid));
+				}
+				else if (LoopNode.GetNextNode().GetConnectedSegments().Last()->GetType() == ESegmentType::IsoLoopToLoop)
+				{
+					FIsoNode& Node = LoopNode.GetNextNode().GetConnectedSegments().Last()->GetOtherNode(&LoopNode.GetNextNode());
+					CreateComplementarySegment(LoopNode, Node);
+					DisplaySegment(Node.Get2DPoint(EGridSpace::Default2D, Grid), LoopNode.Get2DPoint(EGridSpace::Default2D, Grid));
+				}
+			}
+		}
+	}
+}
+#endif
+
+#ifdef UNUSED
+void FIsoTriangulator::FindSegmentToLinkLoopToLoop()
+{
+	FTimePoint StartTime = FChrono::Now();
+
+	double DiagonalMax = Grid.GetMaxDeltaU(EIso::IsoU) + Grid.GetMaxDeltaU(EIso::IsoV);
+	DiagonalMax *= 2;
+
+	TFunction<double(const TArray<double>&, int32&, const double&)> GetDeltaCoordinate = [](const TArray<double>& Iso, int32& Index, const double& PointCoord)
+	{
+		while (PointCoord + SMALL_NUMBER < Iso[Index])
+		{
+			if (Index == 0)
+			{
+				break;
+			}
+			Index--;
+		}
+		for (; Index < Iso.Num(); ++Index)
+		{
+			if (Index == Iso.Num() - 2)
+			{
+				break;
+			}
+			if (PointCoord < Iso[Index + 1])
+			{
+				break;
+			}
+		}
+
+		// If PointCoord is closer Iso[Index + 1]
+		if ((Iso[Index + 1] - PointCoord) * 2 < (Iso[Index + 1] - Iso[Index]))
+		{
+			if (Index + 2 < Iso.Num())
+			{
+				return FMath::Max(Iso[Index + 1] - Iso[Index], Iso[Index + 2] - Iso[Index + 1]);
+			}
+			else
+			{
+				return Iso[Index + 1] - Iso[Index];
+			}
+		}
+		else
+		{
+			if (Index > 1)
+			{
+				return FMath::Max(Iso[Index] - Iso[Index - 1], Iso[Index + 1] - Iso[Index]);
+			}
+			else
+			{
+				return Iso[Index + 1] - Iso[Index];
+			}
+		}
+		return Iso[Index + 1] - Iso[Index];
+	};
+
+	TFunction<void(FLoopNode*, const FPoint2D&, FLoopNode*, const FPoint2D&)> CreateSegment = [this](FLoopNode* NodeA, const FPoint2D& ACoordinates, FLoopNode* NodeB, const FPoint2D& BCoordinates)
+	{
+		if (NodeA->GetSegmentConnectedTo(NodeB))
+		{
+			return;
+		}
+
+		// IsExiste
+		if (NodeA->GetLoopIndex() == NodeB->GetLoopIndex())
+		{
+			if ((NodeA->GetIndex() == NodeB->GetIndex() + 1) || (NodeB->GetIndex() == NodeA->GetIndex() + 1))
+			{
+				return;
+			}
+		}
+
+		if (InnerSegmentsIntersectionTool.DoesIntersect(ACoordinates, BCoordinates))
+		{
+			return;
+		}
+
+		if (InnerToLoopSegmentsIntersectionTool.DoesIntersect(*NodeA, *NodeB))
+		{
+			return;
+		}
+
+		if (LoopSegmentsIntersectionTool.DoesIntersect(*NodeA, *NodeB))
+		{
+			return;
+		}
+
+		// Is Outside and not to flat at NodeA
+		const double FlatAngle = 1;
+		if (!IsPointPBeInsideSectorABC(NodeA->GetPreviousNode().Get2DPoint(EGridSpace::Default2D, Grid), ACoordinates, NodeA->GetNextNode().Get2DPoint(EGridSpace::Default2D, Grid), BCoordinates, FlatAngle))
+		{
+			return;
+		}
+
+		//double SlopWithNextloop = GetZeroTwoPiSlope(ACoordinates, NodeA->GetNextNode().Get2DPoint(EGridSpace::Default2D, Grid), 0);
+		//double loopDeltaSlop = GetZeroTwoPiSlope(ACoordinates, NodeA->GetPreviousNode().Get2DPoint(EGridSpace::Default2D, Grid), SlopWithNextloop);
+		//double SegmentSlop = GetZeroTwoPiSlope(ACoordinates, BCoordinates, SlopWithNextloop);
+		//if (SegmentSlop < FlatAngle || SegmentSlop + FlatAngle > loopDeltaSlop)
+		//{
+		//	return;
+		//}
+
+		// Is Outside and not to flat at NodeB
+		if (!IsPointPBeInsideSectorABC(NodeB->GetPreviousNode().Get2DPoint(EGridSpace::Default2D, Grid), BCoordinates, NodeB->GetNextNode().Get2DPoint(EGridSpace::Default2D, Grid), ACoordinates, FlatAngle))
+		{
+			return;
+		}
+
+		//SlopWithNextloop = GetZeroTwoPiSlope(BCoordinates, NodeB->GetNextNode().Get2DPoint(EGridSpace::Default2D, Grid), 0);
+		//loopDeltaSlop = GetZeroTwoPiSlope(BCoordinates, NodeB->GetPreviousNode().Get2DPoint(EGridSpace::Default2D, Grid), SlopWithNextloop);
+		//SegmentSlop = GetZeroTwoPiSlope(BCoordinates, ACoordinates, SlopWithNextloop);
+
+		//if (SegmentSlop < FlatAngle || SegmentSlop + FlatAngle > loopDeltaSlop)
+		//{
+		//	return;
+		//}
+
+		FIsoSegment& Segment = IsoSegmentFactory.New();
+		Segment.Init(*NodeA, *NodeB, ESegmentType::LoopToLoop);
+		CandidateLoopToLoopSegments.Add(&Segment);
+	};
+
+	//Open3DDebugSession(TEXT("FIsoTrianguler::FindSegmentToLinkloopToloop"));
+
+	// find closed points
+	int32 IndexU = 0;
+	int32 IndexV = 0;
+
+	for (int32 IndexA = 0; IndexA < SortedLoopNodes.Num() - 1; ++IndexA)
+	{
+		FLoopNode* NodeA = SortedLoopNodes[IndexA];
+		if (NodeA->GetConnectedSegments().Num() > 2)
+		{
+			// Enough connected
+			continue;
+		}
+
+		// check intersection with
+		const FPoint2D& ACoordinates = NodeA->Get2DPoint(EGridSpace::Default2D, Grid);
+
+		double DeltaUMax = GetDeltaCoordinate(Grid.GetCuttingCoordinatesAlongIso(EIso::IsoU), IndexU, ACoordinates.U);
+		//double DeltaUMin = DeltaUMax * 0.1;
+		DeltaUMax *= 1.75;
+		double DeltaVMax = GetDeltaCoordinate(Grid.GetCuttingCoordinatesAlongIso(EIso::IsoV), IndexV, ACoordinates.V);
+		//double DeltaVMin = DeltaVMax * 0.1;
+		DeltaVMax *= 1.75;
+
+		for (int32 IndexB = IndexA + 1; IndexB < SortedLoopNodes.Num(); ++IndexB)
+		{
+			FLoopNode* NodeB = SortedLoopNodes[IndexB];
+			const FPoint2D& BCoordinates = NodeB->Get2DPoint(EGridSpace::Default2D, Grid);
+			FPoint2D AB = BCoordinates - ACoordinates;
+			double Diagonal = FMath::Abs(AB.U + AB.V);
+			if (Diagonal > DiagonalMax)
+			{
+				break;
+			}
+
+			if (FMath::Abs(AB.U) < DeltaUMax && FMath::Abs(AB.V) < DeltaVMax)
+			{
+				CreateSegment(NodeA, ACoordinates, NodeB, BCoordinates);
+			}
+		}
+	}
+	//Close3DDebugSession();
+	Chronos.FindSegmentToLinkLoopToLoopDuration += FChrono::Elapse(StartTime);
+}
+
+void FIsoTriangulator::ConnectCellLoopsByNeighborhood()
+{
+	FTimePoint StartTime = FChrono::Now();
 	//F3DDebugSession DelaunayDebugSession(TEXT("Delaunay Algo"));
 
 	int32 LoopCount = Grid.GetLoops3D().Num();
@@ -2614,7 +3681,7 @@
 			TriangleSet.Emplace(EdgeVertexIndices[Endex + 1], EdgeVertexIndices[Endex], CenterIndex, LoopBarycenters);
 		}
 #ifdef DEBUG_DELAUNAY
-		DisplayTriangles();
+		//DisplayTriangles();
 #endif
 	}
 
@@ -2634,17 +3701,19 @@
 		int32 Index = 0;
 		for (; Index < 3; Index++)
 		{
-			if(TriangleSet[TriangleIndex].VertexIndex[Index] >= LoopCount)
+			if (TriangleSet[TriangleIndex].VertexIndex[Index] >= LoopCount)
 			{
 				break;
 			}
 		}
 
+		// the triangle connect one of the first 4 nodes
 		if (Index < 3)
 		{
+
 			continue;
 		}
-			
+
 		for (Index = 0; Index < 3; Index++)
 		{
 			int32 StartVertex = TriangleSet[TriangleIndex].VertexIndex[Index];
@@ -2653,7 +3722,7 @@
 			{
 				if (EdgeVertexIndices[Endex] == EndVertex && EdgeVertexIndices[Endex + 1] == StartVertex)
 				{
-					OuterEdgeIndices[Endex/2] = 0;
+					OuterEdgeIndices[Endex / 2] = 0;
 					break;
 				}
 			}
@@ -2687,34 +3756,43 @@
 	{
 		if (NodeA.GetSegmentConnectedTo(&NodeB))
 		{
+			DisplaySegment(NodeA.Get2DPoint(EGridSpace::UniformScaled, Grid), NodeB.Get2DPoint(EGridSpace::UniformScaled, Grid));
 			return;
 		}
 
 		if (InnerSegmentsIntersectionTool.DoesIntersect(ACoordinates, BCoordinates))
 		{
+			DisplaySegment(NodeA.Get2DPoint(EGridSpace::UniformScaled, Grid), NodeB.Get2DPoint(EGridSpace::UniformScaled, Grid));
 			return;
 		}
 
 		if (InnerToLoopSegmentsIntersectionTool.DoesIntersect(NodeA, NodeB))
 		{
+			DisplaySegment(NodeA.Get2DPoint(EGridSpace::UniformScaled, Grid), NodeB.Get2DPoint(EGridSpace::UniformScaled, Grid), 0, EVisuProperty::YellowCurve);
 			return;
 		}
 
 		if (LoopSegmentsIntersectionTool.DoesIntersect(NodeA, NodeB))
 		{
+			DisplaySegment(NodeA.Get2DPoint(EGridSpace::UniformScaled, Grid), NodeB.Get2DPoint(EGridSpace::UniformScaled, Grid), 0, EVisuProperty::YellowCurve);
 			return;
 		}
 
 		// Is Outside and not to flat at NodeA
 		const double FlatAngle = 1;
-		if (!IsPointPIsInsideSectorABC(NodeA.GetPreviousNode().Get2DPoint(EGridSpace::Default2D, Grid), ACoordinates, NodeA.GetNextNode().Get2DPoint(EGridSpace::Default2D, Grid), BCoordinates, FlatAngle))
+		//if (!IsPointPBeInsideSectorABC(NodeA.GetPreviousNode().Get2DPoint(EGridSpace::Default2D, Grid), ACoordinates, NodeA.GetNextNode().Get2DPoint(EGridSpace::Default2D, Grid), BCoordinates, FlatAngle))
+
+		if (NodeA.IsSegmentBeInsideFace(BCoordinates, Grid, FlatAngle))
 		{
+			DisplaySegment(NodeA.Get2DPoint(EGridSpace::UniformScaled, Grid), NodeB.Get2DPoint(EGridSpace::UniformScaled, Grid), 0, EVisuProperty::YellowCurve);
 			return;
 		}
 
 		// Is Outside and not to flat at NodeB
-		if (!IsPointPIsInsideSectorABC(NodeB.GetPreviousNode().Get2DPoint(EGridSpace::Default2D, Grid), BCoordinates, NodeB.GetNextNode().Get2DPoint(EGridSpace::Default2D, Grid), ACoordinates, FlatAngle))
+		//if (!IsPointPBeInsideSectorABC(NodeB.GetPreviousNode().Get2DPoint(EGridSpace::Default2D, Grid), BCoordinates, NodeB.GetNextNode().Get2DPoint(EGridSpace::Default2D, Grid), ACoordinates, FlatAngle))
+		if (NodeB.IsSegmentBeInsideFace(ACoordinates, Grid, FlatAngle))
 		{
+			DisplaySegment(NodeA.Get2DPoint(EGridSpace::UniformScaled, Grid), NodeB.Get2DPoint(EGridSpace::UniformScaled, Grid), 0, EVisuProperty::YellowCurve);
 			return;
 		}
 
@@ -2734,7 +3812,7 @@
 		double MinDistanceSquare = HUGE_VALUE_SQUARE;
 		int32 MinIndexA = -1;
 		int32 MinIndexB = -1;
-		for (int32 IndexA = StartIndexLoopA; IndexA < LoopNodes.Num() - 1; ++IndexA)
+		for (int32 IndexA = StartIndexLoopA; IndexA < LoopNodeCount - 1; ++IndexA)
 		{
 			FLoopNode& NodeA = LoopNodes[IndexA];
 			if (NodeA.GetLoopIndex() != IndexLoopA)
@@ -2744,7 +3822,7 @@
 
 			const FPoint2D& ACoordinates = NodeA.Get2DPoint(EGridSpace::UniformScaled, Grid);
 
-			for (int32 IndexB = StartIndexLoopB; IndexB < LoopNodes.Num(); ++IndexB)
+			for (int32 IndexB = StartIndexLoopB; IndexB < LoopNodeCount; ++IndexB)
 			{
 				FLoopNode& NodeB = LoopNodes[IndexB];
 				if (NodeB.GetLoopIndex() != IndexLoopB)
@@ -2766,9 +3844,12 @@
 		if (MinIndexA >= 0 && MinIndexB >= 0)
 		{
 			FLoopNode& NodeA = LoopNodes[MinIndexA];
-			const FPoint2D& ACoordinates = NodeA.Get2DPoint(EGridSpace::Default2D, Grid);
+			const FPoint2D& ACoordinates = NodeA.Get2DPoint(EGridSpace::UniformScaled, Grid);
 			FLoopNode& NodeB = LoopNodes[MinIndexB];
-			const FPoint2D& BCoordinates = NodeB.Get2DPoint(EGridSpace::Default2D, Grid);
+			const FPoint2D& BCoordinates = NodeB.Get2DPoint(EGridSpace::UniformScaled, Grid);
+
+			//DisplaySegment(NodeA.Get2DPoint(EGridSpace::UniformScaled, Grid), NodeB.Get2DPoint(EGridSpace::UniformScaled, Grid));
+
 			TryCreateSegment(NodeA, ACoordinates, NodeB, BCoordinates);
 		}
 
@@ -2781,13 +3862,19 @@
 #endif
 		for (int32 Index = 0; Index < EdgeVertexIndices.Num(); Index += 2)
 		{
+			//Open3DDebugSession(TEXT("Segment"));
 			int32 IndexLoopA = EdgeVertexIndices[Index];
 			int32 IndexLoopB = EdgeVertexIndices[Index + 1];
 
+			//DisplaySegment(LoopBarycenters[IndexLoopA], LoopBarycenters[IndexLoopB]);
+
 			FindAndTryCreateCandidateSegmentToLinkLoops(IndexLoopA, IndexLoopB);
+			//Close3DDebugSession();
+			//Wait();
 		}
 #ifdef DEBUG_DELAUNAY
 		Close3DDebugSession();
+		Wait();
 #endif
 	}
 
@@ -2813,248 +3900,4 @@
 #endif
 }
 
-// =========================================================================================================================================================================================================
-// =========================================================================================================================================================================================================
-// =========================================================================================================================================================================================================
-//
-//
-//                                                                            NOT YET REVIEWED
-//
-//
-// =========================================================================================================================================================================================================
-// =========================================================================================================================================================================================================
-// =========================================================================================================================================================================================================
-
-
-#ifdef NEED_TO_CHECK_USEFULNESS
-void FIsoTriangulator::CompleteIsoSegmentLoopToLoop()
-{
-	FIntersectionSegmentTool NewSegmentsIntersectionTool(Grid);
-	// similare to FindSegmentToLinkloopToloop::CreateSegment but different 
-	TFunction<void(FLoopNode&, FIsoNode&)> CreateComplementarySegment = [&](FLoopNode& Node1, FIsoNode& Node2)
-	{
-		if (Node1.GetSegmentConnectedTo(&Node2))
-		{
-			return;
-		}
-
-		if (NewSegmentsIntersectionTool.DoesIntersect(Node1, Node2))
-		{
-			return;
-		}
-
-		if (LoopSegmentsIntersectionTool.DoesIntersect(Node1, Node2))
-		{
-			return;
-		}
-
-		FIsoSegment& Segment = IsoSegmentFactory.New();
-		Segment.Init(Node1, Node2, ESegmentType::Unknown);
-		Segment.ConnectToNode();
-		FinalToLoops.Add(&Segment);
-		InnerToLoopSegmentsIntersectionTool.AddSegment(Segment);
-		NewSegmentsIntersectionTool.AddSegment(Segment);
-		NewSegmentsIntersectionTool.Sort();
-	};
-
-	for (FLoopNode& LoopNode : LoopNodes)
-	{
-		if (LoopNode.GetConnectedSegments().Num() == 2)
-		{
-			if (LoopNode.GetPreviousNode().GetConnectedSegments().Num() == 3)
-			{
-				if ((LoopNode.GetPreviousNode().GetConnectedSegments().Last()->GetType() == ESegmentType::IsoLoopToLoop) && (LoopNode.GetNextNode().GetConnectedSegments().Last()->GetType() == ESegmentType::IsoLoopToLoop))
-				{
-					FIsoNode& Node1 = LoopNode.GetPreviousNode().GetConnectedSegments().Last()->GetOtherNode(&LoopNode.GetPreviousNode());
-					DisplaySegment(Node1.Get2DPoint(EGridSpace::Default2D, Grid), LoopNode.Get2DPoint(EGridSpace::Default2D, Grid));
-					CreateComplementarySegment(LoopNode, Node1);
-					FIsoNode& Node2 = LoopNode.GetNextNode().GetConnectedSegments().Last()->GetOtherNode(&LoopNode.GetNextNode());
-					CreateComplementarySegment(LoopNode, Node2);
-					DisplaySegment(Node2.Get2DPoint(EGridSpace::Default2D, Grid), LoopNode.Get2DPoint(EGridSpace::Default2D, Grid));
-				}
-				else if (LoopNode.GetPreviousNode().GetConnectedSegments().Last()->GetType() == ESegmentType::IsoLoopToLoop)
-				{
-					FIsoNode& Node = LoopNode.GetPreviousNode().GetConnectedSegments().Last()->GetOtherNode(&LoopNode.GetPreviousNode());
-					CreateComplementarySegment(LoopNode, Node);
-					DisplaySegment(Node.Get2DPoint(EGridSpace::Default2D, Grid), LoopNode.Get2DPoint(EGridSpace::Default2D, Grid));
-				}
-				else if (LoopNode.GetNextNode().GetConnectedSegments().Last()->GetType() == ESegmentType::IsoLoopToLoop)
-				{
-					FIsoNode& Node = LoopNode.GetNextNode().GetConnectedSegments().Last()->GetOtherNode(&LoopNode.GetNextNode());
-					CreateComplementarySegment(LoopNode, Node);
-					DisplaySegment(Node.Get2DPoint(EGridSpace::Default2D, Grid), LoopNode.Get2DPoint(EGridSpace::Default2D, Grid));
-				}
-			}
-		}
-	}
-}
 #endif
-
-#ifdef UNUSED
-void FIsoTriangulator::FindSegmentToLinkLoopToLoop()
-{
-	FTimePoint StartTime = FChrono::Now();
-
-	double DiagonalMax = Grid.GetMaxDeltaU(EIso::IsoU) + Grid.GetMaxDeltaU(EIso::IsoV);
-	DiagonalMax *= 2;
-
-	TFunction<double(const TArray<double>&, int32&, const double&)> GetDeltaCoordinate = [](const TArray<double>& Iso, int32& Index, const double& PointCoord)
-	{
-		while (PointCoord + SMALL_NUMBER < Iso[Index])
-		{
-			if (Index == 0)
-			{
-				break;
-			}
-			Index--;
-		}
-		for (; Index < Iso.Num(); ++Index)
-		{
-			if (Index == Iso.Num() - 2)
-			{
-				break;
-			}
-			if (PointCoord < Iso[Index + 1])
-			{
-				break;
-			}
-		}
-
-		// If PointCoord is closer Iso[Index + 1]
-		if ((Iso[Index + 1] - PointCoord) * 2 < (Iso[Index + 1] - Iso[Index]))
-		{
-			if (Index + 2 < Iso.Num())
-			{
-				return FMath::Max(Iso[Index + 1] - Iso[Index], Iso[Index + 2] - Iso[Index + 1]);
-			}
-			else
-			{
-				return Iso[Index + 1] - Iso[Index];
-			}
-		}
-		else
-		{
-			if (Index > 1)
-			{
-				return FMath::Max(Iso[Index] - Iso[Index - 1], Iso[Index + 1] - Iso[Index]);
-			}
-			else
-			{
-				return Iso[Index + 1] - Iso[Index];
-			}
-		}
-		return Iso[Index + 1] - Iso[Index];
-	};
-
-	TFunction<void(FLoopNode*, const FPoint2D&, FLoopNode*, const FPoint2D&)> CreateSegment = [this](FLoopNode* NodeA, const FPoint2D& ACoordinates, FLoopNode* NodeB, const FPoint2D& BCoordinates)
-	{
-		if (NodeA->GetSegmentConnectedTo(NodeB))
-		{
-			return;
-		}
-
-		// IsExiste
-		if (NodeA->GetLoopIndex() == NodeB->GetLoopIndex())
-		{
-			if ((NodeA->GetIndex() == NodeB->GetIndex() + 1) || (NodeB->GetIndex() == NodeA->GetIndex() + 1))
-			{
-				return;
-			}
-		}
-
-		if (InnerSegmentsIntersectionTool.DoesIntersect(ACoordinates, BCoordinates))
-		{
-			return;
-		}
-
-		if (InnerToLoopSegmentsIntersectionTool.DoesIntersect(*NodeA, *NodeB))
-		{
-			return;
-		}
-
-		if (LoopSegmentsIntersectionTool.DoesIntersect(*NodeA, *NodeB))
-		{
-			return;
-		}
-
-		// Is Outside and not to flat at NodeA
-		const double FlatAngle = 1;
-		if (!IsPointPIsInsideSectorABC(NodeA->GetPreviousNode().Get2DPoint(EGridSpace::Default2D, Grid), ACoordinates, NodeA->GetNextNode().Get2DPoint(EGridSpace::Default2D, Grid), BCoordinates, FlatAngle))
-		{
-			return;
-		}
-
-		//double SlopWithNextloop = GetZeroTwoPiSlope(ACoordinates, NodeA->GetNextNode().Get2DPoint(EGridSpace::Default2D, Grid), 0);
-		//double loopDeltaSlop = GetZeroTwoPiSlope(ACoordinates, NodeA->GetPreviousNode().Get2DPoint(EGridSpace::Default2D, Grid), SlopWithNextloop);
-		//double SegmentSlop = GetZeroTwoPiSlope(ACoordinates, BCoordinates, SlopWithNextloop);
-		//if (SegmentSlop < FlatAngle || SegmentSlop + FlatAngle > loopDeltaSlop)
-		//{
-		//	return;
-		//}
-
-		// Is Outside and not to flat at NodeB
-		if (!IsPointPIsInsideSectorABC(NodeB->GetPreviousNode().Get2DPoint(EGridSpace::Default2D, Grid), BCoordinates, NodeB->GetNextNode().Get2DPoint(EGridSpace::Default2D, Grid), ACoordinates, FlatAngle))
-		{
-			return;
-		}
-
-		//SlopWithNextloop = GetZeroTwoPiSlope(BCoordinates, NodeB->GetNextNode().Get2DPoint(EGridSpace::Default2D, Grid), 0);
-		//loopDeltaSlop = GetZeroTwoPiSlope(BCoordinates, NodeB->GetPreviousNode().Get2DPoint(EGridSpace::Default2D, Grid), SlopWithNextloop);
-		//SegmentSlop = GetZeroTwoPiSlope(BCoordinates, ACoordinates, SlopWithNextloop);
-
-		//if (SegmentSlop < FlatAngle || SegmentSlop + FlatAngle > loopDeltaSlop)
-		//{
-		//	return;
-		//}
-
-		FIsoSegment& Segment = IsoSegmentFactory.New();
-		Segment.Init(*NodeA, *NodeB, ESegmentType::LoopToLoop);
-		CandidateLoopToLoopSegments.Add(&Segment);
-	};
-
-	//Open3DDebugSession(TEXT("FIsoTrianguler::FindSegmentToLinkloopToloop"));
-
-	// find closed points
-	int32 IndexU = 0;
-	int32 IndexV = 0;
-
-	for (int32 IndexA = 0; IndexA < SortedLoopNodes.Num() - 1; ++IndexA)
-	{
-		FLoopNode* NodeA = SortedLoopNodes[IndexA];
-		if (NodeA->GetConnectedSegments().Num() > 2)
-		{
-			// Enough connected
-			continue;
-		}
-
-		// check intersection with
-		const FPoint2D& ACoordinates = NodeA->Get2DPoint(EGridSpace::Default2D, Grid);
-
-		double DeltaUMax = GetDeltaCoordinate(Grid.GetCuttingCoordinatesAlongIso(EIso::IsoU), IndexU, ACoordinates.U);
-		//double DeltaUMin = DeltaUMax * 0.1;
-		DeltaUMax *= 1.75;
-		double DeltaVMax = GetDeltaCoordinate(Grid.GetCuttingCoordinatesAlongIso(EIso::IsoV), IndexV, ACoordinates.V);
-		//double DeltaVMin = DeltaVMax * 0.1;
-		DeltaVMax *= 1.75;
-
-		for (int32 IndexB = IndexA + 1; IndexB < SortedLoopNodes.Num(); ++IndexB)
-		{
-			FLoopNode* NodeB = SortedLoopNodes[IndexB];
-			const FPoint2D& BCoordinates = NodeB->Get2DPoint(EGridSpace::Default2D, Grid);
-			FPoint2D AB = BCoordinates - ACoordinates;
-			double Diagonal = FMath::Abs(AB.U + AB.V);
-			if (Diagonal > DiagonalMax)
-			{
-				break;
-			}
-
-			if (FMath::Abs(AB.U) < DeltaUMax && FMath::Abs(AB.V) < DeltaVMax)
-			{
-				CreateSegment(NodeA, ACoordinates, NodeB, BCoordinates);
-			}
-		}
-	}
-	//Close3DDebugSession();
-	Chronos.FindSegmentToLinkLoopToLoopDuration += FChrono::Elapse(StartTime);
-}
-#endif
diff --git a/Engine/Source/Runtime/Datasmith/CADKernel/Private/CADKernel/Mesh/Meshers/IsoTriangulator.h b/Engine/Source/Runtime/Datasmith/CADKernel/Private/CADKernel/Mesh/Meshers/IsoTriangulator.h
--- a/Engine/Source/Runtime/Datasmith/CADKernel/Private/CADKernel/Mesh/Meshers/IsoTriangulator.h
+++ b/Engine/Source/Runtime/Datasmith/CADKernel/Private/CADKernel/Mesh/Meshers/IsoTriangulator.h
@@ -13,12 +13,13 @@
 #include "CADKernel/Mesh/Meshers/IsoTriangulator/DefineForDebug.h"
 
 //#define NEED_TO_CHECK_USEFULNESS
-
+#define DEBUG_DELAUNAY
 namespace CADKernel
 {
 	class FGrid;
 	class FIntersectionSegmentTool;
 	class FFaceMesh;
+	struct FCell;
 
 	typedef TFunction<double(const FPoint2D&, const FPoint2D&, double)> FGetSlop;
 
@@ -110,6 +111,7 @@ namespace CADKernel
 
 		TArray<int32> LoopStartIndex;
 		TArray<FLoopNode> LoopNodes;
+		int32 LoopNodeCount = 0;
 		TArray<FLoopNode*> SortedLoopNodes;
 
 		/**
@@ -130,7 +132,7 @@ namespace CADKernel
 		TArray<FIsoSegment*> LoopSegments;
 		TArray<FIsoSegment*> ThinZoneSegments;
 		TArray<FIsoSegment*> FinalInnerSegments;
-
+		TArray<FIsoSegment*> InnerToOuterSegments;
 
 		/**
 		 * Tool to check if a segment intersect or not existing segments.
@@ -140,6 +142,7 @@ namespace CADKernel
 		FIntersectionSegmentTool LoopSegmentsIntersectionTool;
 		FIntersectionSegmentTool InnerSegmentsIntersectionTool;
 		FIntersectionSegmentTool InnerToLoopSegmentsIntersectionTool;
+		FIntersectionSegmentTool InnerToOuterSegmentsIntersectionTool;
 
 
 		/**
@@ -171,6 +174,8 @@ namespace CADKernel
 		 */
 		TArray<FIsoSegment*> CandidateInnerToLoopSegments;
 
+		TArray<FIsoSegment*> NewTestSegments;
+
 		bool bDisplay = false;
 
 	public:
@@ -237,6 +242,38 @@ namespace CADKernel
 		void CompleteIsoSegmentLoopToLoop();
 #endif
 
+		void ConnectCellLoops();
+		void FindCellContainingBoundaryNodes(TArray<FCell>& Cells);
+
+		/**
+		 * The closest loops are connected together
+		 * To do it, a Delaunay triangulation of the loop barycenters is realized.
+		 * Each edge of this mesh defined a near loops pair
+		 * The shortest segment is then build between this pair of loops
+		 */
+		void ConnectCellLoopsByNeighborhood(FCell& cell);
+
+		/**
+		 * 2nd step
+		 */
+		void ConnectCellCornerToInnerLoop(FCell& Cell);
+
+		/**
+		 * 3rd step : IsoSegment linking 
+		 */
+		void FindIsoSegmentToLinkOuterLoopNodes(FCell& Cell);
+		void FindIsoSegmentToLinkOuterLoopNodes2(FCell& Cell);
+
+		/**
+		 * 4th step : If their is no segment candidate  
+		 */
+		void FindSegmentToLinkOuterLoopNodes(FCell& Cell);
+
+		/**
+		 * 5th step : try to connect outer loop Extremities
+		 */
+		void AddSementToLinkOuterLoopExtremities(FCell& Cell);
+
 		/**
 		 * The goal of this algorithm is to connect iso U (or V) aligned loop nodes as soon as they are nearly in the same iso V (or U) strip.
 		 * I.e.:
@@ -244,6 +281,7 @@ namespace CADKernel
 		 * - In the same strip: each node of the segment has the same index "i" that verify: isoV[i] - TolV < Node.V < isoV[i+1] + TolV
 		 */
 		void FindIsoSegmentToLinkLoopToLoop();
+		void LastChanceToCreateSegmentInCell(FCell& Cell);
 
 		/**
 		 * The goal of this algorithm is to connect inner node (node of the grid UV) to iso aligned loop node when they are in the same iso V (or U) strip.
@@ -259,7 +297,7 @@ namespace CADKernel
 		 * The purpose of the method is select a minimal set of segments connecting loops together
 		 * The final segments will be selected with SelectSegmentInCandidateSegments
 		 */
-		void Delaunay();  // to rename and clean
+		void ConnectCellLoopsByNeighborhood();  // to rename and clean
 
 #ifdef TODELETE
 		void FindSegmentToLinkLoopToLoop();
@@ -271,6 +309,9 @@ namespace CADKernel
 		 */
 		void FindCandidateSegmentsToLinkInnerToLoop();
 
+
+		void FindCandidateSegmentsToLinkInnerAndLoop();
+
 		/**
 		 * Complete the final set of segments with the best subset of segments in CandidateSegments
 		 */
@@ -330,26 +371,14 @@ namespace CADKernel
 
 		//void CreateLoopToLoopSegment(FLoopNode* NodeA, const FPoint2D& ACoordinates, FLoopNode* NodeB, const FPoint2D& BCoordinates, TArray<FIsoSegment*>& NewSegments);
 
-		struct DelaunayTriangle
-		{
-			int32 VertexIndex[3];
-			double SquareRadius;
-			FPoint Center;
-
-			DelaunayTriangle(const int32& Index0, const int32& Index1, const int32& Index2, const TArray<FPoint2D>& Vertices)
-			{
-				Set(Index0, Index1, Index2, Vertices);
-			}
-
-			void Set(const int32& Index0, const int32& Index1, const int32& Index2, const TArray<FPoint2D>& Vertices)
-			{
-				VertexIndex[0] = Index0;
-				VertexIndex[1] = Index1;
-				VertexIndex[2] = Index2;
-				FTriangle2D Triangle(Vertices[Index0], Vertices[Index1], Vertices[Index2]);
-				Center = Triangle.CircumCircleCenterWithSquareRadius(SquareRadius);
-			}
-		};
+		// ==========================================================================================
+		// 	   Create segments
+		// ==========================================================================================
+		void TryToConnectTwoLoopsWithShortestSegment(FCell& Cell, int32 IndexLoopA, int32 IndexLoopB);
+		void TryToConnectTwoLoopsWithShortestSegment(FCell& Cell, const TArray<FLoopNode*>& LoopA, int32 IndexLoopB);
+		void TryToConnectTwoLoopsWithShortestSegment(FCell& Cell, const TArray<FLoopNode*>& LoopA, const TArray<FLoopNode*>& LoopB);
+		void TryToConnectTwoLoopsWithTheMostIsoSegment(FCell& Cell, const TArray<FLoopNode*>& LoopA, const TArray<FLoopNode*>& LoopB);
+		void TryToCreateSegment(FCell& Cell, FLoopNode* NodeA, const FPoint2D& ACoordinates, FIsoNode* NodeB, const FPoint2D& BCoordinates, const double FlatAngle);
 
 #ifdef CADKERNEL_DEV
 	public:
@@ -366,6 +395,9 @@ namespace CADKernel
 		void Display(EGridSpace Space, const FIsoSegment& Segment, FIdent Ident = 0, EVisuProperty Property = EVisuProperty::Element, bool bDisplayOrientation = false) const;
 		void Display(EGridSpace Space, const TCHAR* Message, const TArray<FIsoSegment*>& Segment, bool bDisplayNode, bool bDisplayOrientation = false, EVisuProperty Property = EVisuProperty::Element) const;
 
+		void DisplayCells(const TArray<FCell>& Cells) const;
+		void DisplayCell(const FCell& Cell) const;
+
 		void DisplayTriangle(EGridSpace Space, const FIsoNode& NodeA, const FIsoNode& NodeB, const FIsoNode& NodeC) const;
 #endif
 
diff --git a/Engine/Source/Runtime/Datasmith/CADKernel/Private/CADKernel/Mesh/Meshers/IsoTriangulator/DefineForDebug.h b/Engine/Source/Runtime/Datasmith/CADKernel/Private/CADKernel/Mesh/Meshers/IsoTriangulator/DefineForDebug.h
--- a/Engine/Source/Runtime/Datasmith/CADKernel/Private/CADKernel/Mesh/Meshers/IsoTriangulator/DefineForDebug.h
+++ b/Engine/Source/Runtime/Datasmith/CADKernel/Private/CADKernel/Mesh/Meshers/IsoTriangulator/DefineForDebug.h
@@ -3,8 +3,10 @@
 #pragma once
 
 
-static int32 SurfaceToDebug = 4;
-#define SURFACE_TO_DEBUG SurfaceToDebug
+static int32 FaceToDebug = 0;
+#define DEBUG_ONLY_SURFACE_TO_DEBUG
+
+//#define SURFACE_TO_DEBUG SurfaceToDebug
 //#define SURFACE_TO_DEBUG 0
 static bool BoolDisplayDebugMeshStep=false;
 
diff --git a/Engine/Source/Runtime/Datasmith/CADKernel/Private/CADKernel/Mesh/Meshers/IsoTriangulator/IntersectionSegmentTool.cpp b/Engine/Source/Runtime/Datasmith/CADKernel/Private/CADKernel/Mesh/Meshers/IsoTriangulator/IntersectionSegmentTool.cpp
--- a/Engine/Source/Runtime/Datasmith/CADKernel/Private/CADKernel/Mesh/Meshers/IsoTriangulator/IntersectionSegmentTool.cpp
+++ b/Engine/Source/Runtime/Datasmith/CADKernel/Private/CADKernel/Mesh/Meshers/IsoTriangulator/IntersectionSegmentTool.cpp
@@ -17,15 +17,15 @@ bool FIntersectionSegmentTool::DoesIntersect(const FIsoSegment& Segment) const
 
 bool FIntersectionSegmentTool::DoesIntersect(const FIsoNode& StartNode, const FIsoNode& EndNode) const
 {
-	const FPoint2D& StartPoint = StartNode.Get2DPoint(EGridSpace::Default2D, Grid);
-	const FPoint2D& EndPoint = EndNode.Get2DPoint(EGridSpace::Default2D, Grid);
+	const FPoint2D& StartPoint = StartNode.Get2DPoint(EGridSpace::UniformScaled, Grid);
+	const FPoint2D& EndPoint = EndNode.Get2DPoint(EGridSpace::UniformScaled, Grid);
 	const FSegment4IntersectionTools StartEndSegment(StartPoint, EndPoint);
 
 	double AxisMax = StartEndSegment.Boundary[EIso::IsoU].Max + StartEndSegment.Boundary[EIso::IsoV].Max;
 
 	for (const FSegment4IntersectionTools& Segment : Segments)
 	{
-		if (AxisMax < Segment.AxisMin)
+		if (bSegmentsAreSorted && AxisMax < Segment.AxisMin)
 		{
 			break;
 		}
@@ -51,14 +51,14 @@ bool FIntersectionSegmentTool::DoesIntersect(const FIsoNode& StartNode, const FI
 
 bool FIntersectionSegmentTool::DoesIntersect(const FIsoNode& StartNode, const FPoint2D& EndPoint) const
 {
-	const FPoint2D& StartPoint = StartNode.Get2DPoint(EGridSpace::Default2D, Grid);
+	const FPoint2D& StartPoint = StartNode.Get2DPoint(EGridSpace::UniformScaled, Grid);
 	const FSegment4IntersectionTools StartEndSegment(StartPoint, EndPoint);
 
 	double DMax = StartEndSegment.Boundary[EIso::IsoU].Max + StartEndSegment.Boundary[EIso::IsoV].Max;
 
 	for (const FSegment4IntersectionTools& Segment : Segments)
 	{
-		if (DMax < Segment.AxisMin)
+		if (bSegmentsAreSorted && DMax < Segment.AxisMin)
 		{
 			break;
 		}
@@ -89,7 +89,7 @@ bool FIntersectionSegmentTool::DoesIntersect(const FPoint2D& StartPoint, const F
 
 	for (const FSegment4IntersectionTools& Segment : Segments)
 	{
-		if (DMax < Segment.AxisMin)
+		if (bSegmentsAreSorted && DMax < Segment.AxisMin)
 		{
 			break;
 		}
@@ -108,7 +108,7 @@ bool FIntersectionSegmentTool::DoesIntersect(const FPoint2D& StartPoint, const F
 }
 
 FSegment4IntersectionTools::FSegment4IntersectionTools(const FGrid& Grid, const FIsoSegment& InSegment)
-	: FSegment4IntersectionTools(InSegment.GetFirstNode().Get2DPoint(EGridSpace::Default2D, Grid), InSegment.GetSecondNode().Get2DPoint(EGridSpace::Default2D, Grid))
+	: FSegment4IntersectionTools(InSegment.GetFirstNode().Get2DPoint(EGridSpace::UniformScaled, Grid), InSegment.GetSecondNode().Get2DPoint(EGridSpace::UniformScaled, Grid))
 {
 	IsoSegment = &InSegment;
 	AxisMin = Boundary[EIso::IsoU].Min + Boundary[EIso::IsoV].Min;
diff --git a/Engine/Source/Runtime/Datasmith/CADKernel/Private/CADKernel/Mesh/Meshers/IsoTriangulator/IntersectionSegmentTool.h b/Engine/Source/Runtime/Datasmith/CADKernel/Private/CADKernel/Mesh/Meshers/IsoTriangulator/IntersectionSegmentTool.h
--- a/Engine/Source/Runtime/Datasmith/CADKernel/Private/CADKernel/Mesh/Meshers/IsoTriangulator/IntersectionSegmentTool.h
+++ b/Engine/Source/Runtime/Datasmith/CADKernel/Private/CADKernel/Mesh/Meshers/IsoTriangulator/IntersectionSegmentTool.h
@@ -32,6 +32,9 @@ namespace CADKernel
 
 		FSegment4IntersectionTools(const FGrid& Grid, const FIsoSegment& InSegment);
 
+		/**
+		 * WARNING StartPoint, EndPoint must be defined in EGridSpace::UniformScaled
+		 */
 		FSegment4IntersectionTools(const FPoint2D& StartPoint, const FPoint2D& EndPoint)
 			: Segment2D(StartPoint, EndPoint)
 			, Boundary(StartPoint, EndPoint)
@@ -62,11 +65,19 @@ namespace CADKernel
 	{
 		TArray<FSegment4IntersectionTools> Segments;
 		const FGrid& Grid;
+		bool bSegmentsAreSorted;
 
 	public:
 		FIntersectionSegmentTool(const FGrid& InGrid)
 			: Grid(InGrid)
+			, bSegmentsAreSorted(false)
+		{
+		}
+
+		void Empty(int32 InMaxNum)
 		{
+			Segments.Empty(InMaxNum);
+			bSegmentsAreSorted = false;
 		}
 
 		void Reserve(int32 InMaxNum)
@@ -79,13 +90,25 @@ namespace CADKernel
 			Segments.Reserve(NewMaxNum);
 		}
 
-		int32 Num()
+		int32 Count()
 		{
 			return Segments.Num();
 		}
 
-		void AddSegment(const TArray<FIsoSegment*>& InNewSegments)
+		void AddSegments(FIsoSegment** InNewSegments, int32 Count)
 		{
+			bSegmentsAreSorted = false;
+			Segments.Reserve(Count + Segments.Num());
+			for (int32 Index = 0; Index < Count; ++Index)
+			{
+				FIsoSegment* NewSegment = InNewSegments[Index];
+				AddSegment(*NewSegment);
+			}
+		}
+
+		void AddSegments(const TArray<FIsoSegment*>& InNewSegments)
+		{
+			bSegmentsAreSorted = false;
 			Segments.Reserve(InNewSegments.Num() + Segments.Num());
 			for (FIsoSegment* Segment : InNewSegments)
 			{
@@ -95,15 +118,26 @@ namespace CADKernel
 
 		void AddSegment(const FIsoSegment& Segment)
 		{
+			bSegmentsAreSorted = false;
 			Segments.Emplace(Grid, Segment);
 		}
 
+		void AddSegment(const FPoint2D& StartPoint, const FPoint2D& EndPoint)
+		{
+			bSegmentsAreSorted = false;
+			Segments.Emplace(StartPoint, EndPoint);
+		}
+
 		bool DoesIntersect(const FIsoSegment& Segment) const;
 
+		/**
+		 * WARNING StartPoint, EndPoint must be defined in EGridSpace::UniformScaled
+		 */
 		bool DoesIntersect(const FPoint2D& StartPoint, const FPoint2D& EndPoint) const;
 
 		/**
 		 * Allow StartNode to be connected to one segment
+		 * WARNING EndPoint must be defined in EGridSpace::UniformScaled
 		 */
 		bool DoesIntersect(const FIsoNode& StartNode, const FPoint2D& EndPoint) const;
 
@@ -118,6 +152,7 @@ namespace CADKernel
 		void Sort()
 		{
 			Algo::Sort(Segments, [](const FSegment4IntersectionTools& Segment1, const FSegment4IntersectionTools& Segment2) { return Segment1.AxisMin < Segment2.AxisMin; });
+			bSegmentsAreSorted = true;
 		}
 
 #ifdef CADKERNEL_DEV
diff --git a/Engine/Source/Runtime/Datasmith/CADKernel/Private/CADKernel/Mesh/Meshers/IsoTriangulator/IsoCell.h b/Engine/Source/Runtime/Datasmith/CADKernel/Private/CADKernel/Mesh/Meshers/IsoTriangulator/IsoCell.h
new file mode 100644
--- /dev/null
+++ b/Engine/Source/Runtime/Datasmith/CADKernel/Private/CADKernel/Mesh/Meshers/IsoTriangulator/IsoCell.h
@@ -0,0 +1,136 @@
+// Copyright Epic Games, Inc. All Rights Reserved.
+#pragma once
+
+#include "CADKernel/Core/Types.h"
+#include "CADKernel/Mesh/MeshEnum.h"
+#include "CADKernel/Mesh/Meshers/IsoTriangulator/IsoNode.h"
+#include "CADKernel/Mesh/Meshers/IsoTriangulator/IntersectionSegmentTool.h"
+#include "CADKernel/Mesh/Structure/Grid.h"
+#include "CADKernel/Topo/TopologicalFace.h"
+
+namespace CADKernel
+{
+	class FGrid;
+
+	struct FCell
+	{
+		const FGrid& Grid;
+		int32 Id;
+		int32 LoopNodeCount;
+		bool bHasOuterLoop = false;
+		int32 InnerLoopCount = 0;
+
+		TArray<int32> LoopIndexToIndex;
+		TArray<TArray<FLoopNode*>> Loops;
+
+		// the outer loop is subdivide into connected (in the cell) nodes
+		TArray<TArray<FLoopNode*>> OuterLoopSubdivision;
+
+		TArray<FIsoSegment*> CandidateSegments;
+		TArray<FIsoSegment*> FinalSegments;
+		FIntersectionSegmentTool IntersectionTool;
+
+		/** Loop that barycenter is on the border of the barycenter mesh */
+		TSet<int32> BorderLoopIndices;
+
+		FCell(const int32 InLoopIndex, const TArray<FLoopNode*>& InNodes, const FGrid& InGrid)
+			: Grid(InGrid)
+			, Id(InLoopIndex)
+			, LoopNodeCount(InNodes.Num())
+			, IntersectionTool(Grid)
+		{
+			ensureCADKernel(InNodes.Num() > 0);
+
+			int32 FaceLoopCount = Grid.GetFace()->LoopCount();
+			int32 LoopCount = 0;
+			{
+				TArray<int32> LoopVertexCount;
+				LoopVertexCount.Init(0, FaceLoopCount);
+				for (const FLoopNode* Node : InNodes)
+				{
+					LoopVertexCount[Node->GetLoopIndex()]++;
+				}
+
+				for (const int32& VertexCount : LoopVertexCount)
+				{
+					if (VertexCount > 0)
+					{
+						LoopCount++;
+					}
+				}
+
+				Loops.SetNum(LoopCount);
+				int32 Index = 0;
+				for (int32 LoopIndex = 0; LoopIndex < FaceLoopCount; ++LoopIndex)
+				{
+					if (LoopVertexCount[LoopIndex] > 0)
+					{
+						Loops[Index].Reserve(LoopVertexCount[LoopIndex]);
+						LoopVertexCount[LoopIndex] = Index;
+						Index++;
+					}
+				}
+				Swap(LoopIndexToIndex, LoopVertexCount);
+			}
+
+			for (FLoopNode* Node : InNodes)
+			{
+				int32 Index = LoopIndexToIndex[Node->GetLoopIndex()];
+				Loops[Index].Add(Node);
+			}
+
+			bHasOuterLoop = (Loops[0][0]->GetLoopIndex() == 0);
+			InnerLoopCount = bHasOuterLoop ? (Loops.Num() - 1): Loops.Num();
+
+			if (bHasOuterLoop)
+			{
+				// Subdivide Loop0 in SubLoop
+				TArray<FLoopNode*>& Loop0 = Loops[0];
+				Algo::Sort(Loop0, [&](const FLoopNode* LoopNode1, const FLoopNode* LoopNode2)
+					{
+						return LoopNode1->GetIndex() < LoopNode2->GetIndex();
+					});
+
+				OuterLoopSubdivision.Reserve(Loop0.Num());
+				FLoopNode* PreviousNode = nullptr;
+				TArray<FLoopNode*>* SubLoop = nullptr;
+				for (FLoopNode* Node : Loop0)
+				{
+					if (&Node->GetPreviousNode() != PreviousNode)
+					{
+						SubLoop = &OuterLoopSubdivision.Emplace_GetRef();
+					}
+					SubLoop->Add(Node);
+					PreviousNode = Node;
+				}
+			}
+		}
+
+		void SelectSegmentInCandidateSegments(TFactory<FIsoSegment>& SegmentFactory)
+		{
+			Algo::Sort(CandidateSegments, [&](const FIsoSegment* Segment1, const FIsoSegment* Segment2)
+				{
+					return Segment1->Get2DLengthSquare(EGridSpace::UniformScaled, Grid) < Segment2->Get2DLengthSquare(EGridSpace::UniformScaled, Grid);
+				});
+
+			// Validate all candidate segments
+			for (FIsoSegment* Segment : CandidateSegments)
+			{
+				if (IntersectionTool.DoesIntersect(*Segment))
+				{
+					SegmentFactory.DeleteEntity(Segment);
+					continue;
+				}
+
+				FinalSegments.Add(Segment);
+				IntersectionTool.AddSegment(*Segment);
+				Segment->SetSelected();
+				Segment->ConnectToNode();
+			}
+			CandidateSegments.Empty();
+		}
+
+	};
+
+} // namespace CADKernel
+
diff --git a/Engine/Source/Runtime/Datasmith/CADKernel/Private/CADKernel/Mesh/Meshers/IsoTriangulator/IsoNode.h b/Engine/Source/Runtime/Datasmith/CADKernel/Private/CADKernel/Mesh/Meshers/IsoTriangulator/IsoNode.h
--- a/Engine/Source/Runtime/Datasmith/CADKernel/Private/CADKernel/Mesh/Meshers/IsoTriangulator/IsoNode.h
+++ b/Engine/Source/Runtime/Datasmith/CADKernel/Private/CADKernel/Mesh/Meshers/IsoTriangulator/IsoNode.h
@@ -302,6 +302,38 @@ namespace CADKernel
 			return FPoint::ZeroPoint;
 		}
 
+		/**
+		 * Check if the segment starting from this loop node (S) to EndSegmentCoordinate (E) is inside the face
+		 * Warning, inner loops (loop index > 0) are not in the same orientation than outer loop
+		 *
+		 * Inside e.g.:
+		 *    X---------------X----------------X        X------------------X-----------X
+		 *    |                                |        |                              |
+		 *    |         X-----------------X    |        |            X-------------X   |
+		 *    |         |                 |    |        |            |             |   |
+		 *    X   E-----S                 |    |    or  S---------E  |             |   |
+		 *    |         |                 |    |        |            |             |   |
+		 *    |         X-----------------X    |        |            X-------------X   |
+		 *    |                                |        |                              |
+		 *    X---------------X----------------X        X------------------X-----------X
+		 *
+		 * Outside e.g.:
+		 *    X---------------X----------------X        X------------------X-----------X
+		 *    |                                |        |                              |
+		 *    |         X-----------------X    |        |            X-------------X   |
+		 *    |         |                 |    |        |            |             |   |
+		 *    |         S------E          |    |    or  |            |             |   S------E
+		 *    |         |                 |    |        |            |             |   |
+		 *    |         X-----------------X    |        |            X-------------X   |
+		 *    |                                |        |                              |
+		 *    X---------------X----------------X        X------------------X-----------X
+		 */
+		bool IsSegmentBeInsideFace(const FPoint2D& EndSegmentCoordinate, const FGrid& Grid, const double FlatAngle)
+		{
+			//bool bIsOuterLoop = (GetLoopIndex() == 0);
+			return (/*bIsOuterLoop  == */!IsPointPBeInsideSectorABC(GetPreviousNode().Get2DPoint(EGridSpace::UniformScaled, Grid), Get2DPoint(EGridSpace::UniformScaled, Grid), GetNextNode().Get2DPoint(EGridSpace::UniformScaled, Grid), EndSegmentCoordinate, FlatAngle));
+		}
+
 		void SetAsIsoU()
 		{
 			States |= EIsoNodeStates::NearlyIsoU;
diff --git a/Engine/Source/Runtime/Datasmith/CADKernel/Private/CADKernel/Mesh/Meshers/IsoTriangulator/IsoTriangulatorDisplay.cpp b/Engine/Source/Runtime/Datasmith/CADKernel/Private/CADKernel/Mesh/Meshers/IsoTriangulator/IsoTriangulatorDisplay.cpp
--- a/Engine/Source/Runtime/Datasmith/CADKernel/Private/CADKernel/Mesh/Meshers/IsoTriangulator/IsoTriangulatorDisplay.cpp
+++ b/Engine/Source/Runtime/Datasmith/CADKernel/Private/CADKernel/Mesh/Meshers/IsoTriangulator/IsoTriangulatorDisplay.cpp
@@ -3,11 +3,12 @@
 
 #ifdef CADKERNEL_DEV
 
-#include "CADKernel/Mesh/Meshers/IsoTriangulator/IntersectionSegmentTool.h"
 #include "CADKernel/Mesh/Structure/Grid.h"
+#include "CADKernel/Mesh/Meshers/IsoTriangulator/IntersectionSegmentTool.h"
+#include "CADKernel/Mesh/Meshers/IsoTriangulator/IsoCell.h"
+#include "CADKernel/Mesh/Meshers/IsoTriangulator/IsoNode.h"
 #include "CADKernel/Mesh/Meshers/IsoTriangulator/IsoSegment.h"
 #include "CADKernel/UI/Display.h"
-
 using namespace CADKernel;
 
 void FIsoTriangulator::Display(EGridSpace Space, const FIsoSegment& Segment, FIdent Ident, EVisuProperty Property, bool bDisplayOrientation) const
@@ -124,6 +125,38 @@ void FIsoTriangulator::DisplayPixels(TArray<uint8>& Pixel) const
 	Close3DDebugSession();
 }
 
+void FIsoTriangulator::DisplayCells(const TArray<FCell>& Cells) const
+{
+	if (!bDisplay)
+	{
+		return;
+	}
+
+	F3DDebugSession _(TEXT("Cells"));
+	int32 Index = 0;
+	for (const FCell& Cell : Cells)
+	{
+		DisplayCell(Cell);
+	}
+}
+
+void FIsoTriangulator::DisplayCell(const FCell& Cell) const
+{
+	if (!bDisplay)
+	{
+		return;
+	}
+
+	F3DDebugSession _(FString::Printf(TEXT("Cell %d"), Cell.Id));
+	for (const TArray<FLoopNode*>& Nodes : Cell.Loops)
+	{
+		for (const FLoopNode* Node : Nodes)
+		{
+			CADKernel::DisplayPoint(Node->GetPoint(EGridSpace::UniformScaled, Grid), EVisuProperty::OrangePoint, Cell.Id);
+		}
+	}
+}
+
 void FIntersectionSegmentTool::Display(const TCHAR* Message) const
 {
 	if (!Grid.bDisplay)
