CADKernel Source Code v0.02

diff --git a/Engine/Source/Runtime/Datasmith/CADKernel/Private/CADKernel/Mesh/Meshers/ParametricMesher.cpp b/Engine/Source/Runtime/Datasmith/CADKernel/Private/CADKernel/Mesh/Meshers/ParametricMesher.cpp
--- a/Engine/Source/Runtime/Datasmith/CADKernel/Private/CADKernel/Mesh/Meshers/ParametricMesher.cpp
+++ b/Engine/Source/Runtime/Datasmith/CADKernel/Private/CADKernel/Mesh/Meshers/ParametricMesher.cpp
@@ -219,7 +219,7 @@ void FParametricMesher::Mesh(TSharedRef<FTopologicalFace> Face)
 	ensureCADKernel(!Face->IsDeleted());
 	ensureCADKernel(!Face->IsMeshed());
 
-	FMessage::Printf(EVerboseLevel::Log, TEXT("Meshing of surface %d\n"),  Face->GetId());
+	FMessage::Printf(EVerboseLevel::Debug, TEXT("Meshing of surface %d\n"),  Face->GetId());
 
 	FProgress _(1, TEXT("Meshing Entities : Mesh Surface"));
 
@@ -230,11 +230,10 @@ void FParametricMesher::Mesh(TSharedRef<FTopologicalFace> Face)
 
 	FTimePoint StartTime = FChrono::Now();
 
-	FGrid Grid(Face, MeshModel);
-
 	FTimePoint GenerateCloudStartTime = FChrono::Now();
+
+	FGrid Grid(Face, MeshModel);
 	GenerateCloud(Grid);
-	FDuration GenerateCloudDuration = FChrono::Elapse(GenerateCloudStartTime);
 
 	if(Grid.IsDegenerated())
 	{
@@ -247,9 +246,12 @@ void FParametricMesher::Mesh(TSharedRef<FTopologicalFace> Face)
 		return;
 	}
 
-	TSharedRef<FFaceMesh> SurfaceMesh = StaticCastSharedRef<FFaceMesh>(Face->GetOrCreateMesh(MeshModel));
+	FDuration GenerateCloudDuration = FChrono::Elapse(GenerateCloudStartTime);
 
 	FTimePoint IsoTriangulerStartTime = FChrono::Now();
+
+	TSharedRef<FFaceMesh> SurfaceMesh = StaticCastSharedRef<FFaceMesh>(Face->GetOrCreateMesh(MeshModel));
+
 	FIsoTriangulator IsoTrianguler(Grid, SurfaceMesh);
 	if(IsoTrianguler.Triangulate())
 	{
@@ -379,7 +381,7 @@ void FParametricMesher::Mesh(TSharedRef<FTopologicalVertex> InVertex)
 	InVertex->GetOrCreateMesh(GetMeshModel());
 }
 
-void FParametricMesher::Mesh(TSharedRef<FTopologicalEdge> InEdge, TSharedRef<FTopologicalFace> TrimmedSurface)
+void FParametricMesher::Mesh(TSharedRef<FTopologicalEdge> InEdge, TSharedRef<FTopologicalFace> Face)
 {
 	TSharedRef<FTopologicalEdge> ActiveEdge = StaticCastSharedRef<FTopologicalEdge>(InEdge->GetLinkActiveEntity());
 	if (ActiveEdge->IsMeshed())
@@ -396,16 +398,14 @@ void FParametricMesher::Mesh(TSharedRef<FTopologicalEdge> InEdge, TSharedRef<FTo
 		return;
 	}
 
-	double ToleranceGeoEdge = ActiveEdge->GetCurve()->GetParametricTolerance();
-
-	const FSurfacicTolerance& ToleranceIso = TrimmedSurface->GetIsoTolerances();
+	const FSurfacicTolerance& ToleranceIso = Face->GetIsoTolerances();
 
 	// Get Edge intersection with inner surface mesh grid
 	TArray<double> UEdgeSetOfIntersectionWithIsoU;
 	TArray<double> UEdgeSetOfIntersectionWithIsoV;
 
-	const TArray<double>& SurfaceTabU = TrimmedSurface->GetCuttingCoordinatesAlongIso(EIso::IsoU);
-	const TArray<double>& SurfaceTabV = TrimmedSurface->GetCuttingCoordinatesAlongIso(EIso::IsoV);
+	const TArray<double>& SurfaceTabU = Face->GetCuttingCoordinatesAlongIso(EIso::IsoU);
+	const TArray<double>& SurfaceTabV = Face->GetCuttingCoordinatesAlongIso(EIso::IsoV);
 
 	TArray<FPoint2D> EdgeCrossingPoints2D;
 	const TArray<double>& EdgeCrossingPointsU = ActiveEdge->GetCrossingPointUs();
@@ -424,8 +424,8 @@ void FParametricMesher::Mesh(TSharedRef<FTopologicalEdge> InEdge, TSharedRef<FTo
 #endif
 
 #ifdef DEBUG_IntersectEdgeIsos
-	DebugIntersectEdgeIsos(ActiveEdge, TrimmedSurface, EdgeCrossingPoints2D, ToleranceIsoU, SurfaceTabU, EIso::IsoU);
-	DebugIntersectEdgeIsos(ActiveEdge, TrimmedSurface, EdgeCrossingPoints2D, ToleranceIsoV, SurfaceTabV, EIso::IsoV);
+	DebugIntersectEdgeIsos(ActiveEdge, Face, EdgeCrossingPoints2D, ToleranceIsoU, SurfaceTabU, EIso::IsoU);
+	DebugIntersectEdgeIsos(ActiveEdge, Face, EdgeCrossingPoints2D, ToleranceIsoV, SurfaceTabV, EIso::IsoV);
 #endif
 	IntersectEdgeIsos(ActiveEdge, EdgeCrossingPoints2D, ToleranceIso[EIso::IsoU], SurfaceTabU, EIso::IsoU, UEdgeSetOfIntersectionWithIsoU);
 	IntersectEdgeIsos(ActiveEdge, EdgeCrossingPoints2D, ToleranceIso[EIso::IsoV], SurfaceTabV, EIso::IsoV, UEdgeSetOfIntersectionWithIsoV);
@@ -455,6 +455,8 @@ void FParametricMesher::Mesh(TSharedRef<FTopologicalEdge> InEdge, TSharedRef<FTo
 		ImposedIsoCuttingPoints.Reserve(NbImposedCuttingPoints);
 	}
 
+	FPoint2D ExtremityTolerances = ActiveEdge->GetCurve()->GetExtremityTolerances(EdgeBounds);
+	double ToleranceGeoEdge = FMath::Min(ExtremityTolerances[0], ExtremityTolerances[1]);
 	ImposedIsoCuttingPoints.Emplace(EdgeBounds.GetMin(), ECoordinateType::VertexCoordinate, -1, ToleranceGeoEdge);
 
 	double MinDeltaU = HUGE_VALUE;
@@ -1250,10 +1252,11 @@ void FParametricMesher::IntersectEdgeIsos(TSharedPtr<FTopologicalEdge> Edge, con
 		}
 	}
 
-	double ToleranceEdge = Edge->GetCurve()->GetParametricTolerance();
 	const TArray<double>& EdgeCrossingPointsU = Edge->GetCrossingPointUs();
 
-	const FLinearBoundary& EdgeBounds = Edge->GetBoundary();
+	const FLinearBoundary& EdgeBoundary = Edge->GetBoundary();
+	FPoint2D ToleranceEdge = Edge->GetCurve()->GetExtremityTolerances(EdgeBoundary);
+
 	double MinU;
 	double MaxU;
 	for (int32 CrossingPointIndex = 0; CrossingPointIndex < EdgeCrossingPoints2D.Num() - 1; ++CrossingPointIndex)
@@ -1289,7 +1292,8 @@ void FParametricMesher::IntersectEdgeIsos(TSharedPtr<FTopologicalEdge> Edge, con
 				double EdgeCoord = (IsoCoord - EdgeCrossingPoints2D[CrossingPointIndex][(uint8)TypeIso]) / (EdgeCrossingPoints2D[CrossingPointIndex + 1][(uint8)TypeIso] - EdgeCrossingPoints2D[CrossingPointIndex][(uint8)TypeIso]);
 				EdgeCoord *= (EdgeCrossingPointsU[CrossingPointIndex + 1] - EdgeCrossingPointsU[CrossingPointIndex]);
 				EdgeCoord += EdgeCrossingPointsU[CrossingPointIndex];
-				if ((EdgeCoord < (EdgeBounds.GetMin() + ToleranceEdge)) || (EdgeCoord > (EdgeBounds.GetMax() - ToleranceEdge)))
+
+				if ((EdgeCoord < (EdgeBoundary.GetMin() + ToleranceEdge[0])) || (EdgeCoord > (EdgeBoundary.GetMax() - ToleranceEdge[1])))
 				{
 					continue;
 				}
diff --git a/Engine/Source/Runtime/Datasmith/CADKernel/Private/CADKernel/Mesh/Structure/Grid.cpp b/Engine/Source/Runtime/Datasmith/CADKernel/Private/CADKernel/Mesh/Structure/Grid.cpp
--- a/Engine/Source/Runtime/Datasmith/CADKernel/Private/CADKernel/Mesh/Structure/Grid.cpp
+++ b/Engine/Source/Runtime/Datasmith/CADKernel/Private/CADKernel/Mesh/Structure/Grid.cpp
@@ -9,8 +9,9 @@
 #include "CADKernel/Topo/TopologicalEdge.h"
 #include "CADKernel/Topo/TopologicalFace.h"
 #include "CADKernel/Utils/Util.h"
+#include "CADKernel/Utils/ArrayUtils.h"
 
-#define DEBUG_GRID
+//#define DEBUG_GRID
 
 namespace CADKernel
 {
@@ -21,7 +22,9 @@ namespace CADKernel
 		, ThinZoneFinder(*this)
 		, CuttingCoordinates(Face->GetCuttingPointCoordinates())
 	{
-		if (Face->GetId() == SURFACE_TO_DEBUG)
+#ifdef DEBUG_ONLY_SURFACE_TO_DEBUG
+	if (Face->GetId() == FaceToDebug)
+#endif
 		{
 			bDisplay = true;
 		}
@@ -34,7 +37,7 @@ namespace CADKernel
 
 	void FGrid::ProcessPointCloud()
 	{
-		EGridSpace DisplaySpace = EGridSpace::Default2D;
+		EGridSpace DisplaySpace = EGridSpace::UniformScaled;
 		FTimePoint StartTime = FChrono::Now();
 
 		if(!GetMeshOfLoops())
@@ -42,46 +45,40 @@ namespace CADKernel
 			return;
 		}
 
-#ifdef DEBUG_GRID
-		DisplayLoop(TEXT("FGrid::Loop 2D with thin zone"), GetLoops2D(DisplaySpace), false, false);
-		DisplayLoop(TEXT("FGrid::Loop 2D with thin zone"), GetLoops2D(DisplaySpace), true, false);
-		DisplayLoop(TEXT("FGrid::Loop 3D"), GetLoops3D(), true, false);
-		//Wait();
-		DisplayInnerDomainPoints(TEXT("FGrid::Initial PointCloud 2D"), GetInner2DPoints(DisplaySpace));
-#endif
 		ScaleLoops();
 
 #ifdef DEBUG_GRID
-		DisplayInnerDomainPoints(TEXT("FGrid::PointCloud 2D UniformScaled"), GetInner2DPoints(EGridSpace::UniformScaled));
-		DisplayLoop(TEXT("FGrid::Loop 2D UniformScaled"), GetLoops2D(EGridSpace::UniformScaled), true, false);
-		DisplayInnerDomainPoints(TEXT("FGrid::PointCloud 2D Scaled"), GetInner2DPoints(EGridSpace::Scaled));
 		DisplayLoop(TEXT("FGrid::Loop 2D Scaled"), GetLoops2D(EGridSpace::Scaled), true, false);
+		DisplayInnerDomainPoints(TEXT("FGrid::PointCloud 2D Scaled"), GetInner2DPoints(EGridSpace::Scaled));
+
+		DisplayLoop(TEXT("FGrid::Loop 2D UniformScaled"), GetLoops2D(EGridSpace::UniformScaled), true, false);
+		DisplayInnerDomainPoints(TEXT("FGrid::PointCloud 2D UniformScaled"), GetInner2DPoints(EGridSpace::UniformScaled));
+#endif
+
+#ifdef DEBUG_GRID
 		//Wait(bDisplay);
 #endif
 
 		FindInnerFacePoints();
 
 #ifdef DEBUG_GRID
-		DisplayFindInnerDomainPoints(DisplaySpace);
+		DisplayGridPoints(DisplaySpace);
 		//Wait(bDisplay);
 #endif
 		FindPointsCloseToLoop();
 #ifdef DEBUG_GRID
-		DisplayFindPointsCloseToLoop(DisplaySpace);
-		DisplayFindPointsCloseAndInsideToLoop(DisplaySpace);
+		DisplayGridInnerPoints(DisplaySpace, TEXT("FGrid::FindPointsCloseAndInsideToLoop result"));
 #endif
 		RemovePointsClosedToLoop();
 
 #ifdef DEBUG_GRID
-		DisplayFindPointsCloseAndInsideToLoop(DisplaySpace);
-		DisplayInnerDomainPoints(TEXT("FGrid::Final PointCloud 2D"), GetInner2DPoints(DisplaySpace));
+		DisplayGridInnerPoints(DisplaySpace, TEXT("FGrid::RemovePointsClosedToLoop (Final PointCloud)"));
 		//Wait(bDisplay);
-		//DisplayInnerDomainPoints("FGrid::Final PointCloud 3D"), Get3DPoints());
 #endif
 
 		// Removed of Thin zone boundary (the last boundaries). In case of thin zone, the number of 2d boundary will be biggest than 3d boundary one.
 		// Only EGridSpace::Default2D is needed.
-		FaceLoops2D[(int32)EGridSpace::Default2D].SetNum(FaceLoops3D.Num());
+		FaceLoops2D[EGridSpace::Default2D].SetNum(FaceLoops3D.Num());
 
 #ifdef DEBUG_GRID
 		DisplayLoop(TEXT("FGrid::Final Loop 2D"), GetLoops2D(DisplaySpace), true, false);
@@ -183,7 +180,7 @@ namespace CADKernel
 		{
 			for (const FOrientedEdge& Edge : Loop->GetEdges())
 			{
-				nbPoints += (int32)Edge.Entity->GetOrCreateMesh(MeshModel)->GetNodeCoordinates().Num() + 1;
+				nbPoints += Edge.Entity->GetOrCreateMesh(MeshModel)->GetNodeCoordinates().Num() + 1;
 			}
 		}
 
@@ -338,7 +335,7 @@ namespace CADKernel
 		IsCloseToLoop.Init(0, CuttingSize);
 
 		CountOfInnerNodes = CuttingSize;
-		for (int32 Index = 0; Index < (int32)EGridSpace::Max; ++Index)
+		for (int32 Index = 0; Index < EGridSpace::EndGridSpace; ++Index)
 		{
 			Points2D[Index].SetNum(CuttingSize);
 		}
@@ -350,7 +347,7 @@ namespace CADKernel
 		{
 			for (int32 IPointU = 0; IPointU < CuttingCount[EIso::IsoU]; ++IPointU, ++Index)
 			{
-				Points2D[(int32)EGridSpace::Default2D][Index].Set(Face->GetCuttingCoordinatesAlongIso(EIso::IsoU)[IPointU], Face->GetCuttingCoordinatesAlongIso(EIso::IsoV)[IPointV]);
+				Points2D[EGridSpace::Default2D][Index].Set(Face->GetCuttingCoordinatesAlongIso(EIso::IsoU)[IPointU], Face->GetCuttingCoordinatesAlongIso(EIso::IsoV)[IPointV]);
 			}
 		}
 
@@ -436,14 +433,14 @@ namespace CADKernel
 			}
 		};
 
-		for (const TArray<FPoint2D>& Loop : FaceLoops2D[(int32)EGridSpace::Default2D])
+		for (const TArray<FPoint2D>& Loop : FaceLoops2D[EGridSpace::UniformScaled])
 		{
 			const FPoint2D* PointA = &Loop.Last();
 
 			IndexU = 1;
 			for (; IndexU < CuttingCount[EIso::IsoU] - 1; ++IndexU)
 			{
-				if (CuttingCoordinates[EIso::IsoU][IndexU] + SMALL_NUMBER_SQUARE > PointA->U)
+				if (UniformCuttingCoordinates[EIso::IsoU][IndexU] + SMALL_NUMBER_SQUARE > PointA->U)
 				{
 					break;
 				}
@@ -452,7 +449,7 @@ namespace CADKernel
 			IndexV = 1;
 			for (; IndexV < CuttingCount[EIso::IsoV] - 1; ++IndexV)
 			{
-				if (CuttingCoordinates[EIso::IsoV][IndexV] + SMALL_NUMBER_SQUARE > PointA->V)
+				if (UniformCuttingCoordinates[EIso::IsoV][IndexV] + SMALL_NUMBER_SQUARE > PointA->V)
 				{
 					break;
 				}
@@ -470,9 +467,9 @@ namespace CADKernel
 			{
 				const FPoint2D* PointB = &Loop[BIndex];
 
-				if ((CuttingCoordinates[EIso::IsoV][IndexV - 1] - SMALL_NUMBER_SQUARE < PointB->V) && (PointB->V < CuttingCoordinates[EIso::IsoV][IndexV] + SMALL_NUMBER_SQUARE))
+				if ((UniformCuttingCoordinates[EIso::IsoV][IndexV - 1] - SMALL_NUMBER_SQUARE < PointB->V) && (PointB->V < UniformCuttingCoordinates[EIso::IsoV][IndexV] + SMALL_NUMBER_SQUARE))
 				{
-					if ((CuttingCoordinates[EIso::IsoU][IndexU - 1] < PointB->U) && (PointB->U < CuttingCoordinates[EIso::IsoU][IndexU]))
+					if ((UniformCuttingCoordinates[EIso::IsoU][IndexU - 1] < PointB->U) && (PointB->U < UniformCuttingCoordinates[EIso::IsoU][IndexU]))
 					{
 						PointA = PointB;
 						continue;
@@ -480,7 +477,7 @@ namespace CADKernel
 
 					if (PointA->U < PointB->U)
 					{
-						while ((CuttingCoordinates[EIso::IsoU][IndexU] < PointB->U) && (IndexU < CuttingCount[EIso::IsoU] - 1))
+						while ((UniformCuttingCoordinates[EIso::IsoU][IndexU] < PointB->U) && (IndexU < CuttingCount[EIso::IsoU] - 1))
 						{
 							IncreaseU();
 							IsCloseToLoop[Index] = 1;
@@ -489,7 +486,7 @@ namespace CADKernel
 					}
 					else
 					{
-						while ((CuttingCoordinates[EIso::IsoU][IndexU - 1] > PointB->U) && (IndexU > 1))
+						while ((UniformCuttingCoordinates[EIso::IsoU][IndexU - 1] > PointB->U) && (IndexU > 1))
 						{
 							DecreaseU();
 							IsCloseToLoop[Index - 1] = 1;
@@ -500,11 +497,11 @@ namespace CADKernel
 					continue;
 				}
 
-				if ((CuttingCoordinates[EIso::IsoU][IndexU - 1] < PointB->U) && (PointB->U < CuttingCoordinates[EIso::IsoU][IndexU]))
+				if ((UniformCuttingCoordinates[EIso::IsoU][IndexU - 1] < PointB->U) && (PointB->U < UniformCuttingCoordinates[EIso::IsoU][IndexU]))
 				{
 					if (PointA->V < PointB->V)
 					{
-						while ((CuttingCoordinates[EIso::IsoV][IndexV] < PointB->V) && (IndexV < CuttingCount[EIso::IsoV] - 1))
+						while ((UniformCuttingCoordinates[EIso::IsoV][IndexV] < PointB->V) && (IndexV < CuttingCount[EIso::IsoV] - 1))
 						{
 							IncreaseV();
 							IsCloseToLoop[Index] = 1;
@@ -513,7 +510,7 @@ namespace CADKernel
 					}
 					else
 					{
-						while ((CuttingCoordinates[EIso::IsoV][IndexV - 1] > PointB->V) && (IndexV > 1))
+						while ((UniformCuttingCoordinates[EIso::IsoV][IndexV - 1] > PointB->V) && (IndexV > 1))
 						{
 							DecreaseV();
 							IsCloseToLoop[Index - CuttingCount[EIso::IsoU]] = 1;
@@ -536,25 +533,25 @@ namespace CADKernel
 
 					if (ABu > 0)
 					{
-						while (CuttingCoordinates[EIso::IsoU][IndexU] + SMALL_NUMBER_SQUARE < PointB->U)
+						while (UniformCuttingCoordinates[EIso::IsoU][IndexU] + SMALL_NUMBER_SQUARE < PointB->U)
 						{
 							IncreaseU();
 
 							double CoordinateV_IndexU = 0;
-							if (CuttingCoordinates[EIso::IsoU][IndexU] > PointB->U)
+							if (UniformCuttingCoordinates[EIso::IsoU][IndexU] > PointB->U)
 							{
 								CoordinateV_IndexU = PointB->V;
 							}
 							else
 							{
-								CoordinateV_IndexU = ABy_ABx * CuttingCoordinates[EIso::IsoU][IndexU] + Ay_ABy_ABx_Ax;
+								CoordinateV_IndexU = ABy_ABx * UniformCuttingCoordinates[EIso::IsoU][IndexU] + Ay_ABy_ABx_Ax;
 								if (ABv < 0)
 								{
 									if (IndexV > 2)
 									{
-										if (CoordinateV_IndexU + SMALL_NUMBER_SQUARE < CuttingCoordinates[EIso::IsoV][IndexV - 2])
+										if (CoordinateV_IndexU + SMALL_NUMBER_SQUARE < UniformCuttingCoordinates[EIso::IsoV][IndexV - 2])
 										{
-											CoordinateV_IndexU = CuttingCoordinates[EIso::IsoV][IndexV - 2] + SMALL_NUMBER_SQUARE;
+											CoordinateV_IndexU = UniformCuttingCoordinates[EIso::IsoV][IndexV - 2] + SMALL_NUMBER_SQUARE;
 											DecreaseU();
 										}
 									}
@@ -563,16 +560,16 @@ namespace CADKernel
 								{
 									if (IndexV < CuttingCount[EIso::IsoV] - 1)
 									{
-										if (CoordinateV_IndexU + SMALL_NUMBER_SQUARE > CuttingCoordinates[EIso::IsoV][IndexV + 1])
+										if (CoordinateV_IndexU + SMALL_NUMBER_SQUARE > UniformCuttingCoordinates[EIso::IsoV][IndexV + 1])
 										{
-											CoordinateV_IndexU = CuttingCoordinates[EIso::IsoV][IndexV + 1] - SMALL_NUMBER_SQUARE;
+											CoordinateV_IndexU = UniformCuttingCoordinates[EIso::IsoV][IndexV + 1] - SMALL_NUMBER_SQUARE;
 											DecreaseU();
 										}
 									}
 								}
 							}
 
-							if (CoordinateV_IndexU - SMALL_NUMBER_SQUARE > CuttingCoordinates[EIso::IsoV][IndexV])
+							if (CoordinateV_IndexU - SMALL_NUMBER_SQUARE > UniformCuttingCoordinates[EIso::IsoV][IndexV])
 							{
 								IncreaseV();
 
@@ -580,7 +577,7 @@ namespace CADKernel
 								IsCloseToLoop[Index - 1] = 1;
 								IsCloseToLoop[Index - CuttingCount[EIso::IsoU]] = 1;
 							}
-							else if (CoordinateV_IndexU + SMALL_NUMBER_SQUARE < CuttingCoordinates[EIso::IsoV][IndexV - 1])
+							else if (CoordinateV_IndexU + SMALL_NUMBER_SQUARE < UniformCuttingCoordinates[EIso::IsoV][IndexV - 1])
 							{
 								DecreaseV();
 
@@ -602,20 +599,20 @@ namespace CADKernel
 							DecreaseU();
 
 							double CoordinateV_IndexU = 0; //= ABy_ABx * CoordinateU[IndexU - 1] + Ay_ABy_ABx_Ax;
-							if (CuttingCoordinates[EIso::IsoU][IndexU - 1] < PointB->U)
+							if (UniformCuttingCoordinates[EIso::IsoU][IndexU - 1] < PointB->U)
 							{
 								CoordinateV_IndexU = PointB->V;
 							}
 							else
 							{
-								CoordinateV_IndexU = ABy_ABx * CuttingCoordinates[EIso::IsoU][IndexU - 1] + Ay_ABy_ABx_Ax;
+								CoordinateV_IndexU = ABy_ABx * UniformCuttingCoordinates[EIso::IsoU][IndexU - 1] + Ay_ABy_ABx_Ax;
 								if (ABv < 0)
 								{
 									if (IndexV > 2)
 									{
-										if (CoordinateV_IndexU + SMALL_NUMBER_SQUARE < CuttingCoordinates[EIso::IsoV][IndexV - 2])
+										if (CoordinateV_IndexU + SMALL_NUMBER_SQUARE < UniformCuttingCoordinates[EIso::IsoV][IndexV - 2])
 										{
-											CoordinateV_IndexU = CuttingCoordinates[EIso::IsoV][IndexV - 2] + SMALL_NUMBER_SQUARE;
+											CoordinateV_IndexU = UniformCuttingCoordinates[EIso::IsoV][IndexV - 2] + SMALL_NUMBER_SQUARE;
 											IncreaseU();
 										}
 									}
@@ -624,16 +621,16 @@ namespace CADKernel
 								{
 									if (IndexV < CuttingCount[EIso::IsoV] - 1)
 									{
-										if (CoordinateV_IndexU + SMALL_NUMBER_SQUARE > CuttingCoordinates[EIso::IsoV][IndexV + 1])
+										if (CoordinateV_IndexU + SMALL_NUMBER_SQUARE > UniformCuttingCoordinates[EIso::IsoV][IndexV + 1])
 										{
-											CoordinateV_IndexU = CuttingCoordinates[EIso::IsoV][IndexV + 1] - SMALL_NUMBER_SQUARE;
+											CoordinateV_IndexU = UniformCuttingCoordinates[EIso::IsoV][IndexV + 1] - SMALL_NUMBER_SQUARE;
 											IncreaseU();
 										}
 									}
 								}
 							}
 
-							if (CoordinateV_IndexU - SMALL_NUMBER_SQUARE > CuttingCoordinates[EIso::IsoV][IndexV])
+							if (CoordinateV_IndexU - SMALL_NUMBER_SQUARE > UniformCuttingCoordinates[EIso::IsoV][IndexV])
 							{
 								IsCloseToLoop[Index - 1] = 1;
 
@@ -641,7 +638,7 @@ namespace CADKernel
 								IsCloseToLoop[Index - 1] = 1;
 								IsCloseToLoop[Index] = 1;
 							}
-							else if (CoordinateV_IndexU + SMALL_NUMBER_SQUARE < CuttingCoordinates[EIso::IsoV][IndexV - 1])
+							else if (CoordinateV_IndexU + SMALL_NUMBER_SQUARE < UniformCuttingCoordinates[EIso::IsoV][IndexV - 1])
 							{
 								DecreaseV();
 
@@ -654,7 +651,7 @@ namespace CADKernel
 								IsCloseToLoop[Index - 1] = 1;
 								IsCloseToLoop[Index - CuttingCount[EIso::IsoU] - 1] = 1;
 							}
-						} while (CuttingCoordinates[EIso::IsoU][IndexU - 1] - SMALL_NUMBER_SQUARE > PointB->U);
+						} while (UniformCuttingCoordinates[EIso::IsoU][IndexU - 1] - SMALL_NUMBER_SQUARE > PointB->U);
 					}
 				}
 				else
@@ -666,24 +663,24 @@ namespace CADKernel
 
 					if (ABv > 0)
 					{
-						while (CuttingCoordinates[EIso::IsoV][IndexV] + SMALL_NUMBER_SQUARE < PointB->V)
+						while (UniformCuttingCoordinates[EIso::IsoV][IndexV] + SMALL_NUMBER_SQUARE < PointB->V)
 						{
 							IncreaseV();
 							double CoordinateU_IndexV = 0;// = ABx_ABy * CoordinateV[IndexV - 1] + Ax_ABx_ABy_Ay;
-							if (CuttingCoordinates[EIso::IsoV][IndexV] > PointB->V)
+							if (UniformCuttingCoordinates[EIso::IsoV][IndexV] > PointB->V)
 							{
 								CoordinateU_IndexV = PointB->U;
 							}
 							else
 							{
-								CoordinateU_IndexV = ABu_ABv * CuttingCoordinates[EIso::IsoV][IndexV] + Au_ABu_ABv_Av;
+								CoordinateU_IndexV = ABu_ABv * UniformCuttingCoordinates[EIso::IsoV][IndexV] + Au_ABu_ABv_Av;
 								if (ABu < 0)
 								{
 									if (IndexU > 2)
 									{
-										if (CoordinateU_IndexV + SMALL_NUMBER_SQUARE < CuttingCoordinates[EIso::IsoU][IndexU - 2])
+										if (CoordinateU_IndexV + SMALL_NUMBER_SQUARE < UniformCuttingCoordinates[EIso::IsoU][IndexU - 2])
 										{
-											CoordinateU_IndexV = CuttingCoordinates[EIso::IsoU][IndexU - 2] + SMALL_NUMBER_SQUARE;
+											CoordinateU_IndexV = UniformCuttingCoordinates[EIso::IsoU][IndexU - 2] + SMALL_NUMBER_SQUARE;
 											DecreaseV();
 										}
 									}
@@ -692,23 +689,23 @@ namespace CADKernel
 								{
 									if (IndexU < CuttingCount[EIso::IsoU] - 1)
 									{
-										if (CoordinateU_IndexV + SMALL_NUMBER_SQUARE > CuttingCoordinates[EIso::IsoU][IndexU + 1])
+										if (CoordinateU_IndexV + SMALL_NUMBER_SQUARE > UniformCuttingCoordinates[EIso::IsoU][IndexU + 1])
 										{
-											CoordinateU_IndexV = CuttingCoordinates[EIso::IsoU][IndexU + 1] - SMALL_NUMBER_SQUARE;
+											CoordinateU_IndexV = UniformCuttingCoordinates[EIso::IsoU][IndexU + 1] - SMALL_NUMBER_SQUARE;
 											DecreaseV();
 										}
 									}
 								}
 							}
 
-							if ((IndexU < CuttingCount[EIso::IsoU] - 1) && CoordinateU_IndexV - SMALL_NUMBER_SQUARE > CuttingCoordinates[EIso::IsoU][IndexU])
+							if ((IndexU < CuttingCount[EIso::IsoU] - 1) && CoordinateU_IndexV - SMALL_NUMBER_SQUARE > UniformCuttingCoordinates[EIso::IsoU][IndexU])
 							{
 								IncreaseU();
 								IsCloseToLoop[Index] = 1;
 								IsCloseToLoop[Index - 1] = 1;
 								IsCloseToLoop[Index - CuttingCount[EIso::IsoU]] = 1;
 							}
-							else if ((IndexU > 1) && CoordinateU_IndexV + SMALL_NUMBER_SQUARE < CuttingCoordinates[EIso::IsoU][IndexU - 1])
+							else if ((IndexU > 1) && CoordinateU_IndexV + SMALL_NUMBER_SQUARE < UniformCuttingCoordinates[EIso::IsoU][IndexU - 1])
 							{
 								DecreaseU();
 
@@ -730,20 +727,20 @@ namespace CADKernel
 							DecreaseV();
 
 							double CoordinateU_IndexV = 0;// = ABx_ABy * CoordinateV[IndexV - 1] + Ax_ABx_ABy_Ay;
-							if (CuttingCoordinates[EIso::IsoV][IndexV - 1] < PointB->V)
+							if (UniformCuttingCoordinates[EIso::IsoV][IndexV - 1] < PointB->V)
 							{
 								CoordinateU_IndexV = PointB->U;
 							}
 							else
 							{
-								CoordinateU_IndexV = ABu_ABv * CuttingCoordinates[EIso::IsoV][IndexV - 1] + Au_ABu_ABv_Av;
+								CoordinateU_IndexV = ABu_ABv * UniformCuttingCoordinates[EIso::IsoV][IndexV - 1] + Au_ABu_ABv_Av;
 								if (ABu < 0)
 								{
 									if (IndexU > 2)
 									{
-										if (CoordinateU_IndexV + SMALL_NUMBER_SQUARE < CuttingCoordinates[EIso::IsoU][IndexU - 2])
+										if (CoordinateU_IndexV + SMALL_NUMBER_SQUARE < UniformCuttingCoordinates[EIso::IsoU][IndexU - 2])
 										{
-											CoordinateU_IndexV = CuttingCoordinates[EIso::IsoU][IndexU - 2] + SMALL_NUMBER_SQUARE;
+											CoordinateU_IndexV = UniformCuttingCoordinates[EIso::IsoU][IndexU - 2] + SMALL_NUMBER_SQUARE;
 											IncreaseV();
 										}
 									}
@@ -752,23 +749,23 @@ namespace CADKernel
 								{
 									if (IndexU < CuttingCount[EIso::IsoU] - 1)
 									{
-										if (CoordinateU_IndexV + SMALL_NUMBER_SQUARE > CuttingCoordinates[EIso::IsoU][IndexU + 1])
+										if (CoordinateU_IndexV + SMALL_NUMBER_SQUARE > UniformCuttingCoordinates[EIso::IsoU][IndexU + 1])
 										{
-											CoordinateU_IndexV = CuttingCoordinates[EIso::IsoU][IndexU + 1] - SMALL_NUMBER_SQUARE;
+											CoordinateU_IndexV = UniformCuttingCoordinates[EIso::IsoU][IndexU + 1] - SMALL_NUMBER_SQUARE;
 											IncreaseV();
 										}
 									}
 								}
 							}
 
-							if (CoordinateU_IndexV - SMALL_NUMBER_SQUARE > CuttingCoordinates[EIso::IsoU][IndexU])
+							if (CoordinateU_IndexV - SMALL_NUMBER_SQUARE > UniformCuttingCoordinates[EIso::IsoU][IndexU])
 							{
 								IncreaseU();
 								IsCloseToLoop[Index] = 1;
 								IsCloseToLoop[Index - CuttingCount[EIso::IsoU] - 1] = 1;
 								IsCloseToLoop[Index - CuttingCount[EIso::IsoU]] = 1;
 							}
-							else if (CoordinateU_IndexV + SMALL_NUMBER_SQUARE < CuttingCoordinates[EIso::IsoU][IndexU - 1])
+							else if (CoordinateU_IndexV + SMALL_NUMBER_SQUARE < UniformCuttingCoordinates[EIso::IsoU][IndexU - 1])
 							{
 								DecreaseU();
 								IsCloseToLoop[Index - 1] = 1;
@@ -780,7 +777,7 @@ namespace CADKernel
 								IsCloseToLoop[Index - CuttingCount[EIso::IsoU]] = 1;
 								IsCloseToLoop[Index - CuttingCount[EIso::IsoU] - 1] = 1;
 							}
-						} while (CuttingCoordinates[EIso::IsoV][IndexV - 1] - SMALL_NUMBER_SQUARE > PointB->V);
+						} while (UniformCuttingCoordinates[EIso::IsoV][IndexV - 1] - SMALL_NUMBER_SQUARE > PointB->V);
 					}
 				}
 				PointA = PointB;
@@ -841,18 +838,19 @@ namespace CADKernel
 		TArray<FGridSegment> LoopSegments;
 		{
 			int32 SegmentNum = 0;
-			for (const TArray<FPoint2D>& Loop : FaceLoops2D[(int32)EGridSpace::Default2D])
+			for (const TArray<FPoint2D>& Loop : FaceLoops2D[EGridSpace::UniformScaled])
 			{
-				SegmentNum += (int32)Loop.Num();
+				SegmentNum += Loop.Num();
 			}
 			LoopSegments.Reserve(SegmentNum);
 
-			for (TArray<FPoint2D>& Loop : FaceLoops2D[(int32)EGridSpace::Default2D])
+			for (TArray<FPoint2D>& Loop : FaceLoops2D[EGridSpace::UniformScaled])
 			{
 				for (int32 Index = 0; Index < Loop.Num() - 1; ++Index)
 				{
 					LoopSegments.Emplace(Loop[Index], Loop[Index + 1]);
 				}
+				LoopSegments.Emplace(Loop.Last(), Loop[0]);
 			}
 
 			Algo::Sort(LoopSegments, [](const FGridSegment& Seg1, const FGridSegment& Seg2) -> bool
@@ -880,7 +878,7 @@ namespace CADKernel
 			SortedPointIndexes.Reserve(IndexOfPointsNearAndInsideLoop.Num());
 			for (const int32& Index : IndexOfPointsNearAndInsideLoop)
 			{
-				GridPointWeight.Add(Points2D[(int32)EGridSpace::Default2D][Index].U + Points2D[(int32)EGridSpace::Default2D][Index].V);
+				GridPointWeight.Add(Points2D[EGridSpace::UniformScaled][Index].U + Points2D[EGridSpace::UniformScaled][Index].V);
 			}
 			for (int32 Index = 0; Index < IndexOfPointsNearAndInsideLoop.Num(); ++Index)
 			{
@@ -889,22 +887,24 @@ namespace CADKernel
 			SortedPointIndexes.Sort([&GridPointWeight](const int32& Index1, const int32& Index2) { return GridPointWeight[Index1] < GridPointWeight[Index2]; });
 		}
 
+		// only used to reduce the search of neighborhood
 		double DeltaUVMax = 0;
 		{
 			double MaxDeltaU = 0;
 			for (int32 Index = 0; Index < CuttingCount[EIso::IsoU] - 1; ++Index)
 			{
-				MaxDeltaU = FMath::Max(MaxDeltaU, FMath::Abs(CuttingCoordinates[EIso::IsoU][Index + 1] - CuttingCoordinates[EIso::IsoU][Index]));
+				MaxDeltaU = FMath::Max(MaxDeltaU, FMath::Abs(UniformCuttingCoordinates[EIso::IsoU][Index + 1] - UniformCuttingCoordinates[EIso::IsoU][Index]));
 			}
 
 			double MaxDeltaV = 0;
 			for (int32 Index = 0; Index < CuttingCount[EIso::IsoV] - 1; ++Index)
 			{
-				MaxDeltaV = FMath::Max(MaxDeltaV, FMath::Abs(CuttingCoordinates[EIso::IsoV][Index + 1] - CuttingCoordinates[EIso::IsoV][Index]));
+				MaxDeltaV = FMath::Max(MaxDeltaV, FMath::Abs(UniformCuttingCoordinates[EIso::IsoV][Index + 1] - UniformCuttingCoordinates[EIso::IsoV][Index]));
 			}
 			DeltaUVMax = FMath::Max(MaxDeltaU, MaxDeltaV);
 		}
 
+		// Find DeltaU and DeltaV around a cutting point defined by its index
 		TFunction<void(const int32, double&, double&)> GetDeltaUV = [&](const int32 Index, double& DeltaU, double& DeltaV)
 		{
 			int32 IndexU = Index % CuttingCount[EIso::IsoU];
@@ -912,43 +912,43 @@ namespace CADKernel
 
 			if (IndexU == 0)
 			{
-				DeltaU = FMath::Abs(CuttingCoordinates[EIso::IsoU][1] - CuttingCoordinates[EIso::IsoU][0]);
+				DeltaU = FMath::Abs(UniformCuttingCoordinates[EIso::IsoU][1] - UniformCuttingCoordinates[EIso::IsoU][0]);
 			}
 			else if (IndexU == CuttingCount[EIso::IsoU] - 1)
 			{
-				DeltaU = FMath::Abs(CuttingCoordinates[EIso::IsoU][CuttingCount[EIso::IsoU] - 1] - CuttingCoordinates[EIso::IsoU][CuttingCount[EIso::IsoU] - 2]);
+				DeltaU = FMath::Abs(UniformCuttingCoordinates[EIso::IsoU][CuttingCount[EIso::IsoU] - 1] - UniformCuttingCoordinates[EIso::IsoU][CuttingCount[EIso::IsoU] - 2]);
 			}
 			else
 			{
-				DeltaU = FMath::Abs(CuttingCoordinates[EIso::IsoU][IndexU + 1] - CuttingCoordinates[EIso::IsoU][IndexU - 1]) * .5;
+				DeltaU = FMath::Abs(UniformCuttingCoordinates[EIso::IsoU][IndexU + 1] - UniformCuttingCoordinates[EIso::IsoU][IndexU - 1]) * .5;
 			}
 
 			if (IndexV == 0)
 			{
-				DeltaV = FMath::Abs(CuttingCoordinates[EIso::IsoV][1] - CuttingCoordinates[EIso::IsoV][0]);
+				DeltaV = FMath::Abs(UniformCuttingCoordinates[EIso::IsoV][1] - UniformCuttingCoordinates[EIso::IsoV][0]);
 			}
 			else if (IndexV == CuttingCount[EIso::IsoV] - 1)
 			{
-				DeltaV = FMath::Abs(CuttingCoordinates[EIso::IsoV][CuttingCount[EIso::IsoV] - 1] - CuttingCoordinates[EIso::IsoV][CuttingCount[EIso::IsoV] - 2]);
+				DeltaV = FMath::Abs(UniformCuttingCoordinates[EIso::IsoV][CuttingCount[EIso::IsoV] - 1] - UniformCuttingCoordinates[EIso::IsoV][CuttingCount[EIso::IsoV] - 2]);
 			}
 			else
 			{
-				DeltaV = FMath::Abs(CuttingCoordinates[EIso::IsoV][IndexV + 1] - CuttingCoordinates[EIso::IsoV][IndexV - 1]) * .5;
+				DeltaV = FMath::Abs(UniformCuttingCoordinates[EIso::IsoV][IndexV + 1] - UniformCuttingCoordinates[EIso::IsoV][IndexV - 1]) * .5;
 			}
 		};
 
-		const double DeltaUVMinSquare = FMath::Square(1. / 3.);
-
 		int32 SegmentIndex = 0;
+		SegmentIndex = 0;
 		for (const int32& SortedIndex : SortedPointIndexes)
 		{
 			int32 Index = IndexOfPointsNearAndInsideLoop[SortedIndex];
-			const FPoint2D& Point2D = Points2D[(int32)EGridSpace::Default2D][Index];
+			const FPoint2D& Point2D = Points2D[EGridSpace::UniformScaled][Index];
 
 			double DeltaU;
 			double DeltaV;
 			GetDeltaUV(Index, DeltaU, DeltaV);
 
+			//find the first segment that could be close to the point			
 			for (; SegmentIndex < LoopSegments.Num(); ++SegmentIndex)
 			{
 				if (GridPointWeight[SortedIndex] < LoopSegments[SegmentIndex].EndPointWeight + DeltaUVMax)
@@ -986,9 +986,9 @@ namespace CADKernel
 				FPoint2D Projection = ProjectPointOnSegment(Point2D, *Segment.StartPoint, *Segment.EndPoint, Coordinate, /*bRestrictCoodinateToInside*/ true);
 
 				// If Projected point is in the oval center on Point2D => the node is too close
-				double SqrDistance2D = FMath::Square((Point2D.U - Projection.U) / DeltaU);
-				SqrDistance2D += FMath::Square((Point2D.V - Projection.V) / DeltaV);
-				if (SqrDistance2D > DeltaUVMinSquare)
+				FPoint2D ProjectionToPoint = Abs(Point2D - Projection);
+
+				if (ProjectionToPoint.U > DeltaU * 0.05  || ProjectionToPoint.V > DeltaV * 0.05)
 				{
 					continue;
 				}
@@ -1066,11 +1066,11 @@ namespace CADKernel
 		int32 ThinZoneNum = 0;
 		if (Face->HasThinZone())
 		{
-			ThinZoneNum = (int32)ThinZoneFinder.GetThinZones().Num();
+			ThinZoneNum = ThinZoneFinder.GetThinZones().Num();
 		}
 
 		int32 LoopCount = Face->GetLoops().Num();
-		FaceLoops2D[(int32)EGridSpace::Default2D].Reserve(LoopCount + ThinZoneNum);
+		FaceLoops2D[EGridSpace::Default2D].Reserve(LoopCount + ThinZoneNum);
 
 		FaceLoops3D.Reserve(LoopCount);
 		NormalsOfFaceLoops.Reserve(LoopCount);
@@ -1089,7 +1089,7 @@ namespace CADKernel
 				LoopNodeCount += Edge.Entity->GetLinkActiveEdge()->GetCuttingPoints().Num() + 2;
 			}
 
-			TArray<FPoint2D>& Loop2D = FaceLoops2D[(int32)EGridSpace::Default2D].Emplace_GetRef();
+			TArray<FPoint2D>& Loop2D = FaceLoops2D[EGridSpace::Default2D].Emplace_GetRef();
 			Loop2D.Empty(LoopNodeCount);
 
 			TArray<FPoint>& Loop3D = FaceLoops3D.Emplace_GetRef();
@@ -1210,7 +1210,7 @@ namespace CADKernel
 
 			if (Loop2D.Num() < 3) // degenerated loop
 			{
-				FaceLoops2D[(int32)EGridSpace::Default2D].Pop();
+				FaceLoops2D[EGridSpace::Default2D].Pop();
 				FaceLoops3D.Pop();
 				NormalsOfFaceLoops.Pop();
 				continue;
@@ -1226,9 +1226,9 @@ namespace CADKernel
 		{
 			for (const FThinZone2D& ThinZone : ThinZoneFinder.GetThinZones())
 			{
-				int32 PointNum = (int32)ThinZone.GetFirstSide().GetSegments().Num();
-				PointNum += (int32)ThinZone.GetSecondSide().GetSegments().Num();
-				TArray<FPoint2D>& LoopPoints = FaceLoops2D[(int32)EGridSpace::Default2D].Emplace_GetRef();
+				int32 PointNum = ThinZone.GetFirstSide().GetSegments().Num();
+				PointNum += ThinZone.GetSecondSide().GetSegments().Num();
+				TArray<FPoint2D>& LoopPoints = FaceLoops2D[EGridSpace::Default2D].Emplace_GetRef();
 				LoopPoints.Reserve(PointNum + 4);
 
 				// First point Side1
@@ -1258,7 +1258,7 @@ namespace CADKernel
 
 		// Fit boundaries to Surface bounds.
 		const FSurfacicBoundary& Bounds = Face->GetBoundary();
-		for (TArray<FPoint2D>& Loop : FaceLoops2D[(int32)EGridSpace::Default2D])
+		for (TArray<FPoint2D>& Loop : FaceLoops2D[EGridSpace::Default2D])
 		{
 			for (FPoint2D& Point : Loop)
 			{
@@ -1352,11 +1352,9 @@ namespace CADKernel
 		double MeanLengthU = GetMean(LengthsU);
 		double FactorU = MeanLengthU / (CuttingCoordinates[EIso::IsoU].Last() - CuttingCoordinates[EIso::IsoU][0]);
 
-		TArray<double> ScaledCoordinatesU;
-		ScaleCoordinates(CuttingCoordinates[EIso::IsoU], FactorU, ScaledCoordinatesU);
-
-		TArray<double> ScaledCoordinatesV;
-		ScaleCoordinates(CuttingCoordinates[EIso::IsoV], FactorV, ScaledCoordinatesV);
+		//TArray<double> ScaledCuttingCoordinates;
+		ScaleCoordinates(CuttingCoordinates[EIso::IsoU], FactorU, UniformCuttingCoordinates[EIso::IsoU]);
+		ScaleCoordinates(CuttingCoordinates[EIso::IsoV], FactorV, UniformCuttingCoordinates[EIso::IsoV]);
 
 		{
 			int32 NumUV = 0;
@@ -1364,7 +1362,7 @@ namespace CADKernel
 			{
 				for (int32 IPointU = 0; IPointU < CuttingCount[EIso::IsoU]; ++IPointU, ++NumUV)
 				{
-					Points2D[(int32)EGridSpace::UniformScaled][NumUV].Set(ScaledCoordinatesU[IPointU], ScaledCoordinatesV[IPointV]);
+					Points2D[EGridSpace::UniformScaled][NumUV].Set(UniformCuttingCoordinates[EIso::IsoU][IPointU], UniformCuttingCoordinates[EIso::IsoV][IPointV]);
 				}
 			}
 		}
@@ -1395,7 +1393,7 @@ namespace CADKernel
 				{
 					NumUV = IPointV * CuttingCount[EIso::IsoU] + IPointU;
 					Length += LastPoint.Distance(Points3D[NumUV]);
-					Points2D[(int32)EGridSpace::Scaled][NumUV].Set(Points2D[(int32)EGridSpace::UniformScaled][NumUV].U, Length);
+					Points2D[EGridSpace::Scaled][NumUV].Set(Points2D[EGridSpace::UniformScaled][NumUV].U, Length);
 					LastPoint = Points3D[NumUV];
 				}
 
@@ -1405,7 +1403,7 @@ namespace CADKernel
 				{
 					NumUV = IPointV * CuttingCount[EIso::IsoU] + IPointU;
 					Length -= LastPoint.Distance(Points3D[NumUV]);
-					Points2D[(int32)EGridSpace::Scaled][NumUV].Set(Points2D[(int32)EGridSpace::UniformScaled][NumUV].U, Length);
+					Points2D[EGridSpace::Scaled][NumUV].Set(Points2D[EGridSpace::UniformScaled][NumUV].U, Length);
 					LastPoint = Points3D[NumUV];
 				}
 			}
@@ -1433,7 +1431,7 @@ namespace CADKernel
 				{
 					NumUV = IPointV * CuttingCount[EIso::IsoU] + IPointU;
 					Length += LastPoint.Distance(Points3D[NumUV]);
-					Points2D[(int32)EGridSpace::Scaled][NumUV].Set(Length, Points2D[(int32)EGridSpace::UniformScaled][NumUV].V);
+					Points2D[EGridSpace::Scaled][NumUV].Set(Length, Points2D[EGridSpace::UniformScaled][NumUV].V);
 					LastPoint = Points3D[NumUV];
 				}
 
@@ -1443,7 +1441,7 @@ namespace CADKernel
 				{
 					NumUV = IPointV * CuttingCount[EIso::IsoU] + IPointU;
 					Length -= LastPoint.Distance(Points3D[NumUV]);
-					Points2D[(int32)EGridSpace::Scaled][NumUV].Set(Length, Points2D[(int32)EGridSpace::UniformScaled][NumUV].V);
+					Points2D[EGridSpace::Scaled][NumUV].Set(Length, Points2D[EGridSpace::UniformScaled][NumUV].V);
 					LastPoint = Points3D[NumUV];
 				}
 			}
@@ -1461,10 +1459,10 @@ namespace CADKernel
 		{
 			const FPoint2D& Point = InPointsToScale[Index];
 
-			FindCoordinateIndex(CuttingCoordinates[EIso::IsoU], Point.U, IndexU);
-			FindCoordinateIndex(CuttingCoordinates[EIso::IsoV], Point.V, IndexV);
+			ArrayUtils::FindCoordinateIndex(CuttingCoordinates[EIso::IsoU], Point.U, IndexU);
+			ArrayUtils::FindCoordinateIndex(CuttingCoordinates[EIso::IsoV], Point.V, IndexV);
 
-			ComputeNewCoordinate(Points2D[(int32)DestinationSpace], IndexU, IndexV, Point, OutTransformedPoints[Index]);
+			ComputeNewCoordinate(Points2D[DestinationSpace], IndexU, IndexV, Point, OutTransformedPoints[Index]);
 		}
 	}
 
@@ -1477,14 +1475,14 @@ namespace CADKernel
 
 	void FGrid::ScaleLoops()
 	{
-		FaceLoops2D[(int32)EGridSpace::Scaled].SetNum(FaceLoops2D[(int32)EGridSpace::Default2D].Num());
-		FaceLoops2D[(int32)EGridSpace::UniformScaled].SetNum(FaceLoops2D[(int32)EGridSpace::Default2D].Num());
+		FaceLoops2D[EGridSpace::Scaled].SetNum(FaceLoops2D[EGridSpace::Default2D].Num());
+		FaceLoops2D[EGridSpace::UniformScaled].SetNum(FaceLoops2D[EGridSpace::Default2D].Num());
 
-		for (int32 IndexBoudnary = 0; IndexBoudnary < FaceLoops2D[(int32)EGridSpace::Default2D].Num(); ++IndexBoudnary)
+		for (int32 IndexBoudnary = 0; IndexBoudnary < FaceLoops2D[EGridSpace::Default2D].Num(); ++IndexBoudnary)
 		{
-			const TArray<FPoint2D>& Loop = FaceLoops2D[(int32)EGridSpace::Default2D][IndexBoudnary];
-			TArray<FPoint2D>& ScaledLoop = FaceLoops2D[(int32)EGridSpace::Scaled][IndexBoudnary];
-			TArray<FPoint2D>& UniformScaledLoop = FaceLoops2D[(int32)EGridSpace::UniformScaled][IndexBoudnary];
+			const TArray<FPoint2D>& Loop = FaceLoops2D[EGridSpace::Default2D][IndexBoudnary];
+			TArray<FPoint2D>& ScaledLoop = FaceLoops2D[EGridSpace::Scaled][IndexBoudnary];
+			TArray<FPoint2D>& UniformScaledLoop = FaceLoops2D[EGridSpace::UniformScaled][IndexBoudnary];
 
 			ScaledLoop.SetNum(Loop.Num());
 			UniformScaledLoop.SetNum(Loop.Num());
@@ -1495,8 +1493,8 @@ namespace CADKernel
 			{
 				const FPoint2D& Point = Loop[Index];
 
-				FindCoordinateIndex(CuttingCoordinates[EIso::IsoU], Point.U, IndexU);
-				FindCoordinateIndex(CuttingCoordinates[EIso::IsoV], Point.V, IndexV);
+				ArrayUtils::FindCoordinateIndex(CuttingCoordinates[EIso::IsoU], Point.U, IndexU);
+				ArrayUtils::FindCoordinateIndex(CuttingCoordinates[EIso::IsoV], Point.V, IndexV);
 
 				ComputeNewCoordinate(Points2D[EGridSpace::Scaled], IndexU, IndexV, Point, ScaledLoop[Index]);
 				ComputeNewCoordinate(Points2D[EGridSpace::UniformScaled], IndexU, IndexV, Point, UniformScaledLoop[Index]);
@@ -1554,21 +1552,21 @@ namespace CADKernel
 		{
 			int32 IndexV = 0;
 			int32 IndexU = 0;
-			for (TArray<FPoint2D>& Loop : FaceLoops2D[(int32)EGridSpace::Default2D])
+			for (TArray<FPoint2D>& Loop : FaceLoops2D[EGridSpace::UniformScaled])
 			{
 				for (FPoint2D& Point : Loop)
 				{
-					while (IndexV != 0 && (Point.V < CuttingCoordinates[EIso::IsoV][IndexV]))
+					while (IndexV != 0 && (Point.V < UniformCuttingCoordinates[EIso::IsoV][IndexV]))
 					{
 						IndexV--;
 					}
 					for (; IndexV < CuttingCount[EIso::IsoV]; ++IndexV)
 					{
-						if (Point.V + SmallToleranceV < CuttingCoordinates[EIso::IsoV][IndexV])
+						if (Point.V + SmallToleranceV < UniformCuttingCoordinates[EIso::IsoV][IndexV])
 						{
 							break;
 						}
-						if (Point.V - SmallToleranceV > CuttingCoordinates[EIso::IsoV][IndexV])
+						if (Point.V - SmallToleranceV > UniformCuttingCoordinates[EIso::IsoV][IndexV])
 						{
 							continue;
 						}
@@ -1589,17 +1587,17 @@ namespace CADKernel
 						IndexV--;
 					}
 
-					while (IndexU != 0 && (Point.U < CuttingCoordinates[EIso::IsoU][IndexU]))
+					while (IndexU != 0 && (Point.U < UniformCuttingCoordinates[EIso::IsoU][IndexU]))
 					{
 						IndexU--;
 					}
 					for (; IndexU < CuttingCount[EIso::IsoU]; ++IndexU)
 					{
-						if (Point.U + SmallToleranceU < CuttingCoordinates[EIso::IsoU][IndexU])
+						if (Point.U + SmallToleranceU < UniformCuttingCoordinates[EIso::IsoU][IndexU])
 						{
 							break;
 						}
-						if (Point.U - SmallToleranceU > CuttingCoordinates[EIso::IsoU][IndexU])
+						if (Point.U - SmallToleranceU > UniformCuttingCoordinates[EIso::IsoU][IndexU])
 						{
 							continue;
 						}
@@ -1622,10 +1620,10 @@ namespace CADKernel
 			}
 		}
 
-		DisplayLoop(TEXT("FGrid::Loop 2D After move according tol"), GetLoops2D(EGridSpace::Default2D), true, false);
+		DisplayLoop(TEXT("FGrid::Loop 2D After move according tol"), GetLoops2D(EGridSpace::UniformScaled), true, false);
 
 		// Intersection along U axis
-		for (const TArray<FPoint2D>& Loop : FaceLoops2D[(int32)EGridSpace::Default2D])
+		for (const TArray<FPoint2D>& Loop : FaceLoops2D[EGridSpace::UniformScaled])
 		{
 			const FPoint2D* FirstSegmentPoint = &Loop.Last();
 			for (const FPoint2D& LoopPoint : Loop)
@@ -1662,7 +1660,7 @@ namespace CADKernel
 				int32 Index = 0;
 				for (; IndexV < CuttingCount[EIso::IsoV]; ++IndexV)
 				{
-					if (CuttingCoordinates[EIso::IsoV][IndexV] >= VMin)
+					if (UniformCuttingCoordinates[EIso::IsoV][IndexV] >= VMin)
 					{
 						break;
 					}
@@ -1671,7 +1669,7 @@ namespace CADKernel
 
 				for (; IndexV < CuttingCount[EIso::IsoV]; ++IndexV)
 				{
-					if (CuttingCoordinates[EIso::IsoV][IndexV] > Vmax)
+					if (UniformCuttingCoordinates[EIso::IsoV][IndexV] > Vmax)
 					{
 						break;
 					}
@@ -1684,17 +1682,17 @@ namespace CADKernel
 							continue;
 						}
 
-						if (CuttingCoordinates[EIso::IsoU][IndexU] < UMin)
+						if (UniformCuttingCoordinates[EIso::IsoU][IndexU] < UMin)
 						{
 							NbIntersectVForward[Index] = NbIntersectVForward[Index] > 0 ? 0 : 1;
 						}
-						else if (CuttingCoordinates[EIso::IsoU][IndexU] > Umax)
+						else if (UniformCuttingCoordinates[EIso::IsoU][IndexU] > Umax)
 						{
 							NbIntersectVBackward[Index] = NbIntersectVBackward[Index] > 0 ? 0 : 1;
 						}
 						else
 						{
-							double APvectAB = CuttingCoordinates[EIso::IsoV][IndexV] * ABu - CuttingCoordinates[EIso::IsoU][IndexU] * ABv + AuABVMinusAvABu;
+							double APvectAB = UniformCuttingCoordinates[EIso::IsoV][IndexV] * ABu - UniformCuttingCoordinates[EIso::IsoU][IndexU] * ABv + AuABVMinusAvABu;
 							if (APvectAB > SMALL_NUMBER)
 							{
 								NbIntersectVForward[Index] = NbIntersectVForward[Index] > 0 ? 0 : 1;
@@ -1715,7 +1713,7 @@ namespace CADKernel
 		}
 
 		// Intersection along V axis
-		for (const TArray<FPoint2D>& Loop : FaceLoops2D[(int32)EGridSpace::Default2D])
+		for (const TArray<FPoint2D>& Loop : FaceLoops2D[EGridSpace::UniformScaled])
 		{
 			const FPoint2D* FirstSegmentPoint = &Loop.Last();
 			for (const FPoint2D& LoopPoint : Loop)
@@ -1751,12 +1749,12 @@ namespace CADKernel
 				int32 Index = 0;
 				for (int32 IndexU = 0; IndexU < CuttingCount[EIso::IsoU]; ++IndexU)
 				{
-					if (CuttingCoordinates[EIso::IsoU][IndexU] < UMin)
+					if (UniformCuttingCoordinates[EIso::IsoU][IndexU] < UMin)
 					{
 						continue;
 					}
 
-					if (CuttingCoordinates[EIso::IsoU][IndexU] >= Umax)
+					if (UniformCuttingCoordinates[EIso::IsoU][IndexU] >= Umax)
 					{
 						continue;
 					}
@@ -1769,17 +1767,17 @@ namespace CADKernel
 							continue;
 						}
 
-						if (CuttingCoordinates[EIso::IsoV][IndexV] < VMin)
+						if (UniformCuttingCoordinates[EIso::IsoV][IndexV] < VMin)
 						{
 							NbIntersectUForward[Index] = NbIntersectUForward[Index] > 0 ? 0 : 1;
 						}
-						else if (CuttingCoordinates[EIso::IsoV][IndexV] > Vmax)
+						else if (UniformCuttingCoordinates[EIso::IsoV][IndexV] > Vmax)
 						{
 							NbIntersectUBackward[Index] = NbIntersectUBackward[Index] > 0 ? 0 : 1;
 						}
 						else
 						{
-							double APvectAB = CuttingCoordinates[EIso::IsoV][IndexV] * ABu - CuttingCoordinates[EIso::IsoU][IndexU] * ABv + AuABVMinusAvABu;
+							double APvectAB = UniformCuttingCoordinates[EIso::IsoV][IndexV] * ABu - UniformCuttingCoordinates[EIso::IsoU][IndexU] * ABv + AuABVMinusAvABu;
 							if (APvectAB > SMALL_NUMBER)
 							{
 								NbIntersectUBackward[Index] = NbIntersectUBackward[Index] > 0 ? 0 : 1;
@@ -1835,7 +1833,7 @@ namespace CADKernel
 		Chronos.FindInnerDomainPointsDuration += FChrono::Elapse(StartTime);
 	}
 
-	void FGrid::DisplayFindInnerDomainPoints(EGridSpace DisplaySpace) const
+	void FGrid::DisplayGridPoints(EGridSpace DisplaySpace) const
 	{
 		if (!bDisplay)
 		{
@@ -1847,7 +1845,7 @@ namespace CADKernel
 		{
 			if (IsInsideFace[Index])
 			{
-				DisplayPoint(Points2D[(int32)EGridSpace::Default2D][Index], Index);
+				DisplayPoint(Points2D[DisplaySpace][Index], Index);
 				NbNum++;
 			}
 		}
@@ -1859,7 +1857,7 @@ namespace CADKernel
 		{
 			if (!IsInsideFace[Index])
 			{
-				DisplayPoint(Points2D[(int32)EGridSpace::Default2D][Index], EVisuProperty::GreenPoint, Index);
+				DisplayPoint(Points2D[DisplaySpace][Index], EVisuProperty::OrangePoint, Index);
 			}
 		}
 		Close3DDebugSession();
@@ -1872,7 +1870,7 @@ namespace CADKernel
 			return;
 		}
 
-		Open3DDebugSession(TEXT("FGrid::FindPointsClosedToLoop result"));
+		F3DDebugSession _(TEXT("FGrid::FindPointsClosedToLoop result"));
 		for (int32 Index = 0; Index < CuttingSize; ++Index)
 		{
 			if (IsCloseToLoop[Index])
@@ -1884,17 +1882,16 @@ namespace CADKernel
 				DisplayPoint(Points2D[DisplaySpace][Index], EVisuProperty::YellowPoint);
 			}
 		}
-		Close3DDebugSession();
 	}
 
-	void FGrid::DisplayFindPointsCloseAndInsideToLoop(EGridSpace DisplaySpace) const
+	void FGrid::DisplayGridInnerPoints(EGridSpace DisplaySpace, TCHAR* Message) const
 	{
 		if (!bDisplay)
 		{
 			return;
 		}
 
-		Open3DDebugSession(TEXT("FGrid::FindPointsCloseAndInsideToLoop result"));
+		F3DDebugSession _(Message);
 		for (int32 Index = 0; Index < CuttingSize; ++Index)
 		{
 			if (IsInsideFace[Index])
@@ -1909,12 +1906,11 @@ namespace CADKernel
 				}
 			}
 		}
-		Close3DDebugSession();
 	}
 
 	bool FGrid::CheckIfDegenerated()
 	{
-		if (FaceLoops2D[(int32)EGridSpace::Default2D].Num() == 0)
+		if (FaceLoops2D[EGridSpace::Default2D].Num() == 0)
 		{
 			SetAsDegenerated();
 			return true;
@@ -1922,7 +1918,7 @@ namespace CADKernel
 
 		// if the external boundary is composed by only 2 points, the mesh of the surface is only an edge.
 		// The grid is degenerated.
-		if (FaceLoops2D[(int32)EGridSpace::Default2D][0].Num() < 3)
+		if (FaceLoops2D[EGridSpace::Default2D][0].Num() < 3)
 		{
 			SetAsDegenerated();
 			return true;
diff --git a/Engine/Source/Runtime/Datasmith/CADKernel/Private/CADKernel/Mesh/Structure/Grid.h b/Engine/Source/Runtime/Datasmith/CADKernel/Private/CADKernel/Mesh/Structure/Grid.h
--- a/Engine/Source/Runtime/Datasmith/CADKernel/Private/CADKernel/Mesh/Structure/Grid.h
+++ b/Engine/Source/Runtime/Datasmith/CADKernel/Private/CADKernel/Mesh/Structure/Grid.h
@@ -69,7 +69,7 @@ namespace CADKernel
 	protected:
 
 		const TSharedRef<FTopologicalFace> Face;
-		const FSurfacicTolerance& FaceTolerance;
+		const FSurfacicTolerance FaceTolerance;
 
 		TSharedRef<FModelMesh> MeshModel;
 
@@ -78,7 +78,7 @@ namespace CADKernel
 		/**
 		 * 2D Coordinate of Loop's nodes in each space
 		 */
-		TArray<TArray<FPoint2D>> FaceLoops2D[EGridSpace::Max];
+		TArray<TArray<FPoint2D>> FaceLoops2D[EGridSpace::EndGridSpace];
 
 		/**
 		 * 3D Coordinate of Loop nodes in each space
@@ -107,6 +107,7 @@ namespace CADKernel
 		 * Cutting coordinates of the face respecting the meshing criteria
 		 */
 		const FCoordinateGrid& CuttingCoordinates;
+		FCoordinateGrid UniformCuttingCoordinates;
 
 		/*
 		 * Maximum difference of coordinate along the specified axis of two successive cutting points
@@ -135,7 +136,7 @@ namespace CADKernel
 		/**
 		 * 2D Coordinate of grid nodes in each space
 		 */
-		TArray<FPoint2D> Points2D[EGridSpace::Max];
+		TArray<FPoint2D> Points2D[EGridSpace::EndGridSpace];
 
 		/**
 		 * 3D Coordinate of inner nodes
@@ -221,22 +222,6 @@ namespace CADKernel
 			OutNewScaledPoint /= (CuttingCoordinates[EIso::IsoU][IndexU + 1] - CuttingCoordinates[EIso::IsoU][IndexU + 0]) * (CuttingCoordinates[EIso::IsoV][IndexV + 1] - CuttingCoordinates[EIso::IsoV][IndexV + 0]);
 		};
 
-		void FindCoordinateIndex(const TArray<double>& InCoordinates, double Coordinate, int32& OutIndex) const
-		{
-			ensureCADKernel(InCoordinates.IsValidIndex(OutIndex));
-
-			while (Coordinate < InCoordinates[OutIndex] && OutIndex > 0)
-			{
-				OutIndex--;
-			}
-
-			for (; OutIndex + 2 < InCoordinates.Num() && Coordinate > InCoordinates[OutIndex + 1]; ++OutIndex)
-			{
-			}
-
-			ensureCADKernel(InCoordinates.IsValidIndex(OutIndex));
-		};
-
 		/**
 		 * @see ProcessPointClound (called in ProcessPointClound)
 		 */
@@ -410,6 +395,16 @@ namespace CADKernel
 			return CuttingCoordinates;
 		}
 
+		constexpr const TArray<double>& GetUniformCuttingCoordinatesAlongIso(EIso Iso) const
+		{
+			return UniformCuttingCoordinates[Iso];
+		}
+
+		const FCoordinateGrid& GetUniformCuttingCoordinates() const
+		{
+			return UniformCuttingCoordinates;
+		}
+
 		/**
 		 * @return the array of 3d points of the grid
 		 */
@@ -528,7 +523,8 @@ namespace CADKernel
 		// ======================================================================================================================================================================================================================
 		bool bDisplay = false;
 
-		void DisplayFindInnerDomainPoints(EGridSpace DisplaySpace) const;
+		void DisplayGridPoints(EGridSpace DisplaySpace) const;
+		void DisplayGridInnerPoints(EGridSpace DisplaySpace, TCHAR* Message) const;
 
 		template<typename TPoint>
 		void DisplayPoints(FString Message, const TArray<TPoint>& Points) const
@@ -565,7 +561,7 @@ namespace CADKernel
 					}
 					else
 					{
-						DisplayPoint(Points[Index], EVisuProperty::GreenPoint, Index);
+						DisplayPoint(Points[Index], EVisuProperty::OrangePoint, Index);
 					}
 				}
 			}
@@ -615,7 +611,6 @@ namespace CADKernel
 		}
 
 		void DisplayFindPointsCloseToLoop(EGridSpace DisplaySpace) const;
-		void DisplayFindPointsCloseAndInsideToLoop(EGridSpace DisplaySpace) const;
 		void PrintTimeElapse() const;
 	};
 
diff --git a/Engine/Source/Runtime/Datasmith/CADKernel/Private/CADKernel/Topo/Body.cpp b/Engine/Source/Runtime/Datasmith/CADKernel/Private/CADKernel/Topo/Body.cpp
--- a/Engine/Source/Runtime/Datasmith/CADKernel/Private/CADKernel/Topo/Body.cpp
+++ b/Engine/Source/Runtime/Datasmith/CADKernel/Private/CADKernel/Topo/Body.cpp
@@ -9,6 +9,21 @@ using namespace CADKernel;
 void FBody::AddShell(TSharedRef<FShell> Shell)
 {
 	Shells.Add(Shell);
+	Shell->HostedBy = StaticCastSharedRef<FBody>(AsShared());
+}
+
+void FBody::RemoveEmptyShell()
+{
+	TArray<TSharedPtr<FShell>> NewShells;
+	NewShells.Reserve(Shells.Num());
+	for (TSharedPtr<FShell> Shell : Shells)
+	{
+		if (Shell->FaceCount() > 0)
+		{
+			NewShells.Emplace(Shell);
+		}
+	}
+	Swap(NewShells, Shells);
 }
 
 TSharedPtr<FEntityGeom> FBody::ApplyMatrix(const FMatrixH& InMatrix) const 
@@ -24,7 +39,7 @@ TSharedPtr<FEntityGeom> FBody::ApplyMatrix(const FMatrixH& InMatrix) const
 #ifdef CADKERNEL_DEV
 FInfoEntity& FBody::GetInfo(FInfoEntity& Info) const
 {
-	return FTopologicalEntity::GetInfo(Info).Add(TEXT("shells"), Shells);
+	return FTopologicalEntity::GetInfo(Info).Add(TEXT("shells"), Shells).Add(*this);
 }
 #endif
 
diff --git a/Engine/Source/Runtime/Datasmith/CADKernel/Private/CADKernel/Topo/Joiner.cpp b/Engine/Source/Runtime/Datasmith/CADKernel/Private/CADKernel/Topo/Joiner.cpp
--- a/Engine/Source/Runtime/Datasmith/CADKernel/Private/CADKernel/Topo/Joiner.cpp
+++ b/Engine/Source/Runtime/Datasmith/CADKernel/Private/CADKernel/Topo/Joiner.cpp
@@ -2,7 +2,12 @@
 
 #include "CADKernel/Topo/Joiner.h"
 
+#include "CADKernel/Core/Chrono.h"
+#include "CADKernel/Core/Session.h"
 #include "CADKernel/Geo/Curves/RestrictionCurve.h"
+#include "CADKernel/Topo/Body.h"
+#include "CADKernel/Topo/Model.h"
+#include "CADKernel/Topo/Shell.h"
 #include "CADKernel/Topo/TopologicalEdge.h"
 #include "CADKernel/Topo/TopologicalLink.h"
 #include "CADKernel/Topo/TopologicalFace.h"
@@ -14,45 +19,127 @@
 
 using namespace CADKernel;
 
-FJoiner::FJoiner(TArray<TSharedPtr<FTopologicalFace>>& InSurfaces, double InTolerance)
-	: Surfaces(InSurfaces)
-	, JoiningTolerance(InTolerance)
-	, JoiningToleranceSqr(FMath::Square(JoiningTolerance))
+FJoiner::FJoiner(TSharedRef<FSession> InSession, const TArray<TSharedPtr<FTopologicalFace>>& InFaces, double InTolerance)
+	: Session(InSession)
+	, Faces(InFaces)
+	, JoiningTolerance(InTolerance * UE_SQRT_2)
+	, JoiningToleranceSquare(FMath::Square(JoiningTolerance))
 {
 }
 
-void FJoiner::JoinSurfaces()
+void FJoiner::RemoveFacesFromShell()
 {
+	// remove faces from their shells
+	TSet<TWeakPtr<FShell>> ShellSet;
+	for (const TSharedPtr<FTopologicalFace>& Face : Faces)
+	{
+		if (Face->GetHost().IsValid())
+		{
+			ShellSet.Add(Face->GetHost());
+			Face->ResetHost();
+		}	
+	}
+	for (TWeakPtr<FShell> WeakShell : ShellSet)
+	{
+		TSharedPtr<FShell> Shell = WeakShell.Pin();
+		bool bIsOutter = Shell->IsOutter();
 
+		TArray<FOrientedFace> ShellFace;
+		ShellFace.Reserve(Shell->FaceCount());
+		for (const FOrientedFace& Face : Shell->GetFaces())
+		{
+			if (!Face.Entity->GetHost().IsValid())
+			{
+				if (bIsOutter != (Face.Direction == EOrientation::Front))
+				{
+					Face.Entity->SetBackOriented();
+				}
+			}
+			else
+			{
+				ShellFace.Emplace(Face);
+			}
+		}
+		Shell->ReplaceFaces(ShellFace);
+	}
 }
 
-void FJoiner::JoinSurfaces(bool bProcessOnlyBorderEdges, bool bProcessOnlyNonManifoldEdges)
+FJoiner::FJoiner(TSharedRef<FSession> InSession, const TArray<TSharedPtr<FShell>>& InShells, double InTolerance)
+	: Session(InSession)
+	, Shells(InShells)
+	, JoiningTolerance(InTolerance)
+	, JoiningToleranceSquare(FMath::Square(JoiningTolerance))
 {
-	if (!(bProcessOnlyBorderEdges || bProcessOnlyNonManifoldEdges))
+	int32 FaceCount = 0;
+	for (const TSharedPtr<FShell>& Shell : InShells)
 	{
-		JoinVertices();
-		CheckSelfConnectedEdge();
+		Shell->CompleteMetadata();
+		FaceCount += Shell->FaceCount();
 	}
-	else
+	Faces.Reserve(FaceCount);
+
+	for (const TSharedPtr<FShell>& Shell : InShells)
+	{
+		Shell->SpreadBodyOrientation();
+		for (const FOrientedFace& Face : Shell->GetFaces())
+		{
+			Faces.Add(Face.Entity);
+		}
+		for (const FOrientedFace& Face : Shell->GetFaces())
+		{
+			Face.Entity->CompleteMetadata();
+		}
+	}
+
+	for (const TSharedPtr<FTopologicalFace>& Face : Faces)
+	{
+		Face->ResetMarker2();
+	}
+}
+
+
+void FJoiner::EmptyShells()
+{
+	for (const TSharedPtr<FShell>& Shell : Shells)
 	{
-		TArray<TSharedPtr<FTopologicalVertex>> NewMergeVertices;
-		JoinBorderVertices(NewMergeVertices);
-		JoinEdges(NewMergeVertices);
+		Shell->Empty();
 	}
+}
+
+void FJoiner::JoinFaces()
+{
+	FTimePoint StartJoinTime = FChrono::Now();
+
+	//RemoveIsolatedEdges();
+
+	TArray<TSharedPtr<FTopologicalVertex>> BorderVertices;
+	GetBorderVertices(BorderVertices);
+	MergeCoincidentVertices(BorderVertices);
+	//RemoveIsolatedEdges();
+
+	CheckSelfConnectedEdge();
+
+	//RemoveIsolatedEdges();
+
+	//Wait();
+	MergeCoincidentEdges(BorderVertices);
+
+	//Wait();
+	MergeUnconnectedAdjacentEdges();
+
+	//StitchParallelEdges(BorderVertices);
 
-	JoinEdges(Vertices);
-	TArray<TSharedPtr<FTopologicalVertex>> VerticesToProcess;
-	MergeUnconnectedAdjacentEdges(VerticesToProcess, true);
-	JoinEdges(VerticesToProcess);
+	FDuration InitKernelIODuration = FChrono::Elapse(StartJoinTime);
+	FChrono::PrintClockElapse(EVerboseLevel::Log, TEXT(""), TEXT("Join"), InitKernelIODuration);
 }
 
-void FJoiner::GetVertices()
+void FJoiner::GetVertices(TArray<TSharedPtr<FTopologicalVertex>>& Vertices)
 {
-	Vertices.Reserve(10 * Surfaces.Num());
+	Vertices.Empty(10 * Faces.Num());
 
-	for (TSharedPtr<FTopologicalFace> Surface : Surfaces)
+	for (TSharedPtr<FTopologicalFace> Face : Faces)
 	{
-		for (TSharedPtr<FTopologicalLoop> Loop : Surface->GetLoops())
+		for (TSharedPtr<FTopologicalLoop> Loop : Face->GetLoops())
 		{
 			for (const FOrientedEdge& OrientedEdge : Loop->GetEdges())
 			{
@@ -70,60 +157,85 @@ void FJoiner::GetVertices()
 			}
 		}
 	}
+	for (const TSharedPtr<FTopologicalVertex>& Vertex : Vertices)
+	{
+		Vertex->ResetMarker1();
+	}
 }
 
-void FJoiner::JoinVertices()
+void FJoiner::GetBorderVertices(TArray<TSharedPtr<FTopologicalVertex>>& BorderVertices)
 {
-	GetVertices();
+	TArray<TSharedPtr<FTopologicalVertex>> Vertices;
+	GetVertices(Vertices);
+
+	BorderVertices.Empty(Vertices.Num());
+
+	for (const TSharedPtr<FTopologicalVertex>& Vertex : Vertices)
+	{
+		if (Vertex->IsBorderVertex())
+		{
+			BorderVertices.Add(Vertex);
+		}
+	}
+}
+
+void FJoiner::MergeCoincidentVertices(TArray<TSharedPtr<FTopologicalVertex>>& VerticesToMerge)
+{
+	FTimePoint StartTime = FChrono::Now();
+
+	const double JoiningVerticesToleranceSquare = 2 * JoiningToleranceSquare;
+	const double WeigthTolerance = 3 * JoiningTolerance;
 
-	int32 VertexNum = (int32) Vertices.Num();
+	int32 VertexNum = (int32)VerticesToMerge.Num();
 
 	TArray<double> VerticesWeight;
 	VerticesWeight.Reserve(VertexNum);
-	TArray<int32> SortedVertexIndexes;
-	SortedVertexIndexes.Reserve(VertexNum);
-	TArray<TSharedPtr<FTopologicalVertex>> ActiveVertices;
-	ActiveVertices.Reserve(VertexNum);
 
-	for (TSharedPtr<FTopologicalVertex> Vertex: Vertices)
+	TArray<int32> SortedVertexIndices;
+	SortedVertexIndices.Reserve(VertexNum);
+
+	for (TSharedPtr<FTopologicalVertex> Vertex : VerticesToMerge)
 	{
-		Vertex->ResetMarker1();
 		VerticesWeight.Add(Vertex->GetCoordinates().X + Vertex->GetCoordinates().Y + Vertex->GetCoordinates().Z);
 	}
-	for (int32 Index = 0; Index < Vertices.Num(); ++Index)
+
+	for (int32 Index = 0; Index < VerticesToMerge.Num(); ++Index)
 	{
-		SortedVertexIndexes.Add(Index);
+		SortedVertexIndices.Add(Index);
 	}
-	SortedVertexIndexes.Sort([&VerticesWeight](const int32& Index1, const int32& Index2) { return VerticesWeight[Index1] < VerticesWeight[Index2]; });
+	SortedVertexIndices.Sort([&VerticesWeight](const int32& Index1, const int32& Index2) { return VerticesWeight[Index1] < VerticesWeight[Index2]; });
 
 	for (int32 IndexI = 0; IndexI < VertexNum; ++IndexI)
 	{
-		TSharedPtr<FTopologicalVertex> Vertex = Vertices[SortedVertexIndexes[IndexI]];
+		TSharedPtr<FTopologicalVertex> Vertex = VerticesToMerge[SortedVertexIndices[IndexI]];
 		if (Vertex->HasMarker1())
 		{
 			continue;
 		}
-		ActiveVertices.Add(Vertex);
+
+		ensureCADKernel(Vertex->IsActiveEntity());
+
 		Vertex->SetMarker1();
 
-		double VertexWeigth = VerticesWeight[SortedVertexIndexes[IndexI]];
+		double VertexWeigth = VerticesWeight[SortedVertexIndices[IndexI]];
 		FPoint Barycenter = Vertex->GetBarycenter();
-		for (int32 IndexJ = IndexI+1; IndexJ < VertexNum; ++IndexJ)
+
+		for (int32 IndexJ = IndexI + 1; IndexJ < VertexNum; ++IndexJ)
 		{
-			TSharedPtr<FTopologicalVertex> OtherVertex = Vertices[SortedVertexIndexes[IndexJ]];
+			TSharedPtr<FTopologicalVertex> OtherVertex = VerticesToMerge[SortedVertexIndices[IndexJ]];
 			if (OtherVertex->HasMarker1())
 			{
 				continue;
 			}
 
-			double OtherVertexWeigth = VerticesWeight[SortedVertexIndexes[IndexJ]];
-			if ((OtherVertexWeigth - VertexWeigth) > JoiningTolerance)
+			double OtherVertexWeigth = VerticesWeight[SortedVertexIndices[IndexJ]];
+			if ((OtherVertexWeigth - VertexWeigth) > WeigthTolerance)
 			{
 				break;
 			}
 
 			double DistanceSqr = OtherVertex->GetLinkActiveEntity()->SquareDistance(Barycenter);
-			if (DistanceSqr < JoiningToleranceSqr)
+			if (DistanceSqr < JoiningVerticesToleranceSquare)
 			{
 				OtherVertex->SetMarker1();
 				Vertex->Link(OtherVertex.ToSharedRef());
@@ -132,63 +244,79 @@ void FJoiner::JoinVertices()
 		}
 	}
 
-	for (TSharedPtr<FTopologicalVertex> Vertex : Vertices)
+	for (const TSharedPtr<FTopologicalVertex>& Vertex : VerticesToMerge)
 	{
 		Vertex->ResetMarker1();
 	}
-	Swap(ActiveVertices, Vertices);
-}
 
-void FJoiner::JoinBorderVertices(TArray<TSharedPtr<FTopologicalVertex>>& MergedVertices)
-{
-	if (!Vertices.Num())
+	TArray<TSharedPtr<FTopologicalVertex>> ActiveVertices;
+	ActiveVertices.Reserve(VertexNum);
+
+	for (const TSharedPtr<FTopologicalVertex>& Vertex : VerticesToMerge)
 	{
-		GetVertices();
+		TSharedPtr<FTopologicalVertex> ActiveVertex = Vertex->GetLinkActiveEntity();
+		if (ActiveVertex->HasMarker1())
+		{
+			continue;
+		}
+		ActiveVertex->SetMarker1();
+		ActiveVertices.Add(ActiveVertex);
+	}
+
+	for (const TSharedPtr<FTopologicalVertex>& Vertex : ActiveVertices)
+	{
+		Vertex->ResetMarker1();
 	}
 
+	Swap(ActiveVertices, VerticesToMerge);
+
+	FDuration Duration = FChrono::Elapse(StartTime);
+	FChrono::PrintClockElapse(EVerboseLevel::Log, TEXT("    "), TEXT("Merge Coincident vertices"), Duration);
+
+}
+
+void FJoiner::MergeBorderVerticesWithCoincidentOtherVertices(TArray<TSharedPtr<FTopologicalVertex>>& Vertices)
+{
 	int32 VertexNum = (int32)Vertices.Num();
 
 	TArray<double> VerticesWeight;
 	VerticesWeight.Reserve(VertexNum);
-	TArray<int32> SortedVertexIndexes;
-	SortedVertexIndexes.Reserve(VertexNum);
+	TArray<int32> SortedVertexIndices;
+	SortedVertexIndices.Reserve(VertexNum);
 
-	TArray<TSharedPtr<FTopologicalVertex>> ActiveVertices;
-	ActiveVertices.Reserve(VertexNum);
-
-	for (TSharedPtr<FTopologicalVertex> Vertex : Vertices)
+	for (const TSharedPtr<FTopologicalVertex>& Vertex : Vertices)
 	{
-		Vertex->ResetMarker1();
 		VerticesWeight.Add(Vertex->GetCoordinates().X + Vertex->GetCoordinates().Y + Vertex->GetCoordinates().Z);
 	}
+
 	for (int32 Index = 0; Index < Vertices.Num(); ++Index)
 	{
-		SortedVertexIndexes.Add(Index);
+		SortedVertexIndices.Add(Index);
 	}
-	SortedVertexIndexes.Sort([&VerticesWeight](const int32& Index1, const int32& Index2) { return VerticesWeight[Index1] < VerticesWeight[Index2]; });
+	SortedVertexIndices.Sort([&VerticesWeight](const int32& Index1, const int32& Index2) { return VerticesWeight[Index1] < VerticesWeight[Index2]; });
 
 	int32 StartIndexJ = 1;
-	for (int32 IndexI = 0; IndexI <VertexNum; ++IndexI)
+	for (int32 IndexI = 0; IndexI < VertexNum; ++IndexI)
 	{
-		TSharedRef<FTopologicalVertex> Vertex = Vertices[SortedVertexIndexes[IndexI]].ToSharedRef();
+		TSharedRef<FTopologicalVertex> Vertex = Vertices[SortedVertexIndices[IndexI]].ToSharedRef();
 		if (!Vertex->IsBorderVertex())
 		{
 			continue;
 		}
 
-		double VertexWeigth = VerticesWeight[SortedVertexIndexes[IndexI]];
+		double VertexWeigth = VerticesWeight[SortedVertexIndices[IndexI]];
 		FPoint Barycenter = Vertex->GetBarycenter();
 		for (int32 IndexJ = StartIndexJ; IndexJ < VertexNum; ++IndexJ)
 		{
-			TSharedRef<FTopologicalVertex> OtherVertex = Vertices[SortedVertexIndexes[IndexJ]].ToSharedRef();
+			TSharedRef<FTopologicalVertex> OtherVertex = Vertices[SortedVertexIndices[IndexJ]].ToSharedRef();
 			if (Vertex->GetLink() == OtherVertex->GetLink())
 			{
 				continue;
 			}
 
-			double OtherVertexWeigth = VerticesWeight[SortedVertexIndexes[IndexJ]];
+			double OtherVertexWeigth = VerticesWeight[SortedVertexIndices[IndexJ]];
 
-			if (OtherVertexWeigth + JoiningTolerance - VertexWeigth <0)
+			if (OtherVertexWeigth + JoiningTolerance - VertexWeigth < 0)
 			{
 				StartIndexJ = IndexJ;
 				continue;
@@ -200,7 +328,7 @@ void FJoiner::JoinBorderVertices(TArray<TSharedPtr<FTopologicalVertex>>& MergedV
 			}
 
 			double DistanceSqr = OtherVertex->GetLinkActiveEntity()->SquareDistance(Barycenter);
-			if (DistanceSqr < JoiningToleranceSqr)
+			if (DistanceSqr < JoiningToleranceSquare)
 			{
 				TArray<TSharedPtr<FTopologicalEdge>> CommonEdges;
 				Vertex->GetConnectedEdges(OtherVertex, CommonEdges);
@@ -212,238 +340,387 @@ void FJoiner::JoinBorderVertices(TArray<TSharedPtr<FTopologicalVertex>>& MergedV
 				OtherVertex->SetMarker1();
 				Vertex->Link(OtherVertex);
 				Barycenter = Vertex->GetBarycenter();
-				ActiveVertices.Add(Vertex);
 			}
 		}
 	}
 
-	MergedVertices.Reserve(ActiveVertices.Num());
-	for (TSharedPtr<FTopologicalVertex> Vertex : ActiveVertices)
+	TArray<TSharedPtr<FTopologicalVertex>> NewVertices;
+	NewVertices.Empty(VertexNum);
+
+	for (const TSharedPtr<FTopologicalVertex>& Vertex : Vertices)
 	{
-		if (Vertex->GetLinkActiveEntity()->HasMarker1())
+		const TSharedPtr<FTopologicalVertex>& ActiveVertex = Vertex->GetLinkActiveEntity();
+		if (ActiveVertex->HasMarker1())
 		{
 			continue;
 		}
-		MergedVertices.Add(Vertex->GetLinkActiveEntity());
-		MergedVertices.Last()->SetMarker1();
+		ActiveVertex->SetMarker1();
+		NewVertices.Add(ActiveVertex);
 	}
 
-	for (TSharedPtr<FTopologicalVertex> Vertex : MergedVertices)
+	for (TSharedPtr<FTopologicalVertex> Vertex : NewVertices)
 	{
 		Vertex->ResetMarker1();
 	}
-}
-
-// =========================================================================================================================================================================================================
-// =========================================================================================================================================================================================================
-// =========================================================================================================================================================================================================
-//
-//
-//                                                                            NOT YET REVIEWED
-//
-//
-// =========================================================================================================================================================================================================
-// =========================================================================================================================================================================================================
-// =========================================================================================================================================================================================================
 
+	Swap(NewVertices, Vertices);
+}
 
-void FJoiner::CheckSelfConnectedEdge()
+void FJoiner::MergeCoincidentEdges(TArray<TSharedPtr<FTopologicalVertex>>& VerticesToProcess)
 {
-	TArray<TSharedPtr<FTopologicalEdge>> CrunchedEdges;
-	CrunchedEdges.Reserve(100);
+	FTimePoint StartTime = FChrono::Now();
 
-	TArray<TSharedPtr<FTopologicalEdge>> IsolatedEdges;
-	IsolatedEdges.Reserve(100);
-
-	for (TSharedPtr<FTopologicalVertex> Vertex : Vertices)
+	for (TSharedPtr<FTopologicalVertex> VertexPtr : VerticesToProcess)
 	{
-		if (Vertex->GetTwinsEntities().Num() > 1)
+		ensureCADKernel(VertexPtr.IsValid());
+		TSharedRef<FTopologicalVertex> Vertex = VertexPtr.ToSharedRef();
+
+		TArray<TWeakPtr<FTopologicalEdge>> ConnectedEdges;
+		Vertex->GetConnectedEdges(ConnectedEdges);
+		int32 ConnectedEdgeCount = ConnectedEdges.Num();
+		if (ConnectedEdgeCount == 1)
 		{
-			for (TWeakPtr<FTopologicalVertex> TwinVertex : Vertex->GetTwinsEntities())
+			continue;
+		}
+
+		for (int32 EdgeI = 0; EdgeI < ConnectedEdgeCount - 1; ++EdgeI)
+		{
+			TSharedPtr<FTopologicalEdge> Edge = ConnectedEdges[EdgeI].Pin();
+			if (!Edge->IsActiveEntity())
+			{
+				continue;
+			}
+			bool bFirstEdgeBorder = Edge->IsBorder();
+			TSharedRef<FTopologicalVertex> EndVertex = Edge->GetOtherVertex(Vertex)->GetLinkActiveEntity();
+
+			for (int32 EdgeJ = EdgeI + 1; EdgeJ < ConnectedEdgeCount; ++EdgeJ)
 			{
-				TArray<TWeakPtr<FTopologicalEdge>> Edges = TwinVertex.Pin()->GetConnectedEdges();
-				for (TWeakPtr<FTopologicalEdge> WeakEdge : Edges)
+				TSharedPtr<FTopologicalEdge> SecondEdge = ConnectedEdges[EdgeJ].Pin();
+				if (!SecondEdge->IsActiveEntity())
 				{
-					TSharedPtr<FTopologicalEdge> Edge = WeakEdge.Pin();
-					if (!Edge->GetLoop().IsValid())
-					{
-						IsolatedEdges.Add(Edge);
-						continue;
-					}
+					continue;
+				}
+				bool bSecondEdgeBorder = Edge->IsBorder();
 
-					if (Edge->HasMarker1())
+				// Process only if at least one edge is Border
+				if (!bFirstEdgeBorder && !bSecondEdgeBorder)
+				{
+					continue;
+				}
+
+				TSharedRef<FTopologicalVertex> OtherEdgeEndVertex = SecondEdge->GetOtherVertex(Vertex)->GetLinkActiveEntity();
+
+				if (OtherEdgeEndVertex == EndVertex)
+				{
+					FPoint StartTangentEdge = Edge->GetTangentAt(Vertex);
+					FPoint StartTangentOtherEdge = SecondEdge->GetTangentAt(Vertex);
+
+					double CosAngle = StartTangentEdge.ComputeCosinus(StartTangentOtherEdge);
+					if (CosAngle < 0.9)
 					{
 						continue;
 					}
 
-					if (Edge->GetStartVertex()->GetLink() == Edge->GetEndVertex()->GetLink())
+					if (Edge->GetFace() != SecondEdge->GetFace())
 					{
-						ensureCADKernel(!Edge->IsDegenerated());
-						CrunchedEdges.Add(Edge);
-						Edge->SetMarker1();
+						Edge->Link(SecondEdge.ToSharedRef(), JoiningTolerance);
 					}
 				}
 			}
 		}
 	}
 
-	TArray<TSharedPtr<FTopologicalVertex>> NewVertices;
-	NewVertices.Reserve((CrunchedEdges.Num() + IsolatedEdges.Num()) * 2);
+	FDuration Duration = FChrono::Elapse(StartTime);
+
+	FChrono::PrintClockElapse(EVerboseLevel::Log, TEXT("    "), TEXT("Merge coincident edges"), Duration);
+
+}
+
+TSharedPtr<FTopologicalVertex> FJoiner::SplitAndLink(TSharedRef<FTopologicalVertex>& StartVertex, TSharedPtr<FTopologicalEdge>& EdgeToLink, TSharedPtr<FTopologicalEdge>& EdgeToSplit)
+{
+	TSharedPtr<FTopologicalVertex> VertexToLink = EdgeToLink->GetOtherVertex(StartVertex);
+
+	FPoint ProjectedPoint;
+	double UProjectedPoint = EdgeToSplit->ProjectPoint(VertexToLink->GetBarycenter(), ProjectedPoint);
 
-	TFunction<void(const TSharedPtr<TTopologicalLink<FTopologicalVertex>>&, const TSharedPtr<FTopologicalVertex>&, const TSharedPtr<FTopologicalVertex>&)> DeleteVertexOfLink = [&NewVertices](const TSharedPtr<TTopologicalLink<FTopologicalVertex>>& OldLink, const TSharedPtr<FTopologicalVertex>& Vertex1, const TSharedPtr<FTopologicalVertex>& Vertex2)
+	double SquareDistanceToProjectedPoint = ProjectedPoint.SquareDistance(VertexToLink->GetBarycenter());
+	if (SquareDistanceToProjectedPoint > JoiningToleranceSquare)
 	{
-		TSharedPtr<FTopologicalVertex> NewVertex = TSharedPtr<FTopologicalVertex>();
-		for (TWeakPtr<FTopologicalVertex> IVertex : OldLink->GetTwinsEntities())
-		{
-			ensureCADKernel(IVertex.IsValid());
-			IVertex.Pin()->ResetTopologicalLink();
-			if (IVertex == Vertex1 || IVertex == Vertex2)
-			{
-				continue;
-			}
+		return TSharedPtr<FTopologicalVertex>();
+	}
 
-			if (!NewVertex)
-			{
-				NewVertex = IVertex.Pin();
-				NewVertices.Add(NewVertex);
-			}
-			else
-			{
-				NewVertex->Link(IVertex.Pin().ToSharedRef());
-			}
-		}
-	};
+	// Check if the ProjectedPoint is not nearly equal edge boundary
+	TSharedPtr<FTopologicalVertex> EndVertex = EdgeToSplit->GetOtherVertex(StartVertex);
+	if (EndVertex->SquareDistance(ProjectedPoint) < JoiningToleranceSquare)
+	{
+		VertexToLink->Link(EndVertex.ToSharedRef());
+		EdgeToLink->Link(EdgeToSplit.ToSharedRef(), JoiningTolerance);
+		// TSharedPtr<FTopologicalVertex>() is returned as EndVertex is not new
+		return TSharedPtr<FTopologicalVertex>();
+	}
 
-	for (TSharedPtr<FTopologicalEdge> Edge : IsolatedEdges)
+	// JoinParallelEdges process all edges connected to startVertex (ConnectedEdges).
+	// Connected edges must remain compliant i.e. all edges of ConnectedEdges must be connected to StartVertex
+	// EdgeToSplit->SplitAt() must keep EdgeToSplit connected to StartVertex
+	bool bKeepStartVertexConnectivity = (StartVertex->GetLink() == EdgeToSplit->GetStartVertex()->GetLink());
+
+	TSharedPtr<FTopologicalEdge> NewEdge;
+	TSharedPtr<FTopologicalVertex> NewVertex = EdgeToSplit->SplitAt(UProjectedPoint, ProjectedPoint, bKeepStartVertexConnectivity, NewEdge);
+	if (!NewVertex.IsValid())
 	{
-		TSharedPtr<FTopologicalVertex> Vertex1 = Edge->GetStartVertex();
-		TSharedPtr<FTopologicalVertex> Vertex2 = Edge->GetEndVertex();
-		if (!Vertex1 && !Vertex2)
+		return TSharedPtr<FTopologicalVertex>();
+	}
+
+	VertexToLink->Link(NewVertex.ToSharedRef());
+	EdgeToLink->Link(EdgeToSplit.ToSharedRef(), JoiningTolerance);
+
+	return NewVertex;
+}
+
+int32 CountSplit = 0;
+
+void FJoiner::StitchParallelEdges(TArray<TSharedPtr<FTopologicalVertex>>& VerticesToProcess)
+{
+	FTimePoint StartTime = FChrono::Now();
+
+	for (int32 VertexI = 0; VertexI < VerticesToProcess.Num(); ++VertexI)
+	{
+		TSharedPtr<FTopologicalVertex> VertexPtr = VerticesToProcess[VertexI];
+		ensureCADKernel(VertexPtr.IsValid());
+
+		if (!VertexPtr->IsBorderVertex())
 		{
 			continue;
 		}
 
-		ensureCADKernel(!Vertex1 == !Vertex2);
+		TSharedRef<FTopologicalVertex> Vertex = VertexPtr.ToSharedRef();
 
-		TSharedPtr<TTopologicalLink<FTopologicalVertex>> OldLink1 = Vertex1->GetLink();
-		TSharedPtr<TTopologicalLink<FTopologicalVertex>> OldLink2 = Vertex2->GetLink();
+		TArray<TWeakPtr<FTopologicalEdge>> ConnectedEdges;
+		Vertex->GetConnectedEdges(ConnectedEdges);
+		int32 ConnectedEdgeCount = ConnectedEdges.Num();
+		if (ConnectedEdgeCount == 1)
+		{
+			continue;
+		}
 
-		Vertex1->RemoveConnectedEdge(Edge.ToSharedRef());
-		Vertex2->RemoveConnectedEdge(Edge.ToSharedRef());
-		if (!Vertex1->GetConnectedEdges().Num())
+		for (int32 EdgeI = 0; EdgeI < ConnectedEdgeCount - 1; ++EdgeI)
 		{
-			if (!Vertex2->GetConnectedEdges().Num())
+			TSharedPtr<FTopologicalEdge> Edge = ConnectedEdges[EdgeI].Pin();
+			ensureCADKernel(Edge->GetLoop().IsValid());
+
+			if (Edge->IsDegenerated())
 			{
-				DeleteVertexOfLink(OldLink1, Vertex1, Vertex2);
-				if (OldLink1 != OldLink2)
-				{
-					DeleteVertexOfLink(OldLink2, Vertex1, Vertex2);
-				}
+				continue;
 			}
-			else
+
+			if (!Edge->IsActiveEntity())
 			{
-				DeleteVertexOfLink(OldLink1, Vertex1, nullptr);
+				continue;
 			}
-			continue;
-		}
+			bool bFirstEdgeBorder = Edge->IsBorder();
 
-		if (!Vertex2->GetConnectedEdges().Num())
-		{
-			DeleteVertexOfLink(OldLink2, Vertex2, nullptr);
+			for (int32 EdgeJ = EdgeI + 1; EdgeJ < ConnectedEdgeCount; ++EdgeJ)
+			{
+				TSharedPtr<FTopologicalEdge> SecondEdge = ConnectedEdges[EdgeJ].Pin();
+				if (Edge->IsDegenerated())
+				{
+					continue;
+				}
+				if (!SecondEdge->IsActiveEntity())
+				{
+					continue;
+				}
+				bool bSecondEdgeBorder = Edge->IsBorder();
+
+				// Process only if at least one edge is Border
+				if (!bSecondEdgeBorder && !bSecondEdgeBorder)
+				{
+					continue;
+				}
+
+				FPoint StartTangentEdge = Edge->GetTangentAt(Vertex);
+				FPoint StartTangentOtherEdge = SecondEdge->GetTangentAt(Vertex);
+
+				double CosAngle = StartTangentEdge.ComputeCosinus(StartTangentOtherEdge);
+				if (CosAngle < 0.9)
+				{
+					continue;
+				}
+
+				TSharedRef<FTopologicalVertex> EndVertex = Edge->GetOtherVertex(Vertex)->GetLinkActiveEntity();
+				TSharedRef<FTopologicalVertex> OtherEdgeEndVertex = SecondEdge->GetOtherVertex(Vertex)->GetLinkActiveEntity();
+				if (EndVertex == OtherEdgeEndVertex)
+				{
+					Edge->Link(SecondEdge.ToSharedRef(), JoiningTolerance);
+				}
+				else
+				{
+					double EdgeLength = Edge->Length();
+					double OtherEdgeLength = SecondEdge->Length();
+					if (EdgeLength < OtherEdgeLength)
+					{
+						CountSplit++;
+						TSharedPtr<FTopologicalVertex> NewVertex = SplitAndLink(Vertex, Edge, SecondEdge);
+						if (NewVertex.IsValid())
+						{
+							VerticesToProcess.Add(NewVertex);
+						}
+					}
+					else
+					{
+						CountSplit++;
+						TSharedPtr<FTopologicalVertex> NewVertex = SplitAndLink(Vertex, SecondEdge, Edge);
+						if (NewVertex.IsValid())
+						{
+							VerticesToProcess.Add(NewVertex);
+						}
+					}
+				}
+			}
 		}
 	}
 
-	for(TSharedPtr<FTopologicalEdge> Edge : CrunchedEdges)
-	{
-		Edge->ResetMarkers();
+	FDuration Duration = FChrono::Elapse(StartTime);
+	FChrono::PrintClockElapse(EVerboseLevel::Log, TEXT("    "), TEXT("Stitch Parallel Edges"), Duration);
+}
 
-		TSharedPtr<FTopologicalVertex> Vertex1 = Edge->GetStartVertex();
-		TSharedPtr<FTopologicalVertex> Vertex2 = Edge->GetEndVertex();
-		TSharedPtr<TTopologicalLink<FTopologicalVertex>> OldLink = Vertex1->GetLink();
+void FJoiner::MergeUnconnectedAdjacentEdges()
+{
+	FTimePoint StartTime = FChrono::Now();
 
-		for (TWeakPtr<FTopologicalVertex> Vertex : OldLink->GetTwinsEntities())
-		{
-			Vertex.Pin()->ResetTopologicalLink();
-		}
+	for (TSharedPtr<FTopologicalFace> Face : Faces)
+	{
+		TArray<TArray<FOrientedEdge>> ArrayOfCandidates;
+		ArrayOfCandidates.Reserve(10);
 
-		for (TWeakPtr<FTopologicalVertex> Vertex : OldLink->GetTwinsEntities())
+		// For each loop, find unconnected successive edges...
+		for (TSharedPtr<FTopologicalLoop> Loop : Face->GetLoops())
 		{
-			double Distance1 = Vertex1->Distance(Vertex.Pin().ToSharedRef());
-			double Distance2 = Vertex2->Distance(Vertex.Pin().ToSharedRef());
-			if (Distance1 <Distance2)
+			TArray<FOrientedEdge>& Edges = Loop->GetEdges();
+			int32 EdgeCount = Edges.Num();
+
+
+			// Find the starting edge i.e. the next edge of the first edge that its ending vertex is connecting to 3 or more edges 
+			// The algorithm start to the last edge of the loop, if it verifies the criteria then the first edge is the edges[0]
+			int32 EndIndex = EdgeCount;
 			{
-				Vertex1->Link(Vertex.Pin().ToSharedRef());
+				TSharedPtr<FTopologicalVertex> EndVertex;
+				do
+				{
+					EndIndex--;
+					EndVertex = Edges[EndIndex].Direction == EOrientation::Front ? Edges[EndIndex].Entity->GetEndVertex() : Edges[EndIndex].Entity->GetStartVertex();
+				} while (EndVertex->ConnectedEdgeCount() == 2 && EndIndex > 0);
 			}
-			else
+			EndIndex++;
+
+			// First step
+			// For the loop, find all arrays of successive unconnected edges
+			TArray<FOrientedEdge>* Candidates = &ArrayOfCandidates.Emplace_GetRef();
+			Candidates->Reserve(10);
+			bool bCanStop = false;
+			for (int32 Index = EndIndex; bCanStop == (Index != EndIndex); ++Index)
 			{
-				Vertex2->Link(Vertex.Pin().ToSharedRef());
-			}
-		}
+				if (Index == EdgeCount)
+				{
+					Index = 0;
+				}
+				bCanStop = true;
 
-		NewVertices.Add(Vertex1->GetLinkActiveEntity());
-		NewVertices.Add(Vertex2->GetLinkActiveEntity());
-	}
+				FOrientedEdge& Edge = Edges[Index];
+				if (Edge.Entity->GetTwinsEntityCount() == 1)
+				{
+					TSharedPtr<FTopologicalVertex> EndVertex = Edge.Direction == EOrientation::Front ? Edge.Entity->GetEndVertex() : Edge.Entity->GetStartVertex();
 
-	TArray<TSharedPtr<FTopologicalVertex>> ActiveVertices;
-	ActiveVertices.Reserve(Vertices.Num() + NewVertices.Num());
+					TArray<TWeakPtr<FTopologicalEdge>> ConnectedEdges;
+					EndVertex->GetConnectedEdges(ConnectedEdges);
 
-	for (TSharedPtr<FTopologicalVertex> Vertex : NewVertices)
-	{
-		if (Vertex->GetLinkActiveEntity()->HasMarker1())
-		{
-			continue;
-		}
-		Vertex->GetLinkActiveEntity()->SetMarker1();
-		ActiveVertices.Add(Vertex->GetLinkActiveEntity());
-	}
+					bool bEdgeIsNotTheLast = false;
+					if (ConnectedEdges.Num() == 2)
+					{
+						// check if the edges are tangents
+						FPoint StartTangentEdge = ConnectedEdges[0].Pin()->GetTangentAt(EndVertex.ToSharedRef());
+						FPoint StartTangentOtherEdge = ConnectedEdges[1].Pin()->GetTangentAt(EndVertex.ToSharedRef());
 
-	int32 NumOfVertices = (int32)ActiveVertices.Num();
+						double CosAngle = StartTangentEdge.ComputeCosinus(StartTangentOtherEdge);
+						if (CosAngle < -0.9)
+						{
+							bEdgeIsNotTheLast = true;
+						}
+					}
 
-	for (TSharedPtr<FTopologicalVertex> Vertex : Vertices)
-	{
-		if (Vertex->GetLinkActiveEntity()->HasMarker1())
-		{
-			continue;
+					if (bEdgeIsNotTheLast || Candidates->Num() > 0)
+					{
+						Candidates->Add(Edges[Index]);
+					}
+
+					if (!bEdgeIsNotTheLast && Candidates->Num() > 0)
+					{
+						Candidates = &ArrayOfCandidates.Emplace_GetRef();
+						Candidates->Reserve(10);
+					}
+				}
+			}
 		}
-		ActiveVertices.Add(Vertex->GetLinkActiveEntity());
-	}
 
-	for (int32 Index = 0; Index <NumOfVertices; Index++)
-	{
-		ActiveVertices[Index]->ResetMarker1();
-	}
+		// Second step, 
+		// Each array of edges are merged to generated a single edge, the loop is updated
+		for (TArray<FOrientedEdge>& Candidates : ArrayOfCandidates)
+		{
+			if (Candidates.Num())
+			{
+				TSharedPtr<FTopologicalVertex> StartVertexPtr = Candidates[0].Direction == EOrientation::Front ? Candidates[0].Entity->GetStartVertex() : Candidates[0].Entity->GetEndVertex();
+				TSharedRef<FTopologicalVertex> StartVertex = StartVertexPtr.ToSharedRef();
+				TSharedPtr<FTopologicalVertex> EndVertex = Candidates.Last().Direction == EOrientation::Front ? Candidates.Last().Entity->GetEndVertex() : Candidates.Last().Entity->GetStartVertex();
 
-	Swap(ActiveVertices, Vertices);
-}
+				FPoint StartTangentEdge = Candidates[0].Entity->GetTangentAt(StartVertex);
 
-void FJoiner::FixCollapsedEdges()
-{
-	TArray<TSharedPtr<FTopologicalEdge>> CollapsedEdges;
+				TArray<TSharedPtr<FTopologicalEdge>> Edges;
+				StartVertex->GetConnectedEdges(EndVertex, Edges);
 
-	int32 NbDegeneratedEdges = 0;
-	for (const TSharedPtr<FTopologicalFace>& Surface : Surfaces)
-	{
-		for (const TSharedPtr<FTopologicalLoop>& Loop : Surface->GetLoops())
-		{
-			for (const FOrientedEdge& OrientedEdge : Loop->GetEdges())
-			{
-				const TSharedPtr<FTopologicalEdge>& Edge = OrientedEdge.Entity;
-				for (const TWeakPtr<FTopologicalEdge>& TwinEdge : Edge->GetTwinsEntities())
+				// Remove edge that is not parallel edge
+				TSharedPtr<FTopologicalEdge> ParallelEdge;
+				for (TSharedPtr<FTopologicalEdge>& Edge : Edges)
 				{
-					TSharedPtr<FTopologicalFace> Face = TwinEdge.Pin()->GetFace();
-					if (!Face.IsValid())
+					if (Edge->GetFace() == Face)
 					{
+						Edge.Reset();
 						continue;
 					}
+					
+					FPoint StartTangentOtherEdge = Edge->GetTangentAt(StartVertex);
+					double CosAngle = StartTangentEdge.ComputeCosinus(StartTangentOtherEdge);
+					if (CosAngle > 0.9)
+					{
+						ParallelEdge = Edge;
+						break;
+					}
 				}
 
-				if (Edge->GetStartVertex()->GetLink() == Edge->GetEndVertex()->GetLink())
+				if (ParallelEdge.IsValid())
 				{
-					if (!Edge->IsDegenerated())
+					//FMessage::Printf(Log, TEXT("Face %d UnconnectedAdjacentEdges\n"), Face->GetId());
+					//for (auto Edge : Candidates)
+					//{
+					//	FMessage::Printf(Log, TEXT("       - %d \n"), Edge.Entity->GetId());
+					//}
+
+					TSharedPtr<FTopologicalEdge> NewEdge = FTopologicalEdge::CreateEdgeByMergingEdges(Candidates, StartVertex, EndVertex.ToSharedRef());
+					if (!NewEdge.IsValid())
 					{
-						for (const TWeakPtr<FTopologicalEdge>& TwinEdge : Edge->GetTwinsEntities())
+						// the edges cannot be merge, they will be connected to the parallel edge with "SplitAndLink"
+						// i.e. instead of merging edges to be one and link to the parallel one
+						// the parallel one is split at the extremities of each edges, and each new edge is linked to its parallel edge 
+						break;
+					}
+
+					// Link to the parallel edge
+					// New edge is link to the first parallel one as the other parallel should already be linked together
+					for (TSharedPtr<FTopologicalEdge>& Edge : Edges)
+					{
+						if (Edge.IsValid())
 						{
-							TwinEdge.Pin()->SetAsDegenerated();
+							Edge->Link(NewEdge.ToSharedRef(), JoiningTolerance);
+							break;
 						}
 					}
 				}
@@ -451,142 +728,398 @@ void FJoiner::FixCollapsedEdges()
 		}
 	}
 
-	for (TSharedPtr<FTopologicalEdge>& Edge : CollapsedEdges)
-	{
-		TSharedPtr<FTopologicalLoop> Loop = Edge->GetLoop();
-		Loop->RemoveEdge(Edge);
+	FDuration Duration = FChrono::Elapse(StartTime);
 
-		TSharedPtr<FTopologicalVertex> StartVertex = Edge->GetStartVertex();
-		TSharedPtr<FTopologicalVertex> EndVertex = Edge->GetEndVertex();
+	FChrono::PrintClockElapse(EVerboseLevel::Log, TEXT("    "), TEXT("Merge unconnected adjacent edges"), Duration);
 
-		Edge->Delete();
+}
 
-		StartVertex->DeleteIfIsolated();
-		EndVertex->DeleteIfIsolated();
+void FJoiner::RemoveIsolatedEdges()
+{
+	FTimePoint StartTime = FChrono::Now();
 
-		if (Loop->GetEdges().Num() == 0)
+	TArray<TSharedPtr<FTopologicalEdge>> IsolatedEdges;
+
+	TArray<TSharedPtr<FTopologicalVertex>> Vertices;
+	GetVertices(Vertices);
+
+	for (const TSharedPtr<FTopologicalVertex>& Vertex : Vertices)
+	{
+		for (const TWeakPtr<FTopologicalVertex>& TwinVertex : Vertex->GetTwinsEntities())
 		{
-			TSharedPtr<FTopologicalFace> Face = Loop->GetFace();
-			Face->RemoveLoop(Loop);
+			TArray<TWeakPtr<FTopologicalEdge>> Edges = Vertex->GetDirectConnectedEdges();
+			for (TWeakPtr<FTopologicalEdge> WeakEdge : Edges)
+			{
+				TSharedPtr<FTopologicalEdge> Edge = WeakEdge.Pin();
+				if (!Edge->GetLoop().IsValid())
+				{
+					IsolatedEdges.Add(Edge);
+				}
+			}
 		}
 	}
-}
+	FDuration Duration = FChrono::Elapse(StartTime);
 
+	FChrono::PrintClockElapse(EVerboseLevel::Log, TEXT("    "), TEXT("Remove Isolated Edges"), Duration);
 
+	FMessage::Printf(EVerboseLevel::Log, TEXT("\n\nIsolatedEdges count %d\n\n\n"), IsolatedEdges.Num());
+}
 
-void FJoiner::MergeUnconnectedAdjacentEdges(TArray<TSharedPtr<FTopologicalVertex>>& NewEdgeVertices, bool bOnlyIfSameCurve)
+// =========================================================================================================================================================================================================
+// =========================================================================================================================================================================================================
+// =========================================================================================================================================================================================================
+//
+//
+//                                                                            NOT YET REVIEWED
+//
+//
+// =========================================================================================================================================================================================================
+// =========================================================================================================================================================================================================
+// =========================================================================================================================================================================================================
+
+void FJoiner::CheckSelfConnectedEdge()
 {
-// TODO 
-	ensureCADKernel(false);
-#ifdef TODO
+	FTimePoint StartTime = FChrono::Now();
 
-	NewEdgeVertices.Reserve(Vertices.Num());
-	for (TSharedPtr<FTopologicalVertex> Vertex : Vertices)
+	FMessage::Printf(Log, TEXT("    Self connected edges\n"));
+	for (TSharedPtr<FTopologicalFace> Face : Faces)
 	{
-		TArray<TWeakPtr<FTopologicalEdge>> Edges = Vertex->GetConnectedEdges();
-		if (Edges.Num() == 2)
+		for (TSharedPtr<FTopologicalLoop> Loop : Face->GetLoops())
 		{
-			// TODO Work in 2D
-			FPoint Edge0Tangent;
-			Edges[0].Pin()->GetTangentAt(Vertex, Edge0Tangent);
-
-			FPoint Edge1Tangent;
-			Edges[1].Pin()->GetTangentAt(Vertex, Edge1Tangent);
-
-			double CosAngle = Edge0Tangent.ComputeCosinus(Edge1Tangent);
-			if (CosAngle < -0.75)
+			for (const FOrientedEdge& OrientedEdge : Loop->GetEdges())
 			{
-				TArray<EOrientation> EdgeDirections;
-				EdgeDirections.Reserve(2);
-				EdgeDirections.Add(Edges[0].Pin()->GetEndVertex() == Vertex ? EOrientation::Front : EOrientation::Back);
-				EdgeDirections.Add(Edges[1].Pin()->GetStartVertex() == Vertex ? EOrientation::Front : EOrientation::Back);
-				TSharedPtr<FTopologicalEdge> ExtendedEdge = nullptr;
-
-				// TODO Check better if the both curve are similar (2DCurve as to be similar)
-				if (Edges[0].Pin()->GetCurve() != Edges[1].Pin()->GetCurve())
+				TSharedPtr<FTopologicalEdge> Edge = OrientedEdge.Entity;
+				if (Edge->GetStartVertex()->IsLinkedTo(Edge->GetEndVertex()))
 				{
-					//ExtendedEdge = FEdge::ReplaceEdgesByOne(Edges, EdgeDirections, GetTolerance());
-				}
-				else
-				{
-					//ExtendedEdge = FEdge::CreateEdgeToMerge2Edges(Edges, EdgeDirections);
-				}
-
-				if (ExtendedEdge)
-				{
-					NewEdgeVertices.Add(ExtendedEdge->GetStartVertex());
-					NewEdgeVertices.Add(ExtendedEdge->GetEndVertex());
+					if (!Edge->IsDegenerated() && Edge->Length() < 2 * JoiningTolerance)
+					{
+						FMessage::Printf(Debug, TEXT("Face %d Edge %d was self connected, "), Face->GetId(), Edge->GetId());
+						Edge->GetStartVertex()->UnlinkTo(Edge->GetEndVertex());
+					}
 				}
-				continue;
 			}
 		}
 	}
-#endif
-}
+	FDuration Duration = FChrono::Elapse(StartTime);
 
+	FChrono::PrintClockElapse(EVerboseLevel::Log, TEXT("    "), TEXT("Unconnect Self connected edges"), Duration);
+}
 
-void FJoiner::JoinEdges(TArray<TSharedPtr<FTopologicalVertex>>& VerticesToProcess)
+void FJoiner::SplitIntoConnectedShell()
 {
-	ensureCADKernel(false);
-#ifdef TODO
+	// Processed1 : Surfaces added in CandidateSurfacesForMesh
+
+	int32 TopologicalFaceCount = Faces.Num();
+	// Is closed ?
+	// Is one shell ?
 
-	TArray<TSharedPtr<FTopologicalEdge>> ProcessedEdges;
-	ProcessedEdges.Reserve(VerticesToProcess.Num());
+	TArray<FFaceSubset> SubShells;
 
-	TArray<TWeakPtr<FTopologicalEdge>> ConnectedEdges;
-	for (TSharedPtr<FTopologicalVertex> Vertex : VerticesToProcess)
+	int32 ProcessFaceCount = 0;
+
+	TArray<TSharedPtr<FTopologicalFace>> Front;
+	TFunction<void(const TSharedPtr<FTopologicalFace>&, FFaceSubset&)> GetNeighboringFaces = [&](const TSharedPtr<FTopologicalFace>& Face, FFaceSubset& Shell)
 	{
-		ConnectedEdges.Empty();
-		Vertex->GetConnectedEdges(ConnectedEdges);
-		for (int32 EdgeI = 0; EdgeI < ConnectedEdges.Num(); ++EdgeI)
+		for (const TSharedPtr<FTopologicalLoop>& Loop : Face->GetLoops())
 		{
-			TSharedRef<FTopologicalEdge> Edge = ConnectedEdges[EdgeI].Pin()->GetLinkActiveEdge();
-			if (Edge->HasMarker1())
-			{
-				continue;
-			}
-			TSharedRef<FTopologicalVertex> EndVertex = Edge->GetOtherVertex(Vertex.ToSharedRef())->GetLinkActiveEntity();
-
-			for (int32 EdgeJ = EdgeI; EdgeJ < ConnectedEdges.Num(); ++EdgeJ)
+			for (const FOrientedEdge& OrientedEdge : Loop->GetEdges())
 			{
-				TSharedRef<FTopologicalEdge> OtherEdge = ConnectedEdges[EdgeJ].Pin()->GetLinkActiveEdge();
-				if (OtherEdge == Edge)
+				const TSharedPtr<FTopologicalEdge>& Edge = OrientedEdge.Entity;
+				if (Edge->HasMarker1())
 				{
 					continue;
 				}
-				if (OtherEdge->HasMarker1())
+				Edge->SetMarker1();
+
+				if (Edge->GetTwinsEntityCount() == 1)
 				{
+					if (!Edge->IsDegenerated())
+					{
+						Shell.BorderEdgeCount++;
+					}
 					continue;
 				}
-				TSharedRef<FTopologicalVertex> OtherEdgeEndVertex = OtherEdge->GetOtherVertex(Vertex.ToSharedRef())->GetLinkActiveEntity();
 
-				if (OtherEdgeEndVertex == EndVertex)
+				if (Edge->GetTwinsEntityCount() > 2)
 				{
-					FPoint StartTangentEdge;
-					Edge->GetTangentAt(Vertex, StartTangentEdge);
-					FPoint StartTangentOtherEdge;
-					OtherEdge->GetTangentAt(Vertex, StartTangentOtherEdge);
+					Shell.NonManifoldEdgeCount++;
+				}
 
-					double CosAngle = StartTangentEdge.ComputeCosinus(StartTangentOtherEdge);
-					if (CosAngle < 0.9)
+				for (TWeakPtr<FTopologicalEdge> WeakEdge : Edge->GetTwinsEntities())
+				{
+					TSharedPtr<FTopologicalEdge> NextEdge = WeakEdge.Pin();
+					if (NextEdge->HasMarker1())
 					{
 						continue;
 					}
+					NextEdge->SetMarker1();
 
-					if (Edge->GetFace() != OtherEdge->GetFace())
+					TSharedPtr<FTopologicalFace> NextFace = NextEdge->GetFace();
+					if (!NextFace.IsValid())
 					{
-						Edge->Link(OtherEdge);
-						Edge->GetLinkActiveEdge()->SetMarker1();
-						ProcessedEdges.Add(Edge->GetLinkActiveEdge());
+						continue;
 					}
+
+					if (NextFace->HasMarker1())
+					{
+						continue;
+					}
+					NextFace->SetMarker1();
+					Front.Add(NextFace);
 				}
 			}
 		}
+	};
+
+	TFunction<void(FFaceSubset&)> SpreadFront = [&](FFaceSubset& Shell)
+	{
+		while (Front.Num())
+		{
+			TSharedPtr<FTopologicalFace> Face = Front.Pop();
+			Shell.Faces.Add(Face);
+			GetNeighboringFaces(Face, Shell);
+		}
+	};
+
+	for (const TSharedPtr<FTopologicalFace>& Face : Faces)
+	{
+		if (Face->HasMarker1())
+		{
+			continue;
+		}
+
+		FFaceSubset& Shell = SubShells.Emplace_GetRef();
+		Shell.Faces.Reserve(TopologicalFaceCount - ProcessFaceCount);
+		Front.Empty(TopologicalFaceCount);
+
+		Face->SetMarker1();
+		Front.Add(Face);
+		SpreadFront(Shell);
+		ProcessFaceCount += Shell.Faces.Num();
+
+		if (ProcessFaceCount == TopologicalFaceCount)
+		{
+			break;
+		}
 	}
 
-	for (TSharedPtr<FTopologicalEdge> Edge : ProcessedEdges)
+	// reset Marker
+	for (const TSharedPtr<FTopologicalFace>& Face : Faces)
 	{
-		Edge->ResetMarker1();
+		Face->ResetMarkers();
+		for (const TSharedPtr<FTopologicalLoop>& Loop : Face->GetLoops())
+		{
+			for (const FOrientedEdge& OrientedEdge : Loop->GetEdges())
+			{
+				OrientedEdge.Entity->ResetMarkers();
+			}
+		}
 	}
-#endif
+
+	// for each FaceSubset, find the main shell
+	for (FFaceSubset& FaceSubset : SubShells)
+	{
+		TMap<TWeakPtr<FShell>, int32> ShellToFaceCount;
+		TMap<uint32, int32> ColorToFaceCount;
+		TMap<FString, int32> NameToFaceCount;
+
+		for (const TSharedPtr<FTopologicalFace>& Face : FaceSubset.Faces)
+		{
+			TWeakPtr<FShell> Shell = Face->GetHost();
+			int32* Count = ShellToFaceCount.Find(Shell);
+			if (Count)
+			{
+				(*Count)++;
+			}
+			else
+			{
+				ShellToFaceCount.Add(Shell, 1);
+			}
+
+			ColorToFaceCount.FindOrAdd(Face->GetColorId())++;
+			NameToFaceCount.FindOrAdd(Face->GetName())++;
+		}
+
+		int32 MaxInstance = 0;
+		TWeakPtr<FShell> MainShell;
+		for (TPair<TWeakPtr<FShell>, int32>& Pair : ShellToFaceCount)
+		{
+			if (Pair.Value > MaxInstance)
+			{
+				MaxInstance = Pair.Value;
+				MainShell = Pair.Key;
+			}
+		}
+
+		if (MainShell.IsValid() && MainShell.Pin()->FaceCount() / 2 + 1 < MaxInstance)
+		{
+			FaceSubset.MainShell = MainShell;
+		}
+
+		MaxInstance = FaceSubset.Faces.Num()/3;
+		FString MainName;
+		for (TPair<FString, int32>& Pair : NameToFaceCount)
+		{
+			if (Pair.Value > MaxInstance)
+			{
+				MaxInstance = Pair.Value;
+				MainName = Pair.Key;
+			}
+		}
+
+		MaxInstance = 0;
+		for (TPair<uint32, int32>& Pair : ColorToFaceCount)
+		{
+			if (Pair.Value > MaxInstance)
+			{
+				MaxInstance = Pair.Value;
+				FaceSubset.MainColor = Pair.Key;
+			}
+		}
+
+	}
+
+	if (Shells.Num())
+	{
+		EmptyShells();
+	}
+	else
+	{
+		RemoveFacesFromShell();
+	}
+
+	// for each FaceSubset, process the Shell
+	for (FFaceSubset FaceSubset : SubShells)
+	{
+		if (FaceSubset.MainShell.IsValid())
+		{
+			TSharedPtr<FShell> Shell = FaceSubset.MainShell.Pin();
+			Shell->Empty(FaceSubset.Faces.Num());
+			Shell->Add(FaceSubset.Faces);
+		}
+		else
+		{
+			TSharedRef<FBody> Body = FEntity::MakeShared<FBody>();
+			Session->GetModel()->Add(Body);
+			TSharedRef<FShell> Shell = FEntity::MakeShared<FShell>();
+			Body->AddShell(Shell);
+			Body->SetName(FaceSubset.MainName);
+			Body->SetColorId(FaceSubset.MainColor);
+
+			Shell->Add(FaceSubset.Faces);
+			Shell->SetName(FaceSubset.MainName);
+			Shell->SetColorId(FaceSubset.MainColor);
+		}
+	}
+
+	Session->GetModel()->RemoveEmptyBodies();
 }
+
+
+
+//void FJoiner::FixCollapsedEdges()
+//{
+//	TArray<TSharedPtr<FTopologicalEdge>> CollapsedEdges;
+//
+//	int32 NbDegeneratedEdges = 0;
+//	for (const TSharedPtr<FTopologicalFace>& Face : Faces)
+//	{
+//		for (const TSharedPtr<FTopologicalLoop>& Loop : Face->GetLoops())
+//		{
+//			for (const FOrientedEdge& OrientedEdge : Loop->GetEdges())
+//			{
+//				const TSharedPtr<FTopologicalEdge>& Edge = OrientedEdge.Entity;
+//				for (const TWeakPtr<FTopologicalEdge>& TwinEdge : Edge->GetTwinsEntities())
+//				{
+//					TSharedPtr<FTopologicalFace> Face = TwinEdge.Pin()->GetFace();
+//					if (!Face.IsValid())
+//					{
+//						continue;
+//					}
+//				}
+//
+//				if (Edge->GetStartVertex()->GetLink() == Edge->GetEndVertex()->GetLink())
+//				{
+//					if (!Edge->IsDegenerated())
+//					{
+//						for (const TWeakPtr<FTopologicalEdge>& TwinEdge : Edge->GetTwinsEntities())
+//						{
+//							TwinEdge.Pin()->SetAsDegenerated();
+//						}
+//					}
+//				}
+//			}
+//		}
+//	}
+//
+//	for (TSharedPtr<FTopologicalEdge>& Edge : CollapsedEdges)
+//	{
+//		TSharedPtr<FTopologicalLoop> Loop = Edge->GetLoop();
+//		Loop->RemoveEdge(Edge);
+//
+//		TSharedPtr<FTopologicalVertex> StartVertex = Edge->GetStartVertex();
+//		TSharedPtr<FTopologicalVertex> EndVertex = Edge->GetEndVertex();
+//
+//		Edge->Delete();
+//
+//		StartVertex->DeleteIfIsolated();
+//		EndVertex->DeleteIfIsolated();
+//
+//		if (Loop->GetEdges().Num() == 0)
+//		{
+//			TSharedPtr<FTopologicalFace> Face = Loop->GetFace();
+//			Face->RemoveLoop(Loop);
+//		}
+//	}
+//}
+
+
+
+
+
+
+/*
+NewEdgeVertices.Reserve(AllVertices.Num());
+for (TSharedPtr<FTopologicalVertex> Vertex : AllVertices)
+{
+	TArray<TWeakPtr<FTopologicalEdge>> Edges = Vertex->GetConnectedEdges();
+	if (Edges.Num() == 2)
+	{
+		// TODO Work in 2D
+		FPoint Edge0Tangent;
+		Edges[0].Pin()->GetTangentAt(Vertex, Edge0Tangent);
+
+		FPoint Edge1Tangent;
+		Edges[1].Pin()->GetTangentAt(Vertex, Edge1Tangent);
+
+		double CosAngle = Edge0Tangent.ComputeCosinus(Edge1Tangent);
+		if (CosAngle < -0.75)
+		{
+			TArray<EOrientation> EdgeDirections;
+			EdgeDirections.Reserve(2);
+			EdgeDirections.Add(Edges[0].Pin()->GetEndVertex() == Vertex ? EOrientation::Front : EOrientation::Back);
+			EdgeDirections.Add(Edges[1].Pin()->GetStartVertex() == Vertex ? EOrientation::Front : EOrientation::Back);
+			TSharedPtr<FTopologicalEdge> ExtendedEdge = nullptr;
+
+			// TODO Check better if the both curve are similar (2DCurve as to be similar)
+			if (Edges[0].Pin()->GetCurve() != Edges[1].Pin()->GetCurve())
+			{
+				//ExtendedEdge = FEdge::ReplaceEdgesByOne(Edges, EdgeDirections, GetTolerance());
+			}
+			else
+			{
+				//ExtendedEdge = FEdge::CreateEdgeToMerge2Edges(Edges, EdgeDirections);
+			}
+
+			if (ExtendedEdge)
+			{
+				NewEdgeVertices.Add(ExtendedEdge->GetStartVertex());
+				NewEdgeVertices.Add(ExtendedEdge->GetEndVertex());
+			}
+			continue;
+		}
+	}
+}
+*/
+
+
