CADKernel Source Code v0.02

diff --git a/Engine/Source/Runtime/Datasmith/CADKernel/Private/CADKernel/Topo/Joiner.h b/Engine/Source/Runtime/Datasmith/CADKernel/Private/CADKernel/Topo/Joiner.h
--- a/Engine/Source/Runtime/Datasmith/CADKernel/Private/CADKernel/Topo/Joiner.h
+++ b/Engine/Source/Runtime/Datasmith/CADKernel/Private/CADKernel/Topo/Joiner.h
@@ -5,6 +5,9 @@
 
 namespace CADKernel
 {
+	class FSession;
+	class FShell;
+	class FTopologicalEdge;
 	class FTopologicalFace;
 	class FTopologicalVertex;
 
@@ -15,31 +18,95 @@ namespace CADKernel
 
 	protected:
 
-		TArray<TSharedPtr<FTopologicalFace>>& Surfaces;
-		TArray<TSharedPtr<FTopologicalVertex>> Vertices;
+		TSharedRef<FSession> Session;
+
+		TArray<TSharedPtr<FShell>> Shells;
+		TArray<TSharedPtr<FTopologicalFace>> Faces;
 
 		double JoiningTolerance;
-		double JoiningToleranceSqr;
+		double JoiningToleranceSquare;
 
 	public:
 
-		FJoiner(TArray<TSharedPtr<FTopologicalFace>>& surfaces, double Tolerance);
+		FJoiner(TSharedRef<FSession> InSession, const TArray<TSharedPtr<FShell>>& Shells, double Tolerance);
+		FJoiner(TSharedRef<FSession> InSession, const TArray<TSharedPtr<FTopologicalFace>>& Surfaces, double Tolerance);
+
+		void JoinFaces();
+		//void JoinFaces(bool bProcessOnlyBorderEdges, bool bProcessOnlyNonManifoldEdges);
+
+		/**
+		 * Check topology of each body
+		 */
+		void CheckTopology();
 
-		void JoinSurfaces();
-		void JoinSurfaces(bool bProcessOnlyBorderEdges, bool bProcessOnlyNonManifoldEdges);
+		//void MergeInto(TSharedPtr<FBody> Body, TArray<TSharedPtr<FTopologicalEntity>>& InEntities);
+		//void SortByShell(TSharedPtr<FBody> Body, TArray<TSharedPtr<FBody>>& OutNewBody);
+		//void Join(TArray<TSharedPtr<FBody>> Bodies, double Tolerance);
 
-		void FixCollapsedEdges();
+		void SplitIntoConnectedShell();
+		void RemoveFacesFromShell();
 
 	private:
 
-		void GetVertices();
+		void EmptyShells();
+
+		/**
+		 * Return an array of active vertices.
+		 */
+		void GetVertices(TArray<TSharedPtr<FTopologicalVertex>>& Vertices);
+
+		/**
+		 * Return an array of active border vertices.
+		 */
+		void GetBorderVertices(TArray<TSharedPtr<FTopologicalVertex>>& BorderVertices);
+
+		/**
+		 * Merge Border Vertices with other vertices.
+		 * @param Vertices: the initial array of active vertices, this array is updated at the end of the process
+		 */
+		void MergeCoincidentVertices(TArray<TSharedPtr<FTopologicalVertex>>& VerticesToMerge);
 
-		void JoinVertices();
-		void JoinBorderVertices(TArray<TSharedPtr<FTopologicalVertex>>& MergedVertices);
+		/**
+		 * Merge Border Vertices with other vertices.
+		 * @param Vertices: the initial array of active vertices, this array is updated at the end of the process
+		 */
+		void MergeBorderVerticesWithCoincidentOtherVertices(TArray<TSharedPtr<FTopologicalVertex>>& Vertices);
+
+		//void FixCollapsedEdges();
 		void CheckSelfConnectedEdge();
+		void RemoveIsolatedEdges(); // Usefull ???
+
+		/**
+		 * First step, trivial edge merge i.e. couple of edges with same extremity vertex
+		 */
+		void MergeCoincidentEdges(TArray<TSharedPtr<FTopologicalVertex>>& Vertices);
+
+		/**
+		 * Second step, parallel edges but with different length. the longest must be split
+		 */
+		void StitchParallelEdges(TArray<TSharedPtr<FTopologicalVertex>>& Vertices);
+
+		/**
+		 * Split the edge to split except if one side is too small (JoiningTolerance). In this case (too small), the EdgeToSplit is not split but linked to EdgeToLink. 
+		 * @return the created vertex or TSharedPtr<FTopologicalVertex>()
+		 */
+		TSharedPtr<FTopologicalVertex> SplitAndLink(TSharedRef<FTopologicalVertex>& StartVertex, TSharedPtr<FTopologicalEdge>& EdgeToLink, TSharedPtr<FTopologicalEdge>& EdgeToSplit);
+
+		/**
+		 * For each loop of each surface, check if successive edges are unconnected and if their common vertices are connected only to them. 
+		 * These edges are merged into one edge. 
+		 * E.g. :
+		 * Face A has 3 successive unconnected edges. If these 3 edges are merged to give only one edge, the new edge could be linked to its parallel edge of Face B 
+		 * 
+		 *              \                          Face A                                   |    Face C
+		 *               \                                                                  |
+		 *    Face E     CV ------------------- UV ------------------ UV ----------------- CV --------------------
+		 *               CV -------------------------------------------------------------- CV --------------------
+		 *               /                         Face B                                   |    Face D
+		 *              /                                                                   |
+		 */
+		void MergeUnconnectedAdjacentEdges();
 
-		void JoinEdges(TArray<TSharedPtr<FTopologicalVertex>>& Vertices);
-		void MergeUnconnectedAdjacentEdges(TArray<TSharedPtr<FTopologicalVertex>>& Vertices, bool bOnlyIfSameCurve);
 	};
 
 } // namespace CADKernel
diff --git a/Engine/Source/Runtime/Datasmith/CADKernel/Private/CADKernel/Topo/Model.cpp b/Engine/Source/Runtime/Datasmith/CADKernel/Private/CADKernel/Topo/Model.cpp
--- a/Engine/Source/Runtime/Datasmith/CADKernel/Private/CADKernel/Topo/Model.cpp
+++ b/Engine/Source/Runtime/Datasmith/CADKernel/Private/CADKernel/Topo/Model.cpp
@@ -50,14 +50,51 @@ void FModel::RemoveEntity(TSharedPtr<FTopologicalEntity> Entity)
 		RemoveBody(StaticCastSharedPtr<FBody>(Entity));
 		break;
 	case EEntity::TopologicalFace:
-		RemoveDomain(StaticCastSharedPtr<FTopologicalFace>(Entity));
+		RemoveFace(StaticCastSharedPtr<FTopologicalFace>(Entity));
 		break;
 	default:
 		break;
 	}
 }
 
-void FModel::RemoveDomain(TSharedPtr<FTopologicalFace> Domain)
+void FModel::PrintBodyAndShellCount()
+{
+	int32 NbBody = 0;
+	int32 NbShell = 0;
+
+	TArray<TSharedPtr<FBody>> NewBodies;
+	for (TSharedPtr<FBody> Body : Bodies)
+	{
+		NbShell += Body->GetShells().Num();
+		NbBody++;
+	}
+	FMessage::Printf(Log, TEXT("Body count %d shell count %d \n"), NbBody, NbShell);
+}
+
+void FModel::RemoveEmptyBodies()
+{
+	int32 NbBody = 0;
+	int32 NbShell = 0;
+
+	TArray<TSharedPtr<FBody>> NewBodies;
+	NewBodies.Reserve(Bodies.Num());
+	for (TSharedPtr<FBody> Body : Bodies)
+	{
+		Body->RemoveEmptyShell();
+		if (Body->GetShells().Num())
+		{
+			NbShell += Body->GetShells().Num();
+			NewBodies.Emplace(Body);
+			NbBody++;
+		}
+	}
+	Swap(NewBodies, Bodies);
+	FMessage::Printf(Log, TEXT("After RemoveEmptyBodies, Body count %d shell count %d \n"), NbBody, NbShell);
+}
+
+
+
+/*void FModel::RemoveDomain(TSharedPtr<FTopologicalFace> Domain)
 {
 	Faces.Remove(Domain);
 }
@@ -65,7 +102,7 @@ void FModel::RemoveDomain(TSharedPtr<FTopologicalFace> Domain)
 void FModel::RemoveBody(TSharedPtr<FBody> Body)
 {
 	Bodies.Remove(Body);
-}
+}*/
 
 TSharedPtr<FEntityGeom> FModel::ApplyMatrix(const FMatrixH& InMatrix) const
 {
@@ -127,19 +164,158 @@ FInfoEntity& FModel::GetInfo(FInfoEntity& Info) const
 {
 	return FEntityGeom::GetInfo(Info)
 		.Add(TEXT("Bodies"), Bodies)
-		.Add(TEXT("Domains"), Faces);
+		.Add(TEXT("Domains"), Faces)
+		.Add(*this);
 }
 #endif
 
-void FModel::Add(TSharedPtr<FTopologicalFace> Face)
+
+// Topo functions
+void FModel::MergeInto(TSharedPtr<FBody> Body, TArray<TSharedPtr<FTopologicalEntity>>& InEntities)
 {
-	Faces.Add(Face);
+
 }
 
-void FModel::Add(TSharedPtr<FBody> Body)
+
+struct FBodyShell
+{
+	TSharedPtr<FBody> Body;
+	TSharedPtr<FShell> Shell;
+
+	FBodyShell(TSharedPtr<FBody> InBody, TSharedPtr<FShell> InShell)
+	: Body(InBody)
+	, Shell(InShell)
+	{
+	}
+
+};
+
+//void FModel::CheckTopology()
+void FModel::CheckTopology() 
 {
-	Bodies.Add(Body);
+	TArray<FBodyShell> IsolatedBodies;
+	IsolatedBodies.Reserve(Bodies.Num()*2);
+
+	int32 ShellCount = 0;
+
+	for (TSharedPtr<FBody> Body : Bodies)
+	{
+		for (TSharedPtr<FShell> Shell : Body->GetShells())
+		{
+			ShellCount++;
+			TArray<FFaceSubset> SubShells;
+			Shell->CheckTopology(SubShells);
+
+			if (SubShells.Num() == 1 )
+			{
+				if (Shell->FaceCount() < 3 )
+				{
+					IsolatedBodies.Emplace(Body, Shell);
+				}
+				else
+				{
+					if (SubShells[0].BorderEdgeCount > 0 || SubShells[0].NonManifoldEdgeCount > 0)
+					{
+#ifdef CORETECHBRIDGE_DEBUG
+						FMessage::Printf(Log, TEXT("Body %d shell %d CADId %d is opened and has %d faces "), Body->GetKioId(), Shell->GetKioId(), Shell->GetId(), Shell->FaceCount());
+#else
+						FMessage::Printf(Log, TEXT("Body %d shell %d is opened and has %d faces "), Body->GetId(), Shell->GetId(), Shell->FaceCount());
+#endif
+						FMessage::Printf(Log, TEXT("and has %d border edges and %d nonManifold edges\n"), SubShells[0].BorderEdgeCount, SubShells[0].NonManifoldEdgeCount);
+					}
+					else
+					{
+#ifdef CORETECHBRIDGE_DEBUG
+						FMessage::Printf(Log, TEXT("Body %d shell %d CADId %d is closed and has %d faces\n"), Body->GetKioId(), Shell->GetKioId(), Shell->GetId(), Shell->FaceCount());
+#else
+						FMessage::Printf(Log, TEXT("Body %d shell %d is closed and has %d faces\n"), Body->GetId(), Shell->GetId(), Shell->FaceCount());
+#endif
+					}
+				}
+			}
+			else
+			{
+#ifdef CORETECHBRIDGE_DEBUG
+				FMessage::Printf(Log, TEXT("Body %d shell %d CADId %d has %d subshells\n"), Body->GetKioId(), Shell->GetKioId(), Shell->GetId(), SubShells.Num());
+#else
+				FMessage::Printf(Log, TEXT("Body %d shell %d has %d subshells\n"), Body->GetId(), Shell->GetId(), SubShells.Num());
+#endif
+for (const FFaceSubset& Subset : SubShells)
+				{
+					FMessage::Printf(Log, TEXT("     - Subshell of %d faces %d border edges and %d nonManifold edges\n"), Subset. Faces.Num(), Subset.BorderEdgeCount, Subset.NonManifoldEdgeCount);
+				}
+			}
+		}
+	}
+}
+
+
+void FModel::Split(TSharedPtr<FBody> Body, TArray<TSharedPtr<FBody>>& OutNewBody)
+{
+
+}
+
+/**
+ * Fore each shell of each body, try to stitch topological gap
+ */
+void FModel::HealModelTopology(double JoiningTolerance)
+{
+
 }
 
 
 
+//void FModel::FixModelTopology()
+void FModel::FixModelTopology(double JoiningTolerance)
+{
+	//HealModelTopology(JoiningTolerance);
+
+	//// Find Isolated Shell
+	//TArray<TSharedPtr<FShell>> IsolatedShell;
+
+	//for (TSharedPtr<FBody> Body : Bodies)
+	//{
+	//	for (TSharedPtr<FShell> Shell : Body->GetShells())
+	//	{
+	//		if (Shell->IsOpenShell())
+	//		{
+	//			IsolatedShell.Add(Shell);
+	//		}
+	//	}
+	//}
+
+	//FJoiner Joiner(IsolatedShell, JoiningTolerance);
+	//Joiner.JoinFaces();
+
+	/*
+	Body->GetFaces(Surfaces);
+	TSharedRef<FBody> MergedBody = FEntity::MakeShared<FBody>();
+	TSharedRef<FShell> MergedShell = FEntity::MakeShared<FShell>();
+
+	MergedBody->AddShell(MergedShell);
+	TArray<TSharedPtr<FBody>> NewBodies;
+	NewBodies.Reserve(Bodies.Num());
+
+	for (TSharedPtr<FBody>& Body : Bodies)
+	{
+		if (Body->GetShells().Num() == 1)
+		{
+			TSharedPtr<FShell> Shell = Body->GetShells()[0];
+			if (Shell->FaceCount() < 3)
+			{
+				MergedShell->Merge(Shell);
+				Body->Empty();
+			}
+		}
+		else
+		{
+			NewBodies.Emplace(Body);
+			Body.Reset();
+		}
+	}
+
+	NewBodies.Emplace(MergedBody);
+
+	Swap(NewBodies, Bodies);
+	*/
+}
diff --git a/Engine/Source/Runtime/Datasmith/CADKernel/Private/CADKernel/Topo/Shell.cpp b/Engine/Source/Runtime/Datasmith/CADKernel/Private/CADKernel/Topo/Shell.cpp
--- a/Engine/Source/Runtime/Datasmith/CADKernel/Private/CADKernel/Topo/Shell.cpp
+++ b/Engine/Source/Runtime/Datasmith/CADKernel/Private/CADKernel/Topo/Shell.cpp
@@ -51,17 +51,44 @@ TSharedPtr<FEntityGeom> FShell::ApplyMatrix(const FMatrixH& InMatrix) const
 	return TSharedPtr<FEntityGeom>();
 }
 
+void FShell::Empty(int32 NewSize)
+{
+	for(FOrientedFace& Face : TopologicalFaces)
+	{
+		Face.Entity->ResetHost();
+	}
+	TopologicalFaces.Empty(NewSize);
+}
+
+void FShell::Add(TArray<TSharedPtr<FTopologicalFace>> Faces)
+{
+	TSharedPtr<FShell> Shell = StaticCastSharedRef<FShell>(AsShared());
+
+	TopologicalFaces.Reserve(TopologicalFaces.Num() + Faces.Num());
+
+	for(TSharedPtr<FTopologicalFace>& Face : Faces)
+	{
+		TopologicalFaces.Emplace(Face, Face->IsBackOriented() ? EOrientation::Back : EOrientation::Front);
+		Face->SetHost(Shell);
+	}
+}
+
+
 void FShell::Add(TSharedRef<FTopologicalFace> InTopologicalFace, EOrientation Orientation)
 {
 	TSharedPtr<FTopologicalFace> Face = InTopologicalFace;
 	TopologicalFaces.Emplace(Face, Orientation);
+
+	Face->SetHost(StaticCastSharedRef<FShell>(AsShared()));
 }
 
 #ifdef CADKERNEL_DEV
 FInfoEntity& FShell::GetInfo(FInfoEntity& Info) const
 {
 	return FEntity::GetInfo(Info)
-		.Add(TEXT("TopologicalFaces"), (TArray<TOrientedEntity<FEntity>>&) TopologicalFaces);
+		.Add(TEXT("Hosted by"), (TWeakPtr<FEntity>&) HostedBy)
+		.Add(TEXT("TopologicalFaces"), (TArray<TOrientedEntity<FEntity>>&) TopologicalFaces)
+		.Add(*this);
 }
 #endif
 
@@ -79,6 +106,13 @@ void FShell::GetFaces(TArray<TSharedPtr<FTopologicalFace>>& Faces)
 	}
 }
 
+void FShell::Merge(TSharedPtr<FShell>& Shell)
+{
+	TopologicalFaces.Append(Shell->TopologicalFaces);
+	Shell->TopologicalFaces.Empty();
+}
+
+
 void FShell::SpreadBodyOrientation()
 {
 	bool bIsOutter = IsOutter();
@@ -98,6 +132,27 @@ void FShell::SpreadBodyOrientation()
 	}
 }
 
+bool FShell::IsOpenShell()
+{
+	for (const FOrientedFace& OrientedFace : GetFaces())
+	{
+		const TSharedPtr<FTopologicalFace>& Face = OrientedFace.Entity;
+		for (const TSharedPtr<FTopologicalLoop>& Loop : Face->GetLoops())
+		{
+			for (const FOrientedEdge& OrientedEdge : Loop->GetEdges())
+			{
+				const TSharedPtr<FTopologicalEdge>& Edge = OrientedEdge.Entity;
+				if (Edge->GetTwinsEntityCount() == 1)
+				{
+					return true;
+				}
+			}
+		}
+	}
+	return false;
+}
+
+
 void FShell::CheckTopology(TArray<FFaceSubset>& SubShells)
 {
 	// Processed1 : Surfaces added in CandidateSurfacesForMesh
diff --git a/Engine/Source/Runtime/Datasmith/CADKernel/Private/CADKernel/Topo/TopologicalEdge.cpp b/Engine/Source/Runtime/Datasmith/CADKernel/Private/CADKernel/Topo/TopologicalEdge.cpp
--- a/Engine/Source/Runtime/Datasmith/CADKernel/Private/CADKernel/Topo/TopologicalEdge.cpp
+++ b/Engine/Source/Runtime/Datasmith/CADKernel/Private/CADKernel/Topo/TopologicalEdge.cpp
@@ -3,8 +3,12 @@
 #include "CADKernel/Topo/TopologicalEdge.h"
 
 #include "CADKernel/Core/KernelParameters.h"
+#include "CADKernel/Geo/Curves/NURBSCurve.h"
+#include "CADKernel/Geo/Curves/SegmentCurve.h"
+#include "CADKernel/Geo/GeoEnum.h"
 #include "CADKernel/Geo/Sampling/PolylineTools.h"
 #include "CADKernel/Geo/Sampler/SamplerOnParam.h"
+#include "CADKernel/Math/BSpline.h"
 #include "CADKernel/Math/Point.h"
 #include "CADKernel/Math/SlopeUtils.h"
 #include "CADKernel/Mesh/Structure/EdgeMesh.h"
@@ -57,6 +61,18 @@ FTopologicalEdge::FTopologicalEdge(const TSharedRef<FRestrictionCurve>& InCurve,
 	EndVertex = FEntity::MakeShared<FTopologicalVertex>(Points[1].Point);
 }
 
+FTopologicalEdge::FTopologicalEdge(const TSharedRef<FSurface>& InSurface, const FPoint2D& InCoordinateVertex1, const TSharedRef<FTopologicalVertex>& InVertex1, const FPoint2D& InCoordinateVertex2, const TSharedRef<FTopologicalVertex>& InVertex2)
+	: TLinkable<FTopologicalEdge, FEdgeLink>()
+	, StartVertex(InVertex1)
+	, EndVertex(InVertex2)
+	, Length3D(-1)
+	, Loop(TSharedPtr<FTopologicalLoop>())
+	, Mesh(TSharedPtr<FEdgeMesh>())
+{
+	TSharedRef<FSegmentCurve> Curve2D = FEntity::MakeShared<FSegmentCurve>(InCoordinateVertex1, InCoordinateVertex2, 2);
+	Curve = FEntity::MakeShared<FRestrictionCurve>(InSurface, Curve2D);
+}
+
 FTopologicalEdge::FTopologicalEdge(const TSharedRef<FRestrictionCurve>& InCurve)
 	: FTopologicalEdge(InCurve, InCurve->GetBoundary())
 {
@@ -80,7 +96,6 @@ bool FTopologicalEdge::CheckVertices()
 	Curve->Approximate3DPoints(Coordinates, Points);
 
 	double ToleranceGeo = GetTolerance3D();
-
 	TFunction<bool(TSharedPtr<FTopologicalVertex>, FPoint)> CheckExtremityGap = [&](TSharedPtr<FTopologicalVertex> Vertex, FPoint Point)
 	{
 		double GapToVertex = Vertex->GetCoordinates().Distance(Point);
@@ -101,7 +116,7 @@ bool FTopologicalEdge::CheckVertices()
 
 double FTopologicalEdge::GetTolerance3D()
 {
-	return GetLoop()->GetFace()->GetCarrierSurface()->Get3DTolerance();
+	return GetCurve()->GetCarrierSurface()->Get3DTolerance();
 }
 
 bool FTopologicalEdge::CheckIfDegenerated()
@@ -177,45 +192,92 @@ TSharedPtr<FTopologicalEdge> FTopologicalEdge::Make(const TSharedRef<FRestrictio
 	return Edge;
 }
 
-void FTopologicalEdge::Link(const TSharedRef<FTopologicalEdge>& Twin)
+TSharedPtr<FTopologicalEdge> FTopologicalEdge::Make(const TSharedRef<FSurface>& InSurface, const FPoint2D& InCoordinateVertex1, const TSharedRef<FTopologicalVertex>& InVertex1, const FPoint2D& InCoordinateVertex2, const TSharedRef<FTopologicalVertex>& InVertex2)
+{
+	TSharedRef<FTopologicalEdge> Edge = FEntity::MakeShared<FTopologicalEdge>(InSurface, InCoordinateVertex1, InVertex1, InCoordinateVertex2, InVertex2);
+	if (Edge->CheckIfDegenerated())
+	{
+		return TSharedPtr<FTopologicalEdge>();
+	}
+
+	Edge->LinkVertex();
+	return Edge;
+}
+
+void FTopologicalEdge::Link(const TSharedRef<FTopologicalEdge>& Twin, double SquareJoiningTolerance)
 {
 	// Degenerated twin edges are not linked
-	if (IsDegenerated() && Twin->IsDegenerated())
+	if (IsDegenerated() || Twin->IsDegenerated())
 	{
+		SetAsDegenerated();
+		Twin->SetAsDegenerated();
 		return;
 	}
 
-	ensureCADKernel(!IsDegenerated() || !Twin->IsDegenerated());
-
-	FPoint Edge1Vertex1 = GetStartVertex()->GetBarycenter();
-	FPoint Edge1Vertex2 = GetEndVertex()->GetBarycenter();
-	FPoint Edge2Vertex1 = Twin->GetStartVertex()->GetBarycenter();
-	FPoint Edge2Vertex2 = Twin->GetEndVertex()->GetBarycenter();
+	const FPoint& Edge1Vertex1 = GetStartVertex()->GetBarycenter();
+	const FPoint& Edge1Vertex2 = GetEndVertex()->GetBarycenter();
+	const FPoint& Edge2Vertex1 = Twin->GetStartVertex()->GetBarycenter();
+	const FPoint& Edge2Vertex2 = Twin->GetEndVertex()->GetBarycenter();
 
 	// Define the orientation
-	double DistanceSameOrientation = Edge1Vertex1.Distance(Edge2Vertex1) + Edge1Vertex2.Distance(Edge2Vertex2);
-	double DistanceReverseOrientation = Edge1Vertex1.Distance(Edge2Vertex2) + Edge1Vertex2.Distance(Edge2Vertex1);
-	if (DistanceSameOrientation < DistanceReverseOrientation)
-	{
-		if (DistanceSameOrientation > 0.1)
+	const double SquareDistanceE1V1_E2V1 = GetStartVertex()->IsLinkedTo(Twin->GetStartVertex()) ? 0. : Edge1Vertex1.SquareDistance(Edge2Vertex1);
+	const double SquareDistanceE1V2_E2V2 = GetEndVertex()->IsLinkedTo(Twin->GetEndVertex()) ? 0. : Edge1Vertex2.SquareDistance(Edge2Vertex2);
+							 
+	const double SquareDistanceE1V1_E2V2 = GetStartVertex()->IsLinkedTo(Twin->GetEndVertex()) ? 0. : Edge1Vertex1.SquareDistance(Edge2Vertex2);
+	const double SquareDistanceE1V2_E2V1 = GetEndVertex()->IsLinkedTo(Twin->GetStartVertex()) ? 0. : Edge1Vertex2.SquareDistance(Edge2Vertex1);
+
+	bool bCanMergeEdge = true;
+	const double SquareDistanceSameOrientation = SquareDistanceE1V1_E2V1 + SquareDistanceE1V2_E2V2;
+	const double SquareDistanceReverseOrientation = SquareDistanceE1V1_E2V2 + SquareDistanceE1V2_E2V1;
+	if (SquareDistanceSameOrientation < SquareDistanceReverseOrientation)
+	{
+		if (SquareDistanceE1V1_E2V1 < SquareJoiningTolerance)
 		{
-			FMessage::Printf(Log, TEXT("Edge %d and Edge %d are to far to be connected"), Id, Twin->GetId());
-			return;
+			GetStartVertex()->Link(Twin->GetStartVertex());
+		}
+		else
+		{
+			FMessage::Printf(Log, TEXT("Edge %d and Edge %d are to far (%f) to be connected\n"), GetId(), Twin->GetId(), sqrt(SquareDistanceE1V1_E2V1));
+			bCanMergeEdge = false;
+		}
+
+		if (SquareDistanceE1V2_E2V2 < SquareJoiningTolerance)
+		{
+			GetEndVertex()->Link(Twin->GetEndVertex());
+		}
+		else
+		{
+			FMessage::Printf(Log, TEXT("Edge %d and Edge %d are to far (%f) to be connected\n"), GetId(), Twin->GetId(), sqrt(SquareDistanceE1V2_E2V2));
+			bCanMergeEdge = false;
 		}
-		GetStartVertex()->Link(Twin->GetStartVertex());
-		GetEndVertex()->Link(Twin->GetEndVertex());
 	}
 	else
 	{
-		if (DistanceReverseOrientation > 0.1)
+		if (SquareDistanceE1V1_E2V2 < SquareJoiningTolerance)
 		{
-			FMessage::Printf(Log, TEXT("Edge %d and Edge %d are to far to be connected"), Id, Twin->GetId());
-			return;
+			GetStartVertex()->Link(Twin->GetEndVertex());
 		}
-		GetStartVertex()->Link(Twin->GetEndVertex());
-		GetEndVertex()->Link(Twin->GetStartVertex());
+		else
+		{
+			FMessage::Printf(Log, TEXT("Edge %d and Edge %d are to far (%f) to be connected\n"), GetId(), Twin->GetId(), sqrt(SquareDistanceE1V1_E2V2));
+			bCanMergeEdge = false;
+		}
+
+		if (SquareDistanceE1V2_E2V1 < SquareJoiningTolerance)
+		{
+			GetEndVertex()->Link(Twin->GetStartVertex());
+		}
+		else
+		{
+			FMessage::Printf(Log, TEXT("Edge %d and Edge %d are to far (%f) to be connected\n"), GetId(), Twin->GetId(), sqrt(SquareDistanceE1V2_E2V1));
+			bCanMergeEdge = false;
+		}
+	}
+
+	if (bCanMergeEdge)
+	{
+		MakeLink(Twin);
 	}
-	MakeLink(Twin);
 }
 
 void FTopologicalEdge::Delete()
@@ -289,7 +351,7 @@ double FTopologicalEdge::Length() const
 {
 	if (Length3D < 0)
 	{
-		Length3D = Curve->ComputeLength(Boundary);
+		Length3D = Curve->ApproximateLength(Boundary);
 	}
 	return Length3D;
 }
@@ -509,32 +571,250 @@ void FTopologicalEdge::ChooseFinalDeltaUs()
 	}
 }
 
-TSharedPtr<FTopologicalEdge> FTopologicalEdge::CreateEdgeToMerge2Edges(const TArray<TWeakPtr<FTopologicalEdge>>& Edges, const TArray<EOrientation>& EdgeDirections)
+TSharedPtr<FTopologicalEdge> FTopologicalEdge::CreateEdgeByMergingEdges(TArray<FOrientedEdge>& Edges, const TSharedRef<FTopologicalVertex> StartVertex, const TSharedRef<FTopologicalVertex> EndVertex)
 {
-	ensureCADKernel(false);
-	TSharedPtr<FTopologicalEdge> Edge0 = Edges[0].Pin();
-	TSharedPtr<FTopologicalEdge> Edge1 = Edges[1].Pin();
+	// Make merged 2d Nurbs ===================================================
+
+	const double Tolerance3D = Edges[0].Entity->GetTolerance3D();
+	//const double Tolerance2D = Edges[0].Entity->GetCurve()->Get2DCurve()->GetDomainTolerance();
+	TSharedRef<FSurface> CarrierSurface = Edges[0].Entity->GetCurve()->GetCarrierSurface();
+
+	// check if all curve a 2D NURBS
+	bool bAreNurbs = true;
+	int32 NurbsMaxDegree = 0;
+
+	// TODO, Check if some edges are small edges and could be replace by extending next or previous one
+	double NewEdgeLength = 0;
+	for (const FOrientedEdge& Edge : Edges)
+	{
+		NewEdgeLength += Edge.Entity->Length();
+	}
+
+	// MinLength = smaller than this size, the adgacent edge is extend to replace it  
+	const double MinLength = FMath::Min(NewEdgeLength / (Edges.Num()+3.), FMath::Max(NewEdgeLength / 20., Tolerance3D * 5.)); // Tolerance3D = 0.01mm => MinLength = 0.25mm, 1/20 of the length is small enougth to be a detail.
+
+	TArray<TSharedPtr<FNURBSCurve>> NurbsCurves;
+	NurbsCurves.Reserve(Edges.Num());
+
+	bool bCanRemove = true;
+	for (FOrientedEdge& Edge : Edges)
+	{
+		if (Edge.Entity->GetCurve()->Get2DCurve()->GetCurveType() != ECurve::Nurbs)
+		{
+			return TSharedPtr<FTopologicalEdge>();
+		}
+
+		double EdgeLength = Edge.Entity->Length();
+		if (bCanRemove && EdgeLength < MinLength)
+		{
+			NurbsCurves.Emplace(TSharedPtr<FNURBSCurve>());
+			bCanRemove = false;
+			continue; // the edge will be ignored 
+		}
+		bCanRemove = true;
+
+
+		// Find the max degree of the nurbs
+		TSharedPtr<FNURBSCurve> NURBS = NurbsCurves.Emplace_GetRef(StaticCastSharedRef<FNURBSCurve>(Edge.Entity->GetCurve()->Get2DCurve()));
+		int32 NurbsDegree = NURBS->GetDegree();
+		if (NurbsDegree > NurbsMaxDegree)
+		{
+			NurbsMaxDegree = NurbsDegree;
+		}
+
+		// Edge has restricted its curve ?
+		FLinearBoundary EdgeBoundary = Edge.Entity->GetBoundary();
+		FLinearBoundary CurveBoundary = NURBS->GetBoundary();
+
+		double ParametricTolerance = NURBS->GetBoundary().ComputeMinimalTolerance();
+
+		if (!FMath::IsNearlyEqual(EdgeBoundary.Min, CurveBoundary.Min, ParametricTolerance) ||
+			!FMath::IsNearlyEqual(EdgeBoundary.Max, CurveBoundary.Max, ParametricTolerance))
+		{
+			// ToDO, check if the next edge is not the complementary of this
+
+			// cancel
+			return TSharedPtr<FTopologicalEdge>();
+		}
+	}
+
+	bool bEdgeNeedToBeExtend = false;
+	int32 PoleCount = 0;
+	double LastCoordinate = 0;
+
+	for (int32 Index = 0; Index < Edges.Num(); Index++)
+	{
+		TSharedPtr<FNURBSCurve>& NURBS = NurbsCurves[Index];
+		if (!NURBS.IsValid())
+		{
+			bEdgeNeedToBeExtend = true;
+			continue; // the edge will be ignored 
+		}
+
+		if (NURBS->GetDegree() < NurbsMaxDegree)
+		{
+			NURBS = BSpline::DuplicateNurbsCurveWithHigherDegree(NurbsMaxDegree, *NURBS);
+		}
+		else
+		{
+			NURBS = FEntity::MakeShared<FNURBSCurve>(NURBS.ToSharedRef());
+		}
+
+		if (Edges[Index].Direction == EOrientation::Back)
+		{
+			NURBS->Invert();
+		}
+
+		NURBS->SetStartNodalCoordinate(LastCoordinate);
+		LastCoordinate = NURBS->GetBoundary().GetMax();
+
+		PoleCount += NURBS->GetPoleCount();
+	}
+
+	if (bEdgeNeedToBeExtend)
+	{
+		for (int32 Index = 0; Index < Edges.Num(); Index++)
+		{
+			if (!NurbsCurves[Index].IsValid())
+			{
+				double PreviousLength = Index > 0 ? Edges[Index - 1].Entity->Length() : 0;
+				double NextLength = Index < Edges.Num() - 1 ? Edges[Index + 1].Entity->Length() : 0;
+
+				double TargetCoordinate = 0;
+				EOrientation FrontOrientation = PreviousLength > NextLength ? EOrientation::Front : EOrientation::Back;
+				TargetCoordinate = Edges[Index].Direction == FrontOrientation ? Edges[Index].Entity->GetBoundary().GetMax() : Edges[Index].Entity->GetBoundary().GetMin();
+
+				if (PreviousLength > NextLength)
+				{
+					TargetCoordinate = Edges[Index].Direction == EOrientation::Front ? Edges[Index].Entity->GetBoundary().GetMax() : Edges[Index].Entity->GetBoundary().GetMin();
+				}
+				else
+				{
+					TargetCoordinate = Edges[Index].Direction == EOrientation::Front ? Edges[Index].Entity->GetBoundary().GetMin() : Edges[Index].Entity->GetBoundary().GetMax();
+				}
+				FPoint2D Target = Edges[Index].Entity->Approximate2DPoint(TargetCoordinate);
+
+				int32 NeigborIndex = PreviousLength > NextLength ? Index - 1 : Index + 1;
+				NurbsCurves[NeigborIndex]->ExtendTo(Target);
+			}
+		}
+	}
+
+	TArray<double> NewNodalVector;
+	TArray<double> NewWeights;
+	TArray<FPoint> NewPoles;
+	NurbsMaxDegree++;
+	NewNodalVector.Reserve(PoleCount + NurbsMaxDegree);
+	NewWeights.Reserve(PoleCount + NurbsMaxDegree);
+	NewPoles.Reserve(PoleCount + NurbsMaxDegree);
+
+	for (const TSharedPtr<FNURBSCurve>& NurbsCurve : NurbsCurves)
+	{
+		if (!NurbsCurve.IsValid())
+		{
+			continue;
+		}
+
+		if (!NewPoles.IsEmpty())
+		{
+			NewPoles.Pop();
+			NewWeights.Pop();
+		}
 
-	TSharedPtr<FTopologicalVertex> Vertex1 = EdgeDirections[0] == EOrientation::Front ? Edge0->GetStartVertex() : Edge0->GetEndVertex();
-	double Coordinate1 = EdgeDirections[0] == EOrientation::Front ? Edge0->GetStartCurvilinearCoordinates() : Edge0->GetEndCurvilinearCoordinates();
-	TSharedPtr<FTopologicalVertex> Vertex2 = EdgeDirections[1] == EOrientation::Front ? Edge1->GetEndVertex() : Edge1->GetStartVertex();
-	double Coordinate2 = EdgeDirections[1] == EOrientation::Front ? Edge1->GetEndCurvilinearCoordinates() : Edge1->GetStartCurvilinearCoordinates();
-	TSharedPtr<FTopologicalVertex> DeletedVertex = EdgeDirections[1] == EOrientation::Front ? Edge1->GetStartVertex() : Edge1->GetEndVertex();
-	if (Coordinate1 > Coordinate2)
+		NewPoles.Append(NurbsCurve->GetPoles());
+		NewWeights.Append(NurbsCurve->GetWeights());
+	}
+
+	for (const TSharedPtr<FNURBSCurve>& NurbsCurve : NurbsCurves)
+	{
+		if (!NurbsCurve.IsValid())
+		{
+			continue;
+		}
+
+		if (NewNodalVector.IsEmpty())
+		{
+			NewNodalVector.Append(NurbsCurve->GetNodalVector());
+		}
+		else
+		{
+			NewNodalVector.SetNum(NewNodalVector.Num() - 1);
+			NewNodalVector.Append(NurbsCurve->GetNodalVector().GetData() + NurbsMaxDegree, NurbsCurve->GetNodalVector().Num() - NurbsMaxDegree);
+		}
+	}
+
+	TSharedRef<FNURBSCurve> MergedNURBS = FEntity::MakeShared<FNURBSCurve>(NurbsMaxDegree - 1, NewNodalVector, NewPoles, NewWeights, 2);
+
+	// Make new edge and delete the old ones ===================================================
+
+	TSharedRef<FRestrictionCurve> RestrictionCurve = FEntity::MakeShared<FRestrictionCurve>(CarrierSurface, MergedNURBS);
+
+	TSharedPtr<FTopologicalEdge> NewEdge = Make(RestrictionCurve, StartVertex, EndVertex);
+	if (!NewEdge.IsValid())
 	{
-		Swap(Vertex1, Vertex2);
-		Swap(Coordinate1, Coordinate2);
+		printf("SSS");
 	}
 
-	Edge0->SetStartVertex(Coordinate1, Vertex1->GetCoordinates());
-	Edge0->SetEndVertex(Coordinate2, Vertex2->GetCoordinates());
 
-	DeletedVertex->RemoveConnectedEdge(Edge0.ToSharedRef());
-	DeletedVertex->RemoveConnectedEdge(Edge1.ToSharedRef());
+	TSharedPtr<FTopologicalLoop> Loop = Edges[0].Entity->GetLoop();
+	Loop->ReplaceEdges(Edges, NewEdge);
 
-	Edge0->GetLoop()->RemoveEdge(Edge1);
+	for (const FOrientedEdge& OrientedEdge : Edges)
+	{
+		OrientedEdge.Entity->Delete();
+	}
 
-	return Edge0;
+	return NewEdge;
+}
+
+bool FTopologicalEdge::ExtendTo(bool bStartExtremity, const FPoint2D& NewExtremityCoordinate, TSharedRef<FTopologicalVertex> NewVertex)
+{
+	TFunction<void(TSharedPtr<FTopologicalVertex>&)> UpdateVertex = [&](TSharedPtr<FTopologicalVertex>& EdgeVertex)
+	{
+		TSharedRef<FTopologicalEdge> Edge = StaticCastSharedRef<FTopologicalEdge>(AsShared());
+		EdgeVertex->RemoveConnectedEdge(Edge);
+		if (EdgeVertex->GetDirectConnectedEdges().Num() == 0)
+		{
+			EdgeVertex->RemoveFromLink();
+		}
+		EdgeVertex = NewVertex;
+		NewVertex->AddConnectedEdge(Edge);
+		Length3D = -1.;
+	};
+
+	if (bStartExtremity ? FMath::IsNearlyEqual(Boundary.Min, Curve->GetBoundary().Min) : FMath::IsNearlyEqual(Boundary.Max, Curve->GetBoundary().Max))
+	{
+		Curve->ExtendTo(NewExtremityCoordinate);
+	}
+	else
+	{
+		FPoint ProjectedPoint;
+		double UProjectedPoint = ProjectPoint(NewVertex->GetCoordinates(), ProjectedPoint);
+		if (ProjectedPoint.Distance(NewVertex->GetCoordinates()) > GetTolerance3D())
+		{
+			return false;
+		}
+
+		if (bStartExtremity)
+		{
+			Boundary.Min = UProjectedPoint;
+		}
+		else
+		{
+			Boundary.Max = UProjectedPoint;
+		}
+	}
+
+	if (bStartExtremity)
+	{
+		UpdateVertex(StartVertex);
+	}
+	else
+	{
+		UpdateVertex(EndVertex);
+	}
+
+	return true;
 }
 
 void FTopologicalEdge::ComputeEdge2DProperties(FEdge2DProperties& EdgeCharacteristics)
@@ -558,6 +838,45 @@ void FTopologicalEdge::ComputeEdge2DProperties(FEdge2DProperties& EdgeCharacteri
 	}
 }
 
+FPoint FTopologicalEdge::GetTangentAt(const TSharedRef<FTopologicalVertex>& InVertex)
+{
+	if (InVertex->GetLink() == StartVertex->GetLink())
+	{
+		return Curve->GetTangentAt(Boundary.GetMin());
+	}
+	else if (InVertex->GetLink() == EndVertex->GetLink())
+	{
+		FPoint Tangent = Curve->GetTangentAt(Boundary.GetMax());
+		Tangent *= -1.;
+		return Tangent;
+	}
+	else
+	{
+		ensureCADKernel(false);
+		return FPoint::ZeroPoint;
+	}
+}
+
+FPoint2D FTopologicalEdge::GetTangent2DAt(const TSharedRef<FTopologicalVertex>& InVertex)
+{
+	if (InVertex->GetLink() == StartVertex->GetLink())
+	{
+		return Curve->GetTangent2DAt(Boundary.GetMin());
+	}
+	else if (InVertex->GetLink() == EndVertex->GetLink())
+	{
+		FPoint2D Tangent = Curve->GetTangent2DAt(Boundary.GetMax());
+		Tangent *= -1.;
+		return Tangent;
+	}
+	else
+	{
+		ensureCADKernel(false);
+		return FPoint2D::ZeroPoint;
+	}
+}
+
+
 void FTopologicalEdge::SpawnIdent(FDatabase& Database)
 {
 	if (!FEntity::SetId(Database))
@@ -579,6 +898,55 @@ void FTopologicalEdge::SpawnIdent(FDatabase& Database)
 	}
 }
 
+TSharedPtr<FTopologicalVertex> FTopologicalEdge::SplitAt(double SplittingCoordinate, const FPoint& NewVertexCoordinate, bool bKeepStartVertexConnectivity, TSharedPtr<FTopologicalEdge>& NewEdge)
+{
+	if (GetTwinsEntityCount() > 1)
+	{
+		return TSharedPtr<FTopologicalVertex>();
+		ensureCADKernel(false);
+		// TODO
+	}
+
+	TSharedRef<FTopologicalVertex> MiddelVertex = FEntity::MakeShared<FTopologicalVertex>(NewVertexCoordinate);
+
+	if(bKeepStartVertexConnectivity)
+	{
+		FLinearBoundary NewEdgeBoundary(SplittingCoordinate, Boundary.Max);
+		NewEdge = Make(Curve.ToSharedRef(), MiddelVertex, EndVertex.ToSharedRef(), NewEdgeBoundary);
+	}
+	else
+	{
+		FLinearBoundary NewEdgeBoundary(Boundary.Min, SplittingCoordinate);
+		NewEdge = Make(Curve.ToSharedRef(), StartVertex.ToSharedRef(), MiddelVertex, NewEdgeBoundary);
+	}
+	if (!NewEdge.IsValid())
+	{
+		return TSharedPtr<FTopologicalVertex>();
+	}
+
+	TSharedRef<FTopologicalEdge> ThisEdge = StaticCastSharedRef<FTopologicalEdge>(AsShared());
+
+	if(bKeepStartVertexConnectivity)
+	{
+		EndVertex->RemoveConnectedEdge(ThisEdge);
+		EndVertex = MiddelVertex;
+		Boundary.Max = SplittingCoordinate;
+	}
+	else
+	{
+		StartVertex->RemoveConnectedEdge(ThisEdge);
+		StartVertex = MiddelVertex;
+		Boundary.Min = SplittingCoordinate;
+	}
+	MiddelVertex->AddConnectedEdge(ThisEdge);
+	Length3D = -1.;
+
+	Loop.Pin()->SplitEdge(ThisEdge, NewEdge, bKeepStartVertexConnectivity);
+	return MiddelVertex;
+}
+
+
+
 #ifdef CADKERNEL_DEV
 FInfoEntity& FEdgeLink::GetInfo(FInfoEntity& Info) const
 {
diff --git a/Engine/Source/Runtime/Datasmith/CADKernel/Private/CADKernel/Topo/TopologicalFace.cpp b/Engine/Source/Runtime/Datasmith/CADKernel/Private/CADKernel/Topo/TopologicalFace.cpp
--- a/Engine/Source/Runtime/Datasmith/CADKernel/Private/CADKernel/Topo/TopologicalFace.cpp
+++ b/Engine/Source/Runtime/Datasmith/CADKernel/Private/CADKernel/Topo/TopologicalFace.cpp
@@ -53,10 +53,9 @@ void FTopologicalFace::ApplyNaturalLoops()
 
 	TFunction<void(const FPoint&, const FPoint&)> BuildEdge = [&](const FPoint& StartPoint, const FPoint& EndPoint)
 	{
-		double Tolerance2D = CarrierSurface->Get2DTolerance();
 		double Tolerance3D = CarrierSurface->Get3DTolerance();
-		TSharedRef<FCurve> Curve2D = FEntity::MakeShared<FSegmentCurve>(Tolerance2D, StartPoint, EndPoint, 2);
-		TSharedRef<FRestrictionCurve> Curve3D = FEntity::MakeShared<FRestrictionCurve>(Tolerance3D, CarrierSurface.ToSharedRef(), Curve2D);
+		TSharedRef<FCurve> Curve2D = FEntity::MakeShared<FSegmentCurve>(StartPoint, EndPoint, 2);
+		TSharedRef<FRestrictionCurve> Curve3D = FEntity::MakeShared<FRestrictionCurve>(CarrierSurface.ToSharedRef(), Curve2D);
 		TSharedPtr<FTopologicalEdge> Edge = FTopologicalEdge::Make(Curve3D);
 		if (Edge.IsValid())
 		{
@@ -95,7 +94,7 @@ void FTopologicalFace::ApplyNaturalLoops()
 	TArray<EOrientation> Orientations;
 	Orientations.Init(EOrientation::Front, Edges.Num());
 
-	TSharedPtr<FTopologicalLoop> Loop = FTopologicalLoop::Make(Edges, Orientations);
+	TSharedPtr<FTopologicalLoop> Loop = FTopologicalLoop::Make(Edges, Orientations, CarrierSurface->Get3DTolerance());
 	AddLoop(Loop);
 }
 
@@ -264,11 +263,13 @@ void FTopologicalFace::SpawnIdent(FDatabase& Database)
 FInfoEntity& FTopologicalFace::GetInfo(FInfoEntity& Info) const
 {
 	return FTopologicalEntity::GetInfo(Info)
+		.Add(TEXT("Hosted by"), (TWeakPtr<FEntity>&) HostedBy)
 		.Add(TEXT("Carrier Surface"), CarrierSurface)
 		.Add(TEXT("Boundary"), (FSurfacicBoundary&) Boundary)
 		.Add(TEXT("Loops"), Loops)
 		.Add(TEXT("QuadCriteria"), QuadCriteria)
-		.Add(TEXT("mesh"), Mesh);
+		.Add(TEXT("mesh"), Mesh)
+		.Add(*this);
 }
 #endif
 
@@ -308,7 +309,6 @@ void FTopologicalFace::ChooseFinalDeltaUs()
 	ChooseFinalDeltas(CrossingPointDeltaMins[EIso::IsoV], CrossingPointDeltaMaxs[EIso::IsoV]);
 }
 
-
 // =========================================================================================================================================================================================================
 // =========================================================================================================================================================================================================
 // =========================================================================================================================================================================================================
diff --git a/Engine/Source/Runtime/Datasmith/CADKernel/Private/CADKernel/Topo/TopologicalLoop.cpp b/Engine/Source/Runtime/Datasmith/CADKernel/Private/CADKernel/Topo/TopologicalLoop.cpp
--- a/Engine/Source/Runtime/Datasmith/CADKernel/Private/CADKernel/Topo/TopologicalLoop.cpp
+++ b/Engine/Source/Runtime/Datasmith/CADKernel/Private/CADKernel/Topo/TopologicalLoop.cpp
@@ -15,19 +15,20 @@
 
 using namespace CADKernel;
 
-TSharedPtr<FTopologicalLoop> FTopologicalLoop::Make(const TArray<TSharedPtr<FTopologicalEdge>>& InEdges, const TArray<EOrientation>& InEdgeDirections)
+TSharedPtr<FTopologicalLoop> FTopologicalLoop::Make(const TArray<TSharedPtr<FTopologicalEdge>>& InEdges, const TArray<EOrientation>& InEdgeDirections, double GeometricTolerance)
 {
 	TSharedRef<FTopologicalLoop> Loop = FEntity::MakeShared<FTopologicalLoop>(InEdges, InEdgeDirections);
+	Loop->EnsureLogicalClosing(GeometricTolerance);
 	for (FOrientedEdge& OrientedEdge : Loop->GetEdges())
 	{
-		OrientedEdge.Entity->SetBoundary(Loop);
+		OrientedEdge.Entity->SetLoop(Loop);
 	}
 	return Loop;
 }
 
 FTopologicalLoop::FTopologicalLoop(const TArray<TSharedPtr<FTopologicalEdge>>& InEdges, const TArray<EOrientation>& InEdgeDirections)
 	: FTopologicalEntity()
-	, Surface(TSharedPtr<FTopologicalFace>())
+	, Face(TSharedPtr<FTopologicalFace>())
 	, bExternalLoop(true)
 {
 	Edges.Reserve(InEdges.Num());
@@ -37,7 +38,6 @@ FTopologicalLoop::FTopologicalLoop(const TArray<TSharedPtr<FTopologicalEdge>>& I
 		EOrientation Orientation = InEdgeDirections[Index];
 		Edges.Emplace(Edge, Orientation);
 	}
-	EnsureLogicalClosing();
 }
 
 void FTopologicalLoop::RemoveEdge(TSharedPtr<FTopologicalEdge>& EdgeToRemove)
@@ -46,7 +46,7 @@ void FTopologicalLoop::RemoveEdge(TSharedPtr<FTopologicalEdge>& EdgeToRemove)
 	{
 		if (Edges[IEdge].Entity == EdgeToRemove)
 		{
-			EdgeToRemove->RemoveBoundary();
+			EdgeToRemove->RemoveLoop();
 			Edges.RemoveAt(IEdge);
 			return;
 		}
@@ -194,7 +194,7 @@ void FTopologicalLoop::Orient()
 	{
 		F3DDebugSession GraphicSession(TEXT("Points of evaluation"));
 		{
-			F3DDebugSession G(*FString::Printf(TEXT("Loop Discretization %d"), Surface.Pin()->GetId()));
+			F3DDebugSession G(*FString::Printf(TEXT("Loop Discretization %d"), Face.Pin()->GetId()));
 			DisplayPolyline(LoopSampling);
 			for (int32 Index = 0; Index < LoopSampling.Num(); ++Index)
 			{
@@ -202,27 +202,27 @@ void FTopologicalLoop::Orient()
 			}
 		}
 		{
-			F3DDebugSession G(*FString::Printf(TEXT("Seg UMin Loop Discretization %d"), Surface.Pin()->GetId()));
+			F3DDebugSession G(*FString::Printf(TEXT("Seg UMin Loop Discretization %d"), Face.Pin()->GetId()));
 			::DisplaySegment(LoopSampling[IndexUMin + 1], LoopSampling[IndexUMin]);
 			::Display(LoopSampling[IndexUMin], EVisuProperty::RedPoint, IndexUMin);
 		}
 		{
-			F3DDebugSession G(*FString::Printf(TEXT("Seg UMax Loop Discretization %d"), Surface.Pin()->GetId()));
+			F3DDebugSession G(*FString::Printf(TEXT("Seg UMax Loop Discretization %d"), Face.Pin()->GetId()));
 			::DisplaySegment(LoopSampling[IndexUMax + 1], LoopSampling[IndexUMax]);
 			::Display(LoopSampling[IndexUMax], EVisuProperty::RedPoint, IndexUMax);
 		}
 		{
-			F3DDebugSession G(*FString::Printf(TEXT("Seg VMin Loop Discretization %d"), Surface.Pin()->GetId()));
+			F3DDebugSession G(*FString::Printf(TEXT("Seg VMin Loop Discretization %d"), Face.Pin()->GetId()));
 			::DisplaySegment(LoopSampling[IndexVMin + 1], LoopSampling[IndexVMin]);
 			::Display(LoopSampling[IndexVMin], EVisuProperty::RedPoint, IndexVMin);
 		}
 		{
-			F3DDebugSession G(*FString::Printf(TEXT("Seg VMax Loop Discretization %d"), Surface.Pin()->GetId()));
+			F3DDebugSession G(*FString::Printf(TEXT("Seg VMax Loop Discretization %d"), Face.Pin()->GetId()));
 			::DisplaySegment(LoopSampling[IndexVMax + 1], LoopSampling[IndexVMax]);
 			::Display(LoopSampling[IndexVMax], EVisuProperty::RedPoint, IndexVMax);
 		}
 		//Wait();
-		FMessage::Printf(Log, TEXT("WARNING: Loop Orientation of surface %d could failed\n"), Surface.Pin()->GetId());
+		FMessage::Printf(Log, TEXT("WARNING: Loop Orientation of surface %d could failed\n"), Face.Pin()->GetId());
 	}
 #endif
 
@@ -250,8 +250,29 @@ void FTopologicalLoop::ReplaceEdge(TSharedPtr<FTopologicalEdge>& OldEdge, TShare
 		if (Edges[IEdge].Entity == OldEdge)
 		{
 			Edges[IEdge].Entity = NewEdge;
-			OldEdge->RemoveBoundary();
-			NewEdge->SetBoundary(StaticCastSharedRef<FTopologicalLoop>(AsShared()));
+			OldEdge->RemoveLoop();
+			NewEdge->SetLoop(StaticCastSharedRef<FTopologicalLoop>(AsShared()));
+			return;
+		}
+	}
+	ensureCADKernel(false);
+}
+
+void FTopologicalLoop::SplitEdge(TSharedPtr<FTopologicalEdge> SplitEdge, TSharedPtr<FTopologicalEdge> NewEdge, bool bSplitEdgeIsFirst)
+{
+	TSharedRef<FTopologicalLoop> Loop = StaticCastSharedRef<FTopologicalLoop>(AsShared());
+	NewEdge->SetLoop(Loop);
+
+	for (int32 IEdge = 0; IEdge < Edges.Num(); IEdge++)
+	{
+		if (Edges[IEdge].Entity == SplitEdge)
+		{
+			EOrientation OldEdgeDirection = Edges[IEdge].Direction;
+			if ((OldEdgeDirection == EOrientation::Front) == bSplitEdgeIsFirst)
+			{
+				IEdge++;
+			}
+			Edges.EmplaceAt(IEdge, NewEdge, OldEdgeDirection);
 			return;
 		}
 	}
@@ -264,11 +285,11 @@ void FTopologicalLoop::ReplaceEdge(TSharedPtr<FTopologicalEdge>& Edge, TArray<TS
 	int32 NewEdgeNum = Edges.Num() + NewEdges.Num();
 	TmpEdges.Reserve(NewEdgeNum);
 
-	Edge->RemoveBoundary();
+	Edge->RemoveLoop();
 	TSharedRef<FTopologicalLoop> Loop = StaticCastSharedRef<FTopologicalLoop>(AsShared());
 	for (TSharedPtr<FTopologicalEdge>& NewEdge : NewEdges)
 	{
-		NewEdge->SetBoundary(Loop);
+		NewEdge->SetLoop(Loop);
 	}
 
 	for (int32 IEdge = 0; IEdge < Edges.Num(); IEdge++)
@@ -278,37 +299,80 @@ void FTopologicalLoop::ReplaceEdge(TSharedPtr<FTopologicalEdge>& Edge, TArray<TS
 			EOrientation OldEdgeDirection = Edges[IEdge].Direction;
 			if (OldEdgeDirection == EOrientation::Front)
 			{
-				for (int32 INewEdge = 0; INewEdge < NewEdges.Num(); INewEdge++)
+				Edges[IEdge].Entity = NewEdges[0];
+				for (int32 INewEdge = 1; INewEdge < NewEdges.Num(); INewEdge++)
 				{
-					TmpEdges.Emplace(NewEdges[INewEdge], OldEdgeDirection);
+					IEdge++;
+					Edges.EmplaceAt(IEdge, NewEdges[INewEdge], EOrientation::Front);
 				}
 			}
 			else
 			{
-				for (int32 INewEdge = NewEdges.Num() - 1; INewEdge >= 0; --INewEdge)
+				Edges[IEdge].Entity = NewEdges[0];
+				for (int32 INewEdge = 1; INewEdge < NewEdges.Num(); INewEdge++)
 				{
-					TmpEdges.Emplace(NewEdges[INewEdge], OldEdgeDirection);
+					Edges.EmplaceAt(IEdge, NewEdges[INewEdge], EOrientation::Back);
 				}
 			}
+			return;
 		}
-		else 
+	}
+	ensureCADKernel(false);
+}
+
+void FTopologicalLoop::ReplaceEdges(TArray<FOrientedEdge>& OldEdges, TSharedPtr<FTopologicalEdge>& NewEdge)
+{
+	for (FOrientedEdge& Edge : OldEdges)
+	{
+		Edge.Entity->RemoveLoop();
+	}
+
+	TSharedRef<FTopologicalLoop> Loop = StaticCastSharedRef<FTopologicalLoop>(AsShared());
+	NewEdge->SetLoop(Loop);
+
+	for (int32 IEdge = 0; IEdge < Edges.Num(); IEdge++)
+	{
+		if (Edges[IEdge] == OldEdges[0])
 		{
-			TmpEdges.Emplace(Edges[IEdge]);
+			Edges[IEdge].Direction = EOrientation::Front;
+			Edges[IEdge].Entity = NewEdge;
+			IEdge++;
+
+			int32 EdgeCount = Edges.Num();
+			int32 EdgeToRemoveCount = OldEdges.Num() - 1;
+			if (IEdge + EdgeToRemoveCount < EdgeCount)
+			{
+				for (int32 Index = 0; Index < EdgeToRemoveCount; Index++)
+				{
+					Edges.RemoveAt(IEdge);
+				}
+			}
+			else
+			{
+				int32 EdgeToRemoveAtEnd = EdgeCount - IEdge;
+				int32 EdgeToRemoveAtStart = EdgeToRemoveCount - EdgeToRemoveAtEnd;
+				Edges.SetNum(IEdge);
+				for (int32 Index = 0; Index < EdgeToRemoveAtStart; Index++)
+				{
+					Edges.RemoveAt(0);
+				}
+			}
+			return;
 		}
 	}
-
-	Swap(TmpEdges, Edges);
+	ensureCADKernel(false);
 }
 
-void FTopologicalLoop::ReplaceEdgesWithMergedEdge(TArray<TSharedPtr<FTopologicalEdge>>& OldEdges, TSharedPtr<FTopologicalVertex>& MiddleVertex, TSharedPtr<FTopologicalEdge>& NewEdge)
+// Need to be reviewed
+/*void FTopologicalLoop::ReplaceEdgesWithMergedEdge(TArray<TSharedPtr<FTopologicalEdge>>& OldEdges, TSharedPtr<FTopologicalVertex>& MiddleVertex, TSharedPtr<FTopologicalEdge>& NewEdge)
 {
 	ensureCADKernel(OldEdges.Num() == 2);
 	ensureCADKernel(EdgeCount() != 2);
 
-	NewEdge->SetBoundary(StaticCastSharedRef<FTopologicalLoop>(AsShared()));
+	NewEdge->SetLoop(StaticCastSharedRef<FTopologicalLoop>(AsShared()));
 	for (TSharedPtr<FTopologicalEdge> Edge : OldEdges)
 	{
-		Edge->RemoveBoundary();
+		Edge->RemoveLoop();
 		Edge->GetStartVertex()->RemoveConnectedEdge(Edge.ToSharedRef());
 		Edge->GetEndVertex()->RemoveConnectedEdge(Edge.ToSharedRef());
 	}
@@ -340,7 +404,7 @@ void FTopologicalLoop::ReplaceEdgesWithMergedEdge(TArray<TSharedPtr<FTopological
 		Edges.Pop();
 		Edges.Emplace(NewEdge, NewEdgeOrientation);
 	}
-}
+}*/
 
 
 
@@ -378,20 +442,130 @@ void FTopologicalLoop::ComputeBoundaryProperties(const TArray<int32>& StartSideI
 	}
 }
 
-void FTopologicalLoop::EnsureLogicalClosing()
+void FTopologicalLoop::EnsureLogicalClosing(const double Tolerance3D)
 {
-	TSharedPtr<FTopologicalEdge> Edge = Edges.Last().Entity;
+	TSharedRef<FTopologicalLoop> Loop = StaticCastSharedRef<FTopologicalLoop>(AsShared());
+	const TSharedRef<FSurface>& Surface = Edges[0].Entity->GetCurve()->GetCarrierSurface();
 
-	TSharedPtr<FTopologicalVertex> EdgeStartVertex = Edge->GetStartVertex(Edges[Edges.Num() - 1].Direction);
+	FOrientedEdge PreviousEdge = Edges.Last();
 
-	for (int32 Index = 0; Index < Edges.Num(); Index++)
+	FSurfacicCurveExtremity PreviousExtremities;
+	PreviousEdge.Entity->GetExtremities(PreviousExtremities);
+	FSurfacicCurvePointWithTolerance PreviousExtremity;
+
+	if (PreviousEdge.Direction == EOrientation::Front)
+	{
+		PreviousExtremity = PreviousExtremities[1];
+	}
+	else
+	{
+		PreviousExtremity = PreviousExtremities[0];
+	}
+
+	for (int32 Index = 0; Index < Edges.Num(); ++Index)
 	{
-		const TSharedPtr<FTopologicalVertex> PreviousEdgeEndVertex = Edge->GetOtherVertex(EdgeStartVertex.ToSharedRef());
+		FOrientedEdge OrientedEdge = Edges[Index];
+		FSurfacicCurvePointWithTolerance EdgeExtremities[2];
+		OrientedEdge.Entity->GetExtremities(EdgeExtremities);
+
+		double Gap3D;
+
+		int32 ExtremityIndex = OrientedEdge.Direction == EOrientation::Front ? 0 : 1;
+		Gap3D = EdgeExtremities[ExtremityIndex].Point.SquareDistance(PreviousExtremity.Point);
 
-		Edge = Edges[Index].Entity;
-		EdgeStartVertex = Edge->GetStartVertex(Edges[Index].Direction);
+		TSharedRef<FTopologicalVertex> PreviousEdgeEndVertex = PreviousEdge.Direction == EOrientation::Front ? PreviousEdge.Entity->GetEndVertex() : PreviousEdge.Entity->GetStartVertex();
+		TSharedRef<FTopologicalVertex> EdgeStartVertex = OrientedEdge.Direction == EOrientation::Front ? OrientedEdge.Entity->GetStartVertex() : OrientedEdge.Entity->GetEndVertex();
+
+		FPoint2D Gap2D;
+		Gap2D = EdgeExtremities[ExtremityIndex].Point2D - PreviousExtremity.Point2D;
+		FSurfacicTolerance Tolerance2D = Max(EdgeExtremities[ExtremityIndex].Tolerance, PreviousExtremity.Tolerance);
+
+		if (Gap3D > Tolerance3D)
+		{
+			FMessage::Printf(Log, TEXT("Loop %d Gap 3D : %f\n"), Id, sqrt(Gap3D));
+			FPoint PreviousTangent = PreviousEdge.Entity->GetTangentAt(PreviousEdgeEndVertex);
+			FPoint EdgeTangent = OrientedEdge.Entity->GetTangentAt(EdgeStartVertex);
 
-		PreviousEdgeEndVertex->Link(EdgeStartVertex.ToSharedRef());
+			FPoint Gap = EdgeExtremities[ExtremityIndex].Point - PreviousExtremity.Point;
+
+			double CosAngle = Gap.ComputeCosinus(EdgeTangent);
+			if (CosAngle > 0.9)
+			{
+				OrientedEdge.Entity->ExtendTo(OrientedEdge.Direction == EOrientation::Front, PreviousExtremity.Point2D, PreviousEdgeEndVertex);
+			}
+			else
+			{
+				CosAngle = Gap.ComputeCosinus(PreviousTangent);
+				if (CosAngle < -0.9)
+				{
+					PreviousEdge.Entity->ExtendTo(PreviousEdge.Direction == EOrientation::Front, EdgeExtremities[ExtremityIndex].Point2D, EdgeStartVertex);
+				}
+				else
+				{
+					if (PreviousEdgeEndVertex->IsLinkedTo(EdgeStartVertex))
+					{
+						PreviousEdgeEndVertex->UnlinkTo(EdgeStartVertex);
+					}
+
+					TSharedPtr<FTopologicalEdge> Edge = FTopologicalEdge::Make(Surface, PreviousExtremity.Point2D, PreviousEdgeEndVertex, EdgeExtremities[ExtremityIndex].Point2D, EdgeStartVertex);
+					if (Edge.IsValid())
+					{
+						Edges.EmplaceAt(Index, Edge, EOrientation::Front);
+						Edge->SetLoop(Loop);
+						++Index;
+					}
+					PreviousEdgeEndVertex = EdgeStartVertex;
+				}
+			}
+		}
+		else if (FMath::Abs(Gap2D.U) > Tolerance2D.U || FMath::Abs(Gap2D.V) > Tolerance2D.V)
+		{
+			FMessage::Printf(Log, TEXT("Loop %d Gap 2D : [%f, %f] vs Tol2D [%f, %f]\n"), Id, FMath::Abs(Gap2D.U), FMath::Abs(Gap2D.V), Tolerance2D.U, Tolerance2D.V);
+
+			FPoint2D PreviousTangent = PreviousEdge.Entity->GetTangent2DAt(PreviousEdgeEndVertex);
+			FPoint2D EdgeTangent = OrientedEdge.Entity->GetTangent2DAt(EdgeStartVertex);
+
+			double CosAngle = Gap2D.ComputeCosinus(EdgeTangent);
+			if (CosAngle > 0.9)
+			{
+				OrientedEdge.Entity->ExtendTo(OrientedEdge.Direction == EOrientation::Front, PreviousExtremity.Point2D, PreviousEdgeEndVertex);
+			}
+			else
+			{
+				CosAngle = Gap2D.ComputeCosinus(PreviousTangent);
+				if (CosAngle < -0.9)
+				{
+					PreviousEdge.Entity->ExtendTo(PreviousEdge.Direction == EOrientation::Front, EdgeExtremities[ExtremityIndex].Point2D, EdgeStartVertex);
+				}
+				else
+				{
+					TSharedPtr<FTopologicalEdge> Edge = FTopologicalEdge::Make(Surface, PreviousExtremity.Point2D, PreviousEdgeEndVertex, EdgeExtremities[ExtremityIndex].Point2D, EdgeStartVertex);
+					if (Edge.IsValid())
+					{
+						Edges.EmplaceAt(Index, Edge, EOrientation::Front);
+						Edge->SetLoop(Loop);
+
+						PreviousEdgeEndVertex->Link(EdgeStartVertex);
+						++Index;
+
+					}
+				}
+			}
+		}
+		else
+		{
+			PreviousEdgeEndVertex->Link(EdgeStartVertex);
+		}
+
+		if (OrientedEdge.Direction == EOrientation::Front)
+		{
+			PreviousExtremity = EdgeExtremities[1];
+		}
+		else
+		{
+			PreviousExtremity = EdgeExtremities[0];
+		}
+		PreviousEdge = OrientedEdge;
 	}
 }
 
@@ -400,7 +574,7 @@ FInfoEntity& FTopologicalLoop::GetInfo(FInfoEntity& Info) const
 {
 	return FEntity::GetInfo(Info)
 		.Add(TEXT("Edges"), (TArray<TOrientedEntity<FEntity>>&) Edges)
-		.Add(TEXT("domain"), Surface);
+		.Add(TEXT("Hosted by"), Face);
 }
 #endif
 
diff --git a/Engine/Source/Runtime/Datasmith/CADKernel/Private/CADKernel/Topo/TopologicalVertex.cpp b/Engine/Source/Runtime/Datasmith/CADKernel/Private/CADKernel/Topo/TopologicalVertex.cpp
--- a/Engine/Source/Runtime/Datasmith/CADKernel/Private/CADKernel/Topo/TopologicalVertex.cpp
+++ b/Engine/Source/Runtime/Datasmith/CADKernel/Private/CADKernel/Topo/TopologicalVertex.cpp
@@ -40,11 +40,14 @@ void FTopologicalVertex::RemoveConnectedEdge(TSharedRef<FTopologicalEdge> Edge)
 
 bool FTopologicalVertex::IsBorderVertex()
 {
-	for (TWeakPtr<FTopologicalEdge> Edge : GetConnectedEdges())
+	for (const TWeakPtr<FTopologicalVertex>& Vertex : GetTwinsEntities())
 	{
-		if (Edge.Pin()->GetTwinsEntityCount() == 1)
+		for (const TWeakPtr<FTopologicalEdge>& Edge : Vertex.Pin()->GetDirectConnectedEdges())
 		{
-			return true;
+			if (Edge.Pin()->GetTwinsEntityCount() == 1)
+			{
+				return true;
+			}
 		}
 	}
 	return false;
@@ -61,10 +64,10 @@ void FTopologicalVertex::GetConnectedEdges(TSharedPtr<FTopologicalVertex> OtherV
 	Edges.Reserve(GetTwinsEntityCount());
 
 	TSharedPtr<TTopologicalLink<FTopologicalVertex>> OtherVertexLink = OtherVertex->GetLink();
-	for (TWeakPtr<FTopologicalVertex> Vertex : GetTwinsEntities())
+	for (const TWeakPtr<FTopologicalVertex>& Vertex : GetTwinsEntities())
 	{
 		ensureCADKernel(Vertex.IsValid());
-		for (TWeakPtr<FTopologicalEdge> Edge : Vertex.Pin()->GetConnectedEdges())
+		for (const TWeakPtr<FTopologicalEdge>& Edge : Vertex.Pin()->GetDirectConnectedEdges())
 		{
 			ensureCADKernel(Edge.IsValid());
 			if (Edge.Pin()->GetOtherVertex(Vertex.Pin().ToSharedRef())->GetLink() == OtherVertexLink)
@@ -98,10 +101,38 @@ void FTopologicalVertex::Link(TSharedRef<FTopologicalVertex> Twin)
 	GetLink()->DefineActiveEntity();
 }
 
+void FTopologicalVertex::UnlinkTo(TSharedRef<FTopologicalVertex> OtherVertex)
+{
+	TSharedPtr<FVertexLink> OldLink = GetLink();
+	ResetTopologicalLink();
+	OtherVertex->ResetTopologicalLink();
+
+	for (const TWeakPtr<FTopologicalVertex>& WeakVertex : OldLink->GetTwinsEntities())
+	{
+		TSharedPtr<FTopologicalVertex> Vertex = WeakVertex.Pin();
+		if (!Vertex.IsValid() || Vertex.Get() == this || Vertex == OtherVertex)
+		{
+			continue;
+		}
+
+		Vertex->ResetTopologicalLink();
+		double Distance1 = Distance(Vertex.ToSharedRef());
+		double Distance2 = OtherVertex->Distance(Vertex.ToSharedRef());
+		if (Distance1 < Distance2)
+		{
+			Link(Vertex.ToSharedRef());
+		}
+		else
+		{
+			OtherVertex->Link(Vertex.ToSharedRef());
+		}
+	}
+}
+
 void FVertexLink::ComputeBarycenter()
 {
 	Barycenter = FPoint::ZeroPoint;
-	for (TWeakPtr<FTopologicalVertex> Vertex : TwinsEntities)
+	for (const TWeakPtr<FTopologicalVertex>& Vertex : TwinsEntities)
 	{
 		Barycenter += Vertex.Pin()->Coordinates;
 	}
diff --git a/Engine/Source/Runtime/Datasmith/CADKernel/Private/CADKernel/UI/Display.cpp b/Engine/Source/Runtime/Datasmith/CADKernel/Private/CADKernel/UI/Display.cpp
--- a/Engine/Source/Runtime/Datasmith/CADKernel/Private/CADKernel/UI/Display.cpp
+++ b/Engine/Source/Runtime/Datasmith/CADKernel/Private/CADKernel/UI/Display.cpp
@@ -86,7 +86,7 @@ namespace CADKernel
 		bool bShowOrientation = FSystem::Get().GetVisu()->GetParameters()->bDisplayCADOrient;
 
 		TArray<FPoint> Polyline;
-		Curve->GetDiscretizationPoints<FPoint>(Curve->GetBoundary(), EOrientation::Front, Polyline);
+		Curve->GetDiscretizationPoints<FPoint>(Boundary, EOrientation::Front, Polyline);
 		Draw(Polyline, Property);
 
 		if (bShowOrientation)
@@ -463,7 +463,11 @@ namespace CADKernel
 
 	void Display(const TSharedPtr<FShell>& Shell)
 	{
-		Draw(Shell);
+		FProgress Progress(Shell->GetFaces().Num(), TEXT("Display Shell"));
+		for (const FOrientedFace& Face : Shell->GetFaces())
+		{
+			Display(Face.Entity);
+		}
 	}
 
 	void Draw(const TSharedPtr<FShell>& Shell)
@@ -471,13 +475,15 @@ namespace CADKernel
 		FProgress Progress(Shell->GetFaces().Num(), TEXT("Display Shell"));
 		for (const FOrientedFace& Face : Shell->GetFaces())
 		{
-			Display(Face.Entity);
+			Draw(Face.Entity);
 		}
 	}
 
 	void Display(const TSharedPtr<FBody>& Body)
 	{
 		FProgress Progress(Body->GetShells().Num(), TEXT("Display Body"));
+
+		F3DDebugSegment GraphicSegment(Body->GetId());
 		for (TSharedPtr<FShell> Shell : Body->GetShells())
 		{
 			if (!Shell.IsValid())
@@ -832,7 +838,7 @@ namespace CADKernel
 	{
 		FProgress MainProgress(2, TEXT("Display model"));
 		{
-			const TArray<TSharedPtr<FBody>>& Bodies = Model->GetBodyList();
+			const TArray<TSharedPtr<FBody>>& Bodies = Model->GetBodies();
 			FProgress BodyProgress((int32)Bodies.Num(), TEXT("Bodies"));
 			for (const TSharedPtr<FBody>& Body : Bodies)
 			{
@@ -841,8 +847,8 @@ namespace CADKernel
 		}
 
 		{
-			const TArray<TSharedPtr<FTopologicalFace>>& Surfaces = Model->GetSurfaceList();
-			FProgress SurfaceProgress((int32)Surfaces.Num(), TEXT("surfaces"));
+			const TArray<TSharedPtr<FTopologicalFace>>& Surfaces = Model->GetFaces();
+			FProgress SurfaceProgress((int32)Surfaces.Num(), TEXT("Surfaces"));
 			for (const TSharedPtr<FTopologicalFace>& Surface : Surfaces)
 			{
 				Display(Surface);
@@ -850,6 +856,98 @@ namespace CADKernel
 		}
 	}
 
+	void DisplayProductTree(const TSharedPtr<FModel>& Model)
+	{
+		F3DDebugSession GraphicSession(FString::Printf(TEXT("%s %d"), Model->GetTypeName(), Model->GetId()), { Model->GetId() });
+
+		FProgress MainProgress(2, TEXT("Display Model"));
+		{
+			const TArray<TSharedPtr<FBody>>& Bodies = Model->GetBodies();
+			FProgress BodyProgress((int32)Bodies.Num(), TEXT("Bodies"));
+			for (const TSharedPtr<FBody>& Body : Bodies)
+			{
+				DisplayProductTree(Body);
+			}
+		}
+
+		{
+			const TArray<TSharedPtr<FTopologicalFace>>& Faces = Model->GetFaces();
+			FProgress SurfaceProgress((int32)Faces.Num(), TEXT("Faces"));
+			for (const TSharedPtr<FTopologicalFace>& Face : Faces)
+			{
+				Display(Face);
+			}
+		}
+
+	}
+
+	void DisplayProductTree(const TSharedPtr<FBody>& Body)
+	{
+#ifdef CORETECHBRIDGE_DEBUG
+		F3DDebugSession GraphicSession(FString::Printf(TEXT("%s %s KioId: %d Id: %d"), Body->GetTypeName(), Body->GetName(), Body->GetKioId(), Body->GetId()), { Body->GetId() });
+#else
+		F3DDebugSession GraphicSession(FString::Printf(TEXT("%s %s Id: %d"), Body->GetTypeName(), Body->GetName(), Body->GetId()), { Body->GetId() });
+#endif
+		FProgress Progress(Body->GetShells().Num(), TEXT("Display Body"));
+
+		for (TSharedPtr<FShell> Shell : Body->GetShells())
+		{
+			if (!Shell.IsValid())
+			{
+				continue;
+			}
+			DisplayProductTree(Shell);
+		}
+
+	}
+
+	void DisplayProductTree(const TSharedPtr<FShell>& Shell)
+	{
+#ifdef CORETECHBRIDGE_DEBUG
+		F3DDebugSession GraphicSession(FString::Printf(TEXT("%s %s KioId: %d Id: %d"), Shell->GetTypeName(), Shell->GetName(), Shell->GetKioId(), Shell->GetId()), { Shell->GetId() });
+#else
+		F3DDebugSession GraphicSession(FString::Printf(TEXT("%s %s Id: %d"), Shell->GetTypeName(), Shell->GetName(), Shell->GetId()), { Shell->GetId() });
+#endif
+		F3DDebugSegment Segment(Shell->GetId());
+		Draw(Shell);
+	}
+
+	void DisplayProductTree(const TSharedPtr<FEntity>& Entity)
+	{
+		if (!Entity.IsValid())
+		{
+			FMessage::Printf(Log, TEXT("Entity not found"));
+			return;
+		}
+
+		FTimePoint StartTime = FChrono::Now();
+
+		switch (Entity->GetEntityType())
+		{
+		case EEntity::TopologicalFace:
+			Display(StaticCastSharedPtr<FTopologicalFace>(Entity));
+			break;
+		case EEntity::Shell:
+			DisplayProductTree(StaticCastSharedPtr<FShell>(Entity));
+			break;
+		case EEntity::Body:
+			DisplayProductTree(StaticCastSharedPtr<FBody>(Entity));
+			break;
+		case EEntity::Model:
+			DisplayProductTree(StaticCastSharedPtr<FModel>(Entity));
+			break;
+		default:
+			FMessage::Printf(Log, TEXT("Unable to display Entity of type %s"), FEntity::GetTypeName(Entity->GetEntityType()));
+		}
+
+		FDuration DisplayDuration = FChrono::Elapse(StartTime);
+		FChrono::PrintClockElapse(Log, TEXT("  "), TEXT("Display "), DisplayDuration);
+	}
+
+
+
+
+
 	void Display(const TSharedPtr<FGroup>& Group)
 	{
 		TArray<TSharedPtr<FEntity>> Entities;
diff --git a/Engine/Source/Runtime/Datasmith/CADKernel/Private/CADKernel/UI/Message.cpp b/Engine/Source/Runtime/Datasmith/CADKernel/Private/CADKernel/UI/Message.cpp
--- a/Engine/Source/Runtime/Datasmith/CADKernel/Private/CADKernel/UI/Message.cpp
+++ b/Engine/Source/Runtime/Datasmith/CADKernel/Private/CADKernel/UI/Message.cpp
@@ -40,7 +40,7 @@ namespace CADKernel
 		FString Indentation;
 		TCHAR* Buffer = NULL;
 		int32 Result = -1;
-		if (VerboseLevel >= Level || (LogFile.IsValid() && Level < Debug) || (SpyFile.IsValid() && Level == Spy))
+		if (VerboseLevel >= Level || (LogFile.IsValid() && Level <= EVerboseLevel::Debug) || (SpyFile.IsValid() && Level == EVerboseLevel::Spy))
 		{
 			int32 BufferSize = 1024;
 
diff --git a/Engine/Source/Runtime/Datasmith/CADKernel/Private/CADKernel/UI/Visu.cpp b/Engine/Source/Runtime/Datasmith/CADKernel/Private/CADKernel/UI/Visu.cpp
--- a/Engine/Source/Runtime/Datasmith/CADKernel/Private/CADKernel/UI/Visu.cpp
+++ b/Engine/Source/Runtime/Datasmith/CADKernel/Private/CADKernel/UI/Visu.cpp
@@ -6,7 +6,7 @@ namespace CADKernel
 {
 	const TCHAR* VisuPropertyNames[] = 
 	{
-		TEXT("Bounding Box"),              
+		TEXT("Bounding Box"),                   
 		TEXT("Points"),
 		TEXT("Elements"),
 		TEXT("Curves"),
@@ -31,6 +31,10 @@ namespace CADKernel
 		TEXT("GreenCurve"),
 		TEXT("PinkPoint"),
 		TEXT("PinkCurve"),
+		TEXT("OrangePoint"),
+		TEXT("OrangeCurve"),
 		nullptr
 	};
+
+
 }
diff --git a/Engine/Source/Runtime/Datasmith/CADKernel/Private/CADKernel/Utils/ArrayUtils.h b/Engine/Source/Runtime/Datasmith/CADKernel/Private/CADKernel/Utils/ArrayUtils.h
--- a/Engine/Source/Runtime/Datasmith/CADKernel/Private/CADKernel/Utils/ArrayUtils.h
+++ b/Engine/Source/Runtime/Datasmith/CADKernel/Private/CADKernel/Utils/ArrayUtils.h
@@ -198,5 +198,27 @@ namespace CADKernel
 				OutArray.Add(Value);
 			}
 		}
+
+		/**
+		 * Find the index of the segment containing the coordinate i.e. Coordinate is in [InCoordinates[OutIndex], InCoordinates[OutIndex+1]]
+		 * This method is very fast if initial value of Index is near the solution.
+		 * This is very useful to process linked set of coordinate like loop point coordinates.
+		 */
+		inline void FindCoordinateIndex(const TArray<double>& InCoordinates, double Coordinate, int32& OutIndex)
+		{
+			ensureCADKernel(InCoordinates.IsValidIndex(OutIndex));
+
+			while (Coordinate < InCoordinates[OutIndex] && OutIndex > 0)
+			{
+				OutIndex--;
+			}
+
+			for (; OutIndex + 2 < InCoordinates.Num() && Coordinate > InCoordinates[OutIndex + 1]; ++OutIndex)
+			{
+			}
+
+			ensureCADKernel(InCoordinates.IsValidIndex(OutIndex));
+		};
+
 	};
 }
diff --git a/Engine/Source/Runtime/Datasmith/CADKernel/Public/CADKernel/Core/CADKernelArchive.h b/Engine/Source/Runtime/Datasmith/CADKernel/Public/CADKernel/Core/CADKernelArchive.h
--- a/Engine/Source/Runtime/Datasmith/CADKernel/Public/CADKernel/Core/CADKernelArchive.h
+++ b/Engine/Source/Runtime/Datasmith/CADKernel/Public/CADKernel/Core/CADKernelArchive.h
@@ -11,12 +11,14 @@ namespace CADKernel
 {
 	class FEntity;
 	class FSession;
+	class FModel;
 
 	class FCADKernelArchive
 	{
 	public:
 		FArchive& Archive;
 		FSession& Session;
+		FModel* ArchiveModel = nullptr;
 
 		FCADKernelArchive(FSession& InSession, FArchive& InArchive)
 			: Archive(InArchive)
@@ -62,11 +64,34 @@ namespace CADKernel
 			return Archive.IsSaving();
 		}
 
-		void Serialize(void* Value, int64 Length) 
+		void Serialize(void* Value, int64 Length)
 		{ 
 			Archive.Serialize(Value, Length);
 		}
 
+		template<typename EntityType>
+		void Serialize(TArray<EntityType>& EntityArray)
+		{
+			if(Archive.IsLoading())
+			{
+				int32 ElementCount = 0;
+				Archive << ElementCount;
+				EntityArray.SetNum(ElementCount);
+			}
+			else
+			{
+				int32 ElementCount = EntityArray.Num();
+				Archive << ElementCount;
+			}
+			Archive.Serialize((void*) EntityArray.GetData(), EntityArray.Num()*sizeof(EntityType));
+		}
+
+		template<typename EntityType>
+		void Serialize(EntityType& EntityArray)
+		{
+			Archive.Serialize((void*) &EntityArray, sizeof(EntityType));
+		}
+
 		void SetReferencedEntityOrAddToWaitingList(FIdent ArchiveId, TWeakPtr<FEntity>& Entity);
 		void SetReferencedEntityOrAddToWaitingList(FIdent ArchiveId, TSharedPtr<FEntity>& Entity);
 		void AddEntityToSave(FIdent Id);
diff --git a/Engine/Source/Runtime/Datasmith/CADKernel/Public/CADKernel/Core/Database.h b/Engine/Source/Runtime/Datasmith/CADKernel/Public/CADKernel/Core/Database.h
--- a/Engine/Source/Runtime/Datasmith/CADKernel/Public/CADKernel/Core/Database.h
+++ b/Engine/Source/Runtime/Datasmith/CADKernel/Public/CADKernel/Core/Database.h
@@ -81,6 +81,8 @@ namespace CADKernel
 		 * Remove from database
 		 */
 		void RemoveEntity(FIdent idEntity);
+		void RemoveEntity(FEntity& Entity);
+
 
 		/**
 		 * Recursive serialization of a selection. The selection and all the dependencies are serialized. 
@@ -111,6 +113,10 @@ namespace CADKernel
 		 */
 		void AddEntity(TSharedRef<FEntity> Entity)
 		{
+			if (Entity->GetId() > 0)
+			{
+				return;
+			}
 			FIdent Id = CreateId();
 			DatabaseEntities[Id] = Entity;
 			Entity->Id = Id;
@@ -155,6 +161,8 @@ namespace CADKernel
 				}
 			}
 
+			// Reset Id to add it in the DatabaseEntities map.
+			Entity->Id = 0;
 			// Add the entity to DatabaseEntities map.
 			AddEntity(Entity);
 		}
diff --git a/Engine/Source/Runtime/Datasmith/CADKernel/Public/CADKernel/Core/Entity.h b/Engine/Source/Runtime/Datasmith/CADKernel/Public/CADKernel/Core/Entity.h
--- a/Engine/Source/Runtime/Datasmith/CADKernel/Public/CADKernel/Core/Entity.h
+++ b/Engine/Source/Runtime/Datasmith/CADKernel/Public/CADKernel/Core/Entity.h
@@ -65,6 +65,10 @@ namespace CADKernel
 		{
 			OtherEntity* Entity = new OtherEntity(Forward<InArgTypes>(Args)...);
 			TSharedRef<OtherEntity> NewShared = MakeShareable<OtherEntity>(Entity);
+
+#ifdef CADKERNEL_DEV
+			AddEntityInDatabase(NewShared);
+#endif
 			return NewShared;
 		}
 
@@ -216,6 +220,7 @@ namespace CADKernel
 #ifdef CADKERNEL_DEV
 		void InfoEntity() const;
 		virtual FInfoEntity& GetInfo(FInfoEntity& EntityInfo) const;
+		static void AddEntityInDatabase(TSharedRef<FEntity> Entity);
 #endif
 
 		virtual void SpawnIdent(FDatabase& Database)
diff --git a/Engine/Source/Runtime/Datasmith/CADKernel/Public/CADKernel/Core/EntityGeom.h b/Engine/Source/Runtime/Datasmith/CADKernel/Public/CADKernel/Core/EntityGeom.h
--- a/Engine/Source/Runtime/Datasmith/CADKernel/Public/CADKernel/Core/EntityGeom.h
+++ b/Engine/Source/Runtime/Datasmith/CADKernel/Public/CADKernel/Core/EntityGeom.h
@@ -43,6 +43,13 @@ namespace CADKernel
 			Ar << CtKioId;
 		}
 
+#ifdef CORETECHBRIDGE_DEBUG
+		FIdent GetKioId()
+		{
+			return CtKioId;
+		}
+#endif
+
 #ifdef CADKERNEL_DEV
 		virtual FInfoEntity& GetInfo(FInfoEntity& Info) const override;
 #endif
diff --git a/Engine/Source/Runtime/Datasmith/CADKernel/Public/CADKernel/Core/MetadataDictionary.h b/Engine/Source/Runtime/Datasmith/CADKernel/Public/CADKernel/Core/MetadataDictionary.h
--- a/Engine/Source/Runtime/Datasmith/CADKernel/Public/CADKernel/Core/MetadataDictionary.h
+++ b/Engine/Source/Runtime/Datasmith/CADKernel/Public/CADKernel/Core/MetadataDictionary.h
@@ -1,6 +1,7 @@
 // Copyright Epic Games, Inc. All Rights Reserved.
 #pragma once
 
+#include "CADKernel/Core/CADKernelArchive.h"
 #include "CADKernel/Core/Types.h"
 #include "CADKernel/Math/MathConst.h"
 
@@ -12,37 +13,90 @@ namespace CADKernel
 	private:
 
 		FString Name;
-		int32 HostId = 0;
-		int32 PatchId = 0;
- 		int32 LayerId = -1; 
-		uint32 ColorId = 0;
-		uint32 MaterialId = 0;
+		int32   HostId = 0;
+		int32   PatchId = 0;
+		uint32  ColorId = 0;
+		uint32  MaterialId = 0;
+		int32   LayerId = -1;
 
 	public:
 
 		FMetadataDictionary() = default;
 
+		void SerializeMetadata(FCADKernelArchive& Ar)
+		{
+			Ar.Archive << Name;
+			Ar.Archive << HostId;
+			Ar.Archive << PatchId;
+			Ar.Archive << ColorId;
+			Ar.Archive << MaterialId;
+			Ar.Archive << LayerId;
+		}
+
+		void CompleteDictionary(const FMetadataDictionary& MetaData)
+		{
+			if (Name.IsEmpty())
+			{
+				Name = MetaData.Name;
+			}
+			if (HostId == 0)
+			{
+				HostId = MetaData.HostId;
+			}
+			if (PatchId == 0)
+			{
+				PatchId = MetaData.PatchId;
+			}
+			if (ColorId == 0)
+			{
+				HostId = MetaData.ColorId;
+			}
+			if (MaterialId == 0)
+			{
+				HostId = MetaData.MaterialId;
+			}
+			if (LayerId == 0)
+			{
+				HostId = MetaData.LayerId;
+			}
+		}
+
 		void SetHostId(const int32 InHostId)
 		{
 			HostId = InHostId;
 		}
 
+		const int32 GetHostId(const int32 InHostId) const
+		{
+			return HostId;
+		}
+
 		void SetLayer(const int32 InLayerId)
 		{
 			LayerId = InLayerId;
 		}
 
+		int32 GetLayer() const
+		{
+			return LayerId;
+		}
+
 		void SetName(const FString& InName)
 		{
 			Name = InName;
 		}
 
+		const TCHAR* GetName() const
+		{
+			return *Name;
+		}
+
 		void SetColorId(const uint32& InColorId)
 		{
 			ColorId = InColorId;
 		}
 
-		uint32 GetColorId()
+		uint32 GetColorId() const
 		{
 			return ColorId;
 		}
@@ -52,7 +106,7 @@ namespace CADKernel
 			MaterialId = InMaterialId;
 		}
 
-		uint32 GetMaterialId()
+		uint32 GetMaterialId() const
 		{
 			return MaterialId;
 		}
@@ -62,7 +116,7 @@ namespace CADKernel
 			PatchId = InPatchId;
 		}
 
-		int32 GetPatchId()
+		int32 GetPatchId() const
 		{
 			return PatchId;
 		}
diff --git a/Engine/Source/Runtime/Datasmith/CADKernel/Public/CADKernel/Core/Session.h b/Engine/Source/Runtime/Datasmith/CADKernel/Public/CADKernel/Core/Session.h
--- a/Engine/Source/Runtime/Datasmith/CADKernel/Public/CADKernel/Core/Session.h
+++ b/Engine/Source/Runtime/Datasmith/CADKernel/Public/CADKernel/Core/Session.h
@@ -21,6 +21,9 @@ namespace CADKernel
 		TSharedPtr<FModel> Model;
 
 	public:
+#ifdef CADKERNEL_DEV
+		static TSharedPtr<FSession> Session;
+#endif
 
 		FSession(double InGeometricTolerance)
 			: GeometricTolerance(InGeometricTolerance)
diff --git a/Engine/Source/Runtime/Datasmith/CADKernel/Public/CADKernel/Geo/Curves/BezierCurve.h b/Engine/Source/Runtime/Datasmith/CADKernel/Public/CADKernel/Geo/Curves/BezierCurve.h
--- a/Engine/Source/Runtime/Datasmith/CADKernel/Public/CADKernel/Geo/Curves/BezierCurve.h
+++ b/Engine/Source/Runtime/Datasmith/CADKernel/Public/CADKernel/Geo/Curves/BezierCurve.h
@@ -12,14 +12,14 @@ namespace CADKernel
 	protected:
 		TArray<FPoint> Poles;
 
-		FBezierCurve(const double InTolerance, const TArray<FPoint>& InPoles)
-			: FCurve(InTolerance)
+		FBezierCurve(const TArray<FPoint>& InPoles)
+			: FCurve(3)
 			, Poles{ InPoles }
 		{
 		}
 
 		FBezierCurve(FCADKernelArchive& Archive)
-			: FCurve()
+			: FCurve(3)
 		{
 			Serialize(Archive);
 		}
@@ -29,7 +29,7 @@ namespace CADKernel
 		virtual void Serialize(FCADKernelArchive& Ar) override
 		{
 			FCurve::Serialize(Ar);
-			Ar << Poles;
+			Ar.Serialize(Poles);
 		}
 
 #ifdef CADKERNEL_DEV
diff --git a/Engine/Source/Runtime/Datasmith/CADKernel/Public/CADKernel/Geo/Curves/BoundedCurve.h b/Engine/Source/Runtime/Datasmith/CADKernel/Public/CADKernel/Geo/Curves/BoundedCurve.h
--- a/Engine/Source/Runtime/Datasmith/CADKernel/Public/CADKernel/Geo/Curves/BoundedCurve.h
+++ b/Engine/Source/Runtime/Datasmith/CADKernel/Public/CADKernel/Geo/Curves/BoundedCurve.h
@@ -14,8 +14,8 @@ namespace CADKernel
 
 		TSharedPtr<FCurve> Curve;
 
-		FBoundedCurve(const double InTolerance, TSharedRef<FCurve> InCurve, const FLinearBoundary& InBoundary, int8 InDimension)
-			: FCurve(InTolerance, InBoundary, InDimension)
+		FBoundedCurve(TSharedRef<FCurve> InCurve, const FLinearBoundary& InBoundary, int8 InDimension)
+			: FCurve(InBoundary, InDimension)
 			, Curve(InCurve)
 		{
 		}
diff --git a/Engine/Source/Runtime/Datasmith/CADKernel/Public/CADKernel/Geo/Curves/CompositeCurve.h b/Engine/Source/Runtime/Datasmith/CADKernel/Public/CADKernel/Geo/Curves/CompositeCurve.h
--- a/Engine/Source/Runtime/Datasmith/CADKernel/Public/CADKernel/Geo/Curves/CompositeCurve.h
+++ b/Engine/Source/Runtime/Datasmith/CADKernel/Public/CADKernel/Geo/Curves/CompositeCurve.h
@@ -35,7 +35,7 @@ namespace CADKernel
 		TArray<FOrientedCurve> Curves;
 		TArray<double> Coordinates;
 
-		FCompositeCurve(const double InTolerance, const TArray<TSharedPtr<FCurve>>& Curves, bool bDoInversions = false);
+		FCompositeCurve(const TArray<TSharedPtr<FCurve>>& Curves, bool bDoInversions = false);
 
 		FCompositeCurve(FCADKernelArchive& Archive)
 			: FCurve()
@@ -49,7 +49,7 @@ namespace CADKernel
 		{
 			FCurve::Serialize(Ar);
 			SerializeIdents(Ar, (TArray<TOrientedEntity<FEntity>>&) Curves);
-			Ar << Coordinates;
+			Ar.Serialize(Coordinates);
 		}
 
 		virtual void SpawnIdent(FDatabase& Database) override
diff --git a/Engine/Source/Runtime/Datasmith/CADKernel/Public/CADKernel/Geo/Curves/Curve.h b/Engine/Source/Runtime/Datasmith/CADKernel/Public/CADKernel/Geo/Curves/Curve.h
--- a/Engine/Source/Runtime/Datasmith/CADKernel/Public/CADKernel/Geo/Curves/Curve.h
+++ b/Engine/Source/Runtime/Datasmith/CADKernel/Public/CADKernel/Geo/Curves/Curve.h
@@ -24,33 +24,25 @@ namespace CADKernel
 
 	protected:
 
-		/** Geometric Tolerance of the curve in its space i.e. Tolerance 3D for 3d curve, Tolerance 2D for 2d curve */
-		double Tolerance;
-
-		mutable TCache<double> Tolerance1D;
 		mutable TCache<double> GlobalLength;
 		FLinearBoundary Boundary;
 
 		int8 Dimension;
 
-		FCurve(const double InTolerance, int8 InDimension = 3)
+		FCurve(int8 InDimension = 3)
 			: FEntityGeom()
-			, Tolerance(InTolerance)
 			, Boundary()
 			, Dimension(InDimension)
 		{
 		}
 
-		FCurve(double InTolerance, const FLinearBoundary& InBounds, int8 InDimension = 3)
+		FCurve(const FLinearBoundary& InBounds, int8 InDimension = 3)
 			: FEntityGeom()
-			, Tolerance(InTolerance)
 			, Boundary(InBounds)
 			, Dimension(InDimension)
 		{
 		}
 
-		FCurve() = default;
-
 	public:
 
 		virtual void Serialize(FCADKernelArchive& Ar) override
@@ -62,7 +54,6 @@ namespace CADKernel
 				Ar << CurveType;
 			}
 			FEntityGeom::Serialize(Ar);
-			Ar << Tolerance;
 			Ar << Dimension;
 			Ar << Boundary;
 		}
@@ -105,12 +96,7 @@ namespace CADKernel
 
 		virtual TSharedPtr<FEntityGeom> ApplyMatrix(const FMatrixH& InMatrix) const override = 0;
 
-		double GetLength() const;
-
-		/**
-		 * Return the tolerance in the parametric space of the curve i.e Distance(Curve(U), Curve(U + Tolerance1D)) = Tolerance
-		 */
-		double GetParametricTolerance() const;
+		double GetLength(double Tolerance) const;
 
 		/**
 		 * Evaluate exact 3D point of the curve at the input Coordinate
@@ -196,12 +182,12 @@ namespace CADKernel
 		 * Generate a pre-sampling of the curve saved in OutCoordinates.
 		 * This sampling is light enough to allow a fast computation, precise enough to compute accurately meshing criteria
 		 */
-		void Presample(TArray<double>& OutSampling) const
+		void Presample(TArray<double>& OutSampling, double Tolerance) const
 		{
-			Presample(Boundary, OutSampling);
+			Presample(Boundary, Tolerance, OutSampling);
 		}
 
-		virtual void Presample(const FLinearBoundary& InBoundary, TArray<double>& OutSampling) const;
+		virtual void Presample(const FLinearBoundary& InBoundary, double Tolerance, TArray<double>& OutSampling) const;
 
 		/**
 		 * Make a new curve based on the new bounds.
@@ -226,16 +212,14 @@ namespace CADKernel
 		virtual void ExtendTo(const FPoint2D& DesiredPosition)
 		{
 			ensure(Dimension == 2);
-
-			NOT_IMPLEMENTED;
-			ensureCADKernel(false);
+			FPoint Point = DesiredPosition;
+			ExtendTo(Point);
 		}
 
 	protected:
 
-		virtual double ComputeLength(const FLinearBoundary& InBoundary) const;
-		//virtual double ComputeLength2D(const FLinearBoundary& InBoundary) const;
-		//virtual double ComputeLength3D(const FLinearBoundary& InBoundary) const;
+		virtual double ComputeLength(const FLinearBoundary& InBoundary, double Tolerance) const;
+		virtual double ComputeLength2D(const FLinearBoundary& InBoundary, double Tolerance) const;
 	};
 } // namespace CADKernel
 
diff --git a/Engine/Source/Runtime/Datasmith/CADKernel/Public/CADKernel/Geo/Curves/EllipseCurve.h b/Engine/Source/Runtime/Datasmith/CADKernel/Public/CADKernel/Geo/Curves/EllipseCurve.h
--- a/Engine/Source/Runtime/Datasmith/CADKernel/Public/CADKernel/Geo/Curves/EllipseCurve.h
+++ b/Engine/Source/Runtime/Datasmith/CADKernel/Public/CADKernel/Geo/Curves/EllipseCurve.h
@@ -18,8 +18,8 @@ namespace CADKernel
 		double   RadiusV;
 		FMatrixH InverseMatrix;
 
-		FEllipseCurve(const double InTolerance, const FMatrixH& InMatrix, double InRadiusU, double InRadiusV, int8 InDimension = 3)
-			: FCurve(InTolerance, FLinearBoundary(0, PI * 2.), InDimension)
+		FEllipseCurve(const FMatrixH& InMatrix, double InRadiusU, double InRadiusV, int8 InDimension = 3)
+			: FCurve(FLinearBoundary(0, PI * 2.), InDimension)
 			, Matrix(InMatrix)
 			, RadiusU(InRadiusU)
 			, RadiusV(InRadiusV)
@@ -27,8 +27,8 @@ namespace CADKernel
 		{
 		}
 
-		FEllipseCurve(const double InTolerance, const FMatrixH& InMatrix, double InRadiusU, double InRadiusV, const FLinearBoundary& InBounds, int8 InDimension = 3)
-			: FCurve(InTolerance, InBounds, InDimension)
+		FEllipseCurve(const FMatrixH& InMatrix, double InRadiusU, double InRadiusV, const FLinearBoundary& InBounds, int8 InDimension = 3)
+			: FCurve(InBounds, InDimension)
 			, Matrix(InMatrix)
 			, RadiusU(InRadiusU)
 			, RadiusV(InRadiusV)
diff --git a/Engine/Source/Runtime/Datasmith/CADKernel/Public/CADKernel/Geo/Curves/HyperbolaCurve.h b/Engine/Source/Runtime/Datasmith/CADKernel/Public/CADKernel/Geo/Curves/HyperbolaCurve.h
--- a/Engine/Source/Runtime/Datasmith/CADKernel/Public/CADKernel/Geo/Curves/HyperbolaCurve.h
+++ b/Engine/Source/Runtime/Datasmith/CADKernel/Public/CADKernel/Geo/Curves/HyperbolaCurve.h
@@ -17,8 +17,8 @@ namespace CADKernel
 		double SemiMajorAxis;
 		double SemiImaginaryAxis;
 
-		FHyperbolaCurve(const double InTolerance, const FMatrixH& InMatrix, double InSemiAxis, double InSemiImagAxis, const FLinearBoundary& InBounds, int8 InDimension = 3)
-			: FCurve(InTolerance, InBounds, InDimension)
+		FHyperbolaCurve(const FMatrixH& InMatrix, double InSemiAxis, double InSemiImagAxis, const FLinearBoundary& InBounds, int8 InDimension = 3)
+			: FCurve(InBounds, InDimension)
 			, Matrix(InMatrix)
 			, SemiMajorAxis(InSemiAxis)
 			, SemiImaginaryAxis(InSemiImagAxis)
diff --git a/Engine/Source/Runtime/Datasmith/CADKernel/Public/CADKernel/Geo/Curves/NURBSCurve.h b/Engine/Source/Runtime/Datasmith/CADKernel/Public/CADKernel/Geo/Curves/NURBSCurve.h
--- a/Engine/Source/Runtime/Datasmith/CADKernel/Public/CADKernel/Geo/Curves/NURBSCurve.h
+++ b/Engine/Source/Runtime/Datasmith/CADKernel/Public/CADKernel/Geo/Curves/NURBSCurve.h
@@ -28,9 +28,10 @@ namespace CADKernel
 		TArray<double> HomogeneousPoles;
 		int32 PoleDimension;
 
-		FNURBSCurve(const double InTolerance, int32 InDegree, const TArray<double>& InNodalVector, const TArray<FPoint>& InPoles, int8 InDimension = 3);
-		FNURBSCurve(const double InTolerance, int32 InDegree, const TArray<double>& InNodalVector, const TArray<FPoint>& InPoles, const TArray<double>& InWeights, int8 InDimension = 3);
-		
+		FNURBSCurve(int32 InDegree, const TArray<double>& InNodalVector, const TArray<FPoint>& InPoles, int8 InDimension = 3);
+		FNURBSCurve(int32 InDegree, const TArray<double>& InNodalVector, const TArray<FPoint>& InPoles, const TArray<double>& InWeights, int8 InDimension = 3);
+		FNURBSCurve(const TSharedRef<FNURBSCurve>& Nurbs);
+
 		FNURBSCurve(FCADKernelArchive& Archive)
 			: FCurve()
 		{
@@ -43,9 +44,10 @@ namespace CADKernel
 		{
 			FCurve::Serialize(Ar);
 			Ar << Degree;
-			Ar << NodalVector;
-			Ar << Weights;
-			Ar << Poles;
+			Ar.Serialize(NodalVector);
+			Ar.Serialize(Weights);
+			Ar.Serialize(Poles);
+			Ar << bIsRational;
 
 			if (Ar.IsLoading())
 			{
@@ -114,7 +116,10 @@ namespace CADKernel
 			BSpline::FindNotDerivableParameters(*this, DerivativeOrder, InBoundary, OutNotDerivableCoordinates);
 		}
 
-		virtual void ExtendTo(const FPoint& Point) override;
+		virtual void ExtendTo(const FPoint& DesiredPosition) override;
+
+		void Invert();
+		void SetStartNodalCoordinate(double NewStartBoundary);
 
 	private:
 
diff --git a/Engine/Source/Runtime/Datasmith/CADKernel/Public/CADKernel/Geo/Curves/ParabolaCurve.h b/Engine/Source/Runtime/Datasmith/CADKernel/Public/CADKernel/Geo/Curves/ParabolaCurve.h
--- a/Engine/Source/Runtime/Datasmith/CADKernel/Public/CADKernel/Geo/Curves/ParabolaCurve.h
+++ b/Engine/Source/Runtime/Datasmith/CADKernel/Public/CADKernel/Geo/Curves/ParabolaCurve.h
@@ -15,8 +15,8 @@ namespace CADKernel
 		FMatrixH Matrix;
 		double FocalDistance;
 
-		FParabolaCurve(const double InTolerance, const FMatrixH& InMatrix, double InFocalDistance, const FLinearBoundary& InBounds, int8 InDimension = 3)
-			: FCurve(InTolerance, InBounds, InDimension)
+		FParabolaCurve(const FMatrixH& InMatrix, double InFocalDistance, const FLinearBoundary& InBounds, int8 InDimension = 3)
+			: FCurve(InBounds, InDimension)
 			, Matrix(InMatrix)
 			, FocalDistance(InFocalDistance)
 		{
diff --git a/Engine/Source/Runtime/Datasmith/CADKernel/Public/CADKernel/Geo/Curves/PolylineCurve.h b/Engine/Source/Runtime/Datasmith/CADKernel/Public/CADKernel/Geo/Curves/PolylineCurve.h
--- a/Engine/Source/Runtime/Datasmith/CADKernel/Public/CADKernel/Geo/Curves/PolylineCurve.h
+++ b/Engine/Source/Runtime/Datasmith/CADKernel/Public/CADKernel/Geo/Curves/PolylineCurve.h
@@ -20,8 +20,8 @@ namespace CADKernel
 		
 		TPolylineApproximator<PointType> Approximator;
 
-		TPolylineCurve(const double InTolerance, const TArray<PointType>& InPoints, const TArray<double>& InCoordinates, int8 InDimension)
-			: FCurve(InTolerance, InDimension)
+		TPolylineCurve(const TArray<PointType>& InPoints, const TArray<double>& InCoordinates, int8 InDimension)
+			: FCurve(InDimension)
 			, Approximator(this->Coordinates, this->Points)
 		{
 			this->Coordinates = InCoordinates;
@@ -30,8 +30,8 @@ namespace CADKernel
 			Boundary.Set(this->Coordinates.HeapTop(), this->Coordinates.Last());
 		}
 
-		TPolylineCurve(const double InTolerance, const TArray<PointType>& InPoints, int8 InDimension)
-			: FCurve(InTolerance, InDimension)
+		TPolylineCurve(const TArray<PointType>& InPoints, int8 InDimension)
+			: FCurve(InDimension)
 			, TPolyline<PointType>(InPoints)
 			, Approximator(this->Coordinates, this->Points)
 		{
@@ -116,7 +116,7 @@ namespace CADKernel
 				NewPoints.Emplace(InMatrix.Multiply(Point));
 			}
 
-			return FEntity::MakeShared<PolylineType>(Tolerance, NewPoints, this->Coordinates);
+			return FEntity::MakeShared<PolylineType>(NewPoints, this->Coordinates);
 		}
 
 		virtual void ExtendTo(const FPoint& DesiredPoint) override
@@ -130,13 +130,13 @@ namespace CADKernel
 		friend class FEntity;
 
 	protected:
-		FPolylineCurve(const double InTolerance, const TArray<FPoint>& InPoints, const TArray<double>& InCoordinates)
-			: TPolylineCurve<FPoint, FCurvePoint>(InTolerance, InPoints, InCoordinates, 3)
+		FPolylineCurve(const TArray<FPoint>& InPoints, const TArray<double>& InCoordinates)
+			: TPolylineCurve<FPoint, FCurvePoint>(InPoints, InCoordinates, 3)
 		{
 		}
 
-		FPolylineCurve(const double InTolerance, const TArray<FPoint>& InPoints)
-			: TPolylineCurve<FPoint, FCurvePoint>(InTolerance, InPoints, 3)
+		FPolylineCurve(const TArray<FPoint>& InPoints)
+			: TPolylineCurve<FPoint, FCurvePoint>(InPoints, 3)
 		{
 		}
 
@@ -173,13 +173,13 @@ namespace CADKernel
 		friend class FEntity;
 
 	protected:
-		FPolyline2DCurve(const double InTolerance, const TArray<FPoint2D>& InPoints, const TArray<double>& InCoordinates)
-			: TPolylineCurve<FPoint2D, FCurvePoint2D>(InTolerance, InPoints, InCoordinates, 3)
+		FPolyline2DCurve(const TArray<FPoint2D>& InPoints, const TArray<double>& InCoordinates)
+			: TPolylineCurve<FPoint2D, FCurvePoint2D>(InPoints, InCoordinates, 3)
 		{
 		}
 
-		FPolyline2DCurve(const double InTolerance, const TArray<FPoint2D>& InPoints)
-			: TPolylineCurve<FPoint2D, FCurvePoint2D>(InTolerance, InPoints, 3)
+		FPolyline2DCurve(const TArray<FPoint2D>& InPoints)
+			: TPolylineCurve<FPoint2D, FCurvePoint2D>(InPoints, 3)
 		{
 		}
 
diff --git a/Engine/Source/Runtime/Datasmith/CADKernel/Public/CADKernel/Geo/Curves/RestrictionCurve.h b/Engine/Source/Runtime/Datasmith/CADKernel/Public/CADKernel/Geo/Curves/RestrictionCurve.h
--- a/Engine/Source/Runtime/Datasmith/CADKernel/Public/CADKernel/Geo/Curves/RestrictionCurve.h
+++ b/Engine/Source/Runtime/Datasmith/CADKernel/Public/CADKernel/Geo/Curves/RestrictionCurve.h
@@ -26,17 +26,20 @@ namespace CADKernel
 	protected:
 
 		FSurfacicPolyline Polyline;
+		double MinLinearTolerance;
 
-		FRestrictionCurve(const double InTolerance, TSharedRef<FSurface> InCarrierSurface, TSharedRef<FCurve> InCurve2D)
-			: FSurfacicCurve(InTolerance, InCurve2D, InCarrierSurface)
-			, Polyline(InCarrierSurface, InCurve2D, InTolerance)
+		FRestrictionCurve(TSharedRef<FSurface> InCarrierSurface, TSharedRef<FCurve> InCurve2D)
+			: FSurfacicCurve(InCurve2D, InCarrierSurface)
+			, Polyline(InCarrierSurface, InCurve2D)
 		{
+			MinLinearTolerance = Boundary.ComputeMinimalTolerance();
 		}
 
 		FRestrictionCurve(FCADKernelArchive& Archive)
 			: FSurfacicCurve()
 		{
 			Serialize(Archive);
+			MinLinearTolerance = Boundary.ComputeMinimalTolerance();
 		}
 
 	public:
@@ -56,14 +59,9 @@ namespace CADKernel
 			return ECurve::Restriction;
 		}
 
-		const TSharedPtr<FCurve>& Get2DCurve() const
+		const TSharedRef<FCurve> Get2DCurve() const
 		{
-			return Curve2D;
-		}
-
-		const TSharedPtr<FSurface>& GetCarrierSurface() const
-		{
-			return CarrierSurface;
+			return Curve2D.ToSharedRef();
 		}
 
 		TSharedPtr<FEntityGeom> ApplyMatrix(const FMatrixH& InMatrix) const override
@@ -90,6 +88,16 @@ namespace CADKernel
 			return Polyline.Approximate3DPoint(InCoordinate);
 		}
 
+		FPoint GetTangentAt(double InCoordinate) const
+		{
+			return Polyline.GetTangentAt(InCoordinate);
+		}
+
+		FPoint2D GetTangent2DAt(double InCoordinate) const
+		{
+			return Polyline.GetTangent2DAt(InCoordinate);
+		}
+
 		/**
 		 * Fast computation of the point in the parametric space of the carrier surface.
 		 * This approximation is based on FSurfacicPolyline::Polyline
@@ -116,13 +124,14 @@ namespace CADKernel
 			Polyline.ApproximatePolyline(OutPolyline);
 		}
 
-
-		double GetCoordinateOfProjectedPoint(const FLinearBoundary& InBoundary, const FPoint& PointOnEdge, FPoint& ProjectedPoint) const
+		template<class PointType>
+		double GetCoordinateOfProjectedPoint(const FLinearBoundary& InBoundary, const PointType& PointOnEdge, PointType& ProjectedPoint) const
 		{
 			return Polyline.GetCoordinateOfProjectedPoint(InBoundary, PointOnEdge, ProjectedPoint);
 		}
 
-		void ProjectPoints(const FLinearBoundary& InBoundary, const TArray<FPoint>& InPointsToProject, TArray<double>& ProjectedPointCoordinates, TArray<FPoint>& ProjectedPoints) const
+		template<class PointType>
+		void ProjectPoints(const FLinearBoundary& InBoundary, const TArray<PointType>& InPointsToProject, TArray<double>& ProjectedPointCoordinates, TArray<PointType>& ProjectedPoints) const
 		{
 			Polyline.ProjectPoints(InBoundary, InPointsToProject, ProjectedPointCoordinates, ProjectedPoints);
 		}
@@ -142,10 +151,41 @@ namespace CADKernel
 		 * 
 		 * A curve can be degenerated in 3D and not in 2D in the case of locally degenerated carrier surface.
 		 */
-		void CheckIfDegenerated(const FLinearBoundary& Boudary, bool& bDegeneration2D, bool& bDegeneration3D, double& Length3D) const
+		void CheckIfDegenerated(const FLinearBoundary& InBoundary, bool& bDegeneration2D, bool& bDegeneration3D, double& Length3D) const
 		{
-			double Tolerance2D = GetCarrierSurface()->Get2DTolerance();
-			Polyline.CheckIfDegenerated(Tolerance, Tolerance2D, Boudary, bDegeneration2D, bDegeneration3D, Length3D);
+			double Tolerance = GetCarrierSurface()->Get3DTolerance();
+			const FSurfacicTolerance& Tolerances2D = GetCarrierSurface()->GetIsoTolerances();
+			Polyline.CheckIfDegenerated(Tolerance, Tolerances2D, InBoundary, bDegeneration2D, bDegeneration3D, Length3D);
+		}
+
+		void GetExtremities(const FLinearBoundary& InBoundary, FSurfacicCurveExtremity& Extremities) const
+		{
+			double Tolerance = GetCarrierSurface()->Get3DTolerance();
+			const FSurfacicTolerance& Tolerances2D = GetCarrierSurface()->GetIsoTolerances();
+			Polyline.GetExtremities(InBoundary, Tolerance, Tolerances2D, Extremities);
+		}
+
+		double GetToleranceAt(const double InCoordinate)
+		{
+			FDichotomyFinder Finder(Polyline.GetCoordinates());
+			const int32 Index = Finder.Find(InCoordinate);
+			const double Tolerance = GetCarrierSurface()->Get3DTolerance();
+
+			return Polyline.ComputeLinearToleranceAt(Tolerance, MinLinearTolerance, Index);
+		}
+
+		FPoint2D GetExtremityTolerances(const FLinearBoundary& InBoundary)
+		{
+			FDichotomyFinder Finder(Polyline.GetCoordinates());
+			const int32 StartIndex = Finder.Find(InBoundary.Min);
+			const int32 EndIndex = Finder.Find(InBoundary.Max);
+
+			const double Tolerance = GetCarrierSurface()->Get3DTolerance();
+
+			FPoint2D ExtremityTolerance;
+			ExtremityTolerance[0] = Polyline.ComputeLinearToleranceAt(Tolerance, MinLinearTolerance, StartIndex);
+			ExtremityTolerance[1] = Polyline.ComputeLinearToleranceAt(Tolerance, MinLinearTolerance, EndIndex);
+			return ExtremityTolerance;
 		}
 
 		/**
@@ -160,9 +200,9 @@ namespace CADKernel
 		 * Get the sub polyline bounded by the input InBoundary in the orientation of the input InOrientation and append it to the output OutPoints
 		 */
 		template<class PointType>
-		void GetDiscretizationPoints(const FLinearBoundary& Boundary, EOrientation Orientation, TArray<PointType>& OutPoints) const
+		void GetDiscretizationPoints(const FLinearBoundary& InBoundary, EOrientation Orientation, TArray<PointType>& OutPoints) const
 		{
-			Polyline.GetSubPolyline(Boundary, Orientation, OutPoints);
+			Polyline.GetSubPolyline(InBoundary, Orientation, OutPoints);
 		}
 
 		/**
@@ -173,12 +213,12 @@ namespace CADKernel
 			Polyline.Sample(InBoundary, DesiredSegmentLength, OutCoordinates);
 		}
 
-		virtual double ComputeLength(const FLinearBoundary& InBoundary) const
+		double ApproximateLength(const FLinearBoundary& InBoundary) const
 		{
 			return Polyline.GetLength(InBoundary);
 		}
 
-		void ExtendTo2DPoint(const FPoint2D& Point);
+		void ExtendTo(const FPoint2D& Point) override;
 	};
 
 } // namespace CADKernel
diff --git a/Engine/Source/Runtime/Datasmith/CADKernel/Public/CADKernel/Geo/Curves/SegmentCurve.h b/Engine/Source/Runtime/Datasmith/CADKernel/Public/CADKernel/Geo/Curves/SegmentCurve.h
--- a/Engine/Source/Runtime/Datasmith/CADKernel/Public/CADKernel/Geo/Curves/SegmentCurve.h
+++ b/Engine/Source/Runtime/Datasmith/CADKernel/Public/CADKernel/Geo/Curves/SegmentCurve.h
@@ -15,8 +15,8 @@ namespace CADKernel
 		FPoint StartPoint;
 		FPoint EndPoint;
 
-		FSegmentCurve(const double InTolerance, const FPoint& InStartPoint, const FPoint& InEndPoint, int8 InDimension = 3)
-			: FCurve(InTolerance, InDimension)
+		FSegmentCurve(const FPoint& InStartPoint, const FPoint& InEndPoint, int8 InDimension = 3)
+			: FCurve(InDimension)
 			, StartPoint(InStartPoint)
 			, EndPoint(InEndPoint)
 		{
diff --git a/Engine/Source/Runtime/Datasmith/CADKernel/Public/CADKernel/Geo/Curves/SurfacicCurve.h b/Engine/Source/Runtime/Datasmith/CADKernel/Public/CADKernel/Geo/Curves/SurfacicCurve.h
--- a/Engine/Source/Runtime/Datasmith/CADKernel/Public/CADKernel/Geo/Curves/SurfacicCurve.h
+++ b/Engine/Source/Runtime/Datasmith/CADKernel/Public/CADKernel/Geo/Curves/SurfacicCurve.h
@@ -17,8 +17,8 @@ namespace CADKernel
 		TSharedPtr<FCurve> Curve2D;
 		TSharedPtr<FSurface> CarrierSurface;
 
-		FSurfacicCurve(const double InTolerance, TSharedRef<FCurve> InCurve2D, TSharedRef<FSurface> InSurface)
-			: FCurve(InTolerance, InCurve2D->GetBoundary())
+		FSurfacicCurve(TSharedRef<FCurve> InCurve2D, TSharedRef<FSurface> InSurface)
+			: FCurve(InCurve2D->GetBoundary())
 			, Curve2D(InCurve2D)
 			, CarrierSurface(InSurface)
 		{
@@ -82,6 +82,11 @@ namespace CADKernel
 			return Curve2D;
 		}
 
+		const TSharedRef<FSurface> GetCarrierSurface() const
+		{
+			return CarrierSurface.ToSharedRef();
+		}
+
 		void Set2DCurve(TSharedPtr<FCurve>& NewCurve2D)
 		{
 			ensureCADKernel(NewCurve2D.IsValid());
@@ -114,8 +119,7 @@ namespace CADKernel
 
 		virtual void EvaluatePoints(const TArray<double>& Coordinates, TArray<FCurvePoint>& OutPoints, int32 DerivativeOrder = 0) const override;
 
-		void EvaluateSurfacicPolyline(FSurfacicPolyline& OutPolyline, bool bInWithNormals = false) const;
-		void EvaluateSurfacicPolylineWithNormalAndTangent(FSurfacicPolyline& OutPolyline) const;
+		void EvaluateSurfacicPolyline(FSurfacicPolyline& OutPolyline) const;
 
 		virtual void Evaluate2DPoints(const TArray<double>& Coordinates, TArray<FPoint2D>& OutPoints) const override
 		{
@@ -128,7 +132,8 @@ namespace CADKernel
 		}
 
 		virtual void FindNotDerivableCoordinates(const FLinearBoundary& InBoundary, int32 DerivativeOrder, TArray<double>& OutNotDerivableCoordinates) const override;
-
+private:
+		void EvaluateSurfacicPolylineWithNormalAndTangent(FSurfacicPolyline& OutPolyline) const;
 	};
 
 } // namespace CADKernel
diff --git a/Engine/Source/Runtime/Datasmith/CADKernel/Public/CADKernel/Geo/GeoPoint.h b/Engine/Source/Runtime/Datasmith/CADKernel/Public/CADKernel/Geo/GeoPoint.h
--- a/Engine/Source/Runtime/Datasmith/CADKernel/Public/CADKernel/Geo/GeoPoint.h
+++ b/Engine/Source/Runtime/Datasmith/CADKernel/Public/CADKernel/Geo/GeoPoint.h
@@ -37,6 +37,15 @@ namespace CADKernel
 		FPoint Tangent = FPoint::ZeroPoint;
 	};
 
+	struct FSurfacicCurvePointWithTolerance
+	{
+		FPoint2D Point2D = FPoint2D::ZeroPoint;
+		FPoint Point = FPoint::ZeroPoint;
+		FSurfacicTolerance Tolerance = FPoint2D::FarawayPoint;
+	};
+
+	typedef FSurfacicCurvePointWithTolerance FSurfacicCurveExtremity[2];
+
 	struct FCurvePoint
 	{
 		int32 DerivativeOrder = -1;
