CADKernel Source Code v0.02

diff --git a/Engine/Source/Runtime/Datasmith/CADKernel/Public/CADKernel/Geo/Sampling/PolylineTools.h b/Engine/Source/Runtime/Datasmith/CADKernel/Public/CADKernel/Geo/Sampling/PolylineTools.h
--- a/Engine/Source/Runtime/Datasmith/CADKernel/Public/CADKernel/Geo/Sampling/PolylineTools.h
+++ b/Engine/Source/Runtime/Datasmith/CADKernel/Public/CADKernel/Geo/Sampling/PolylineTools.h
@@ -2,6 +2,7 @@
 #pragma once
 
 #include "CADKernel/Core/Types.h"
+#include "CADKernel/Math/Aabb.h"
 #include "CADKernel/Math/Boundary.h"
 #include "CADKernel/Math/Geometry.h"
 #include "CADKernel/Math/Point.h"
@@ -58,14 +59,14 @@ namespace CADKernel
 			if (DistanceStart > DistanceEnd)
 			{
 				PointType Factor;
-				for (int32 Index = 0; Index < PointType::Dimension(); ++Index)
+				for (int32 Index = 0; Index < PointType::Dimension; ++Index)
 				{
 					Factor[Index] = FMath::Abs(Polyline.Last()[Index] - Polyline[0][Index]) > SMALL_NUMBER_SQUARE ? (DesiredEnd[Index] - Polyline[0][Index]) / (Polyline.Last()[Index] - Polyline[0][Index]) : 1.;
 				}
 
 				Algo::ForEach(Polyline, [&, Factor](PointType& Pole)
 					{
-						for (int32 Index = 0; Index < PointType::Dimension(); ++Index)
+						for (int32 Index = 0; Index < PointType::Dimension; ++Index)
 						{
 							Pole[Index] = Polyline[0][Index] + (Pole[Index] - Polyline[0][Index]) * Factor[Index];
 						}
@@ -75,7 +76,7 @@ namespace CADKernel
 			else
 			{
 				PointType Factor;
-				for (int32 Index = 0; Index < PointType::Dimension(); ++Index)
+				for (int32 Index = 0; Index < PointType::Dimension; ++Index)
 				{
 					Factor[Index] = FMath::Abs(Polyline[0][Index] - Polyline.Last()[Index]) > SMALL_NUMBER_SQUARE ? (DesiredEnd[Index] - Polyline.Last()[Index]) / (Polyline[0][Index] - Polyline.Last()[Index]) : 1.;
 				}
@@ -83,15 +84,30 @@ namespace CADKernel
 				Algo::Reverse(Polyline);
 				Algo::ForEach(Polyline, [&, Factor](PointType& Pole)
 					{
-						for (int32 Index = 0; Index < PointType::Dimension(); ++Index)
+						for (int32 Index = 0; Index < PointType::Dimension; ++Index)
 						{
-							Pole[Index] = Polyline.Last()[Index] + (Pole[Index] - Polyline.Last()[Index]) * Factor[Index];
+							Pole[Index] = Polyline[0][Index] + (Pole[Index] - Polyline[0][Index]) * Factor[Index];
 						}
 					}
 				);
+
 				Algo::Reverse(Polyline);
 			}
 		}
+
+		template<class PointType>
+		PointType ComputePoint(const TArray<double>& PolylineCoordinates, const TArray<PointType>& PolylinePoints, const int32 Index, const double PointCoordinate)
+		{
+			double Delta = PolylineCoordinates[Index + 1] - PolylineCoordinates[Index];
+			if (FMath::IsNearlyZero(Delta, (double)KINDA_SMALL_NUMBER))
+			{
+				return PolylinePoints[Index];
+			}
+
+			return PolylinePoints[Index] + (PolylinePoints[Index + 1] - PolylinePoints[Index]) * (PointCoordinate - PolylineCoordinates[Index]) / Delta;
+		};
+
+
 	}
 
 	template<class PointType>
@@ -110,26 +126,26 @@ namespace CADKernel
 
 	protected:
 
-		double ComputeCurvilinearCoordinatesOfPolyline(const FLinearBoundary& InBoundary, TArray<double>& OutCurvilinearCoordinates, int32& StartIndex, int32& EndIndex) const
+		double ComputeCurvilinearCoordinatesOfPolyline(const FLinearBoundary& InBoundary, TArray<double>& OutCurvilinearCoordinates, int32 BoundaryIndices[2]) const
 		{
-			GetStartEndIndex(InBoundary, StartIndex, EndIndex);
+			GetStartEndIndex(InBoundary, BoundaryIndices);
 
-			OutCurvilinearCoordinates.Reserve(EndIndex - StartIndex + 2);
+			OutCurvilinearCoordinates.Reserve(BoundaryIndices[1] - BoundaryIndices[0] + 2);
 
 			double LastSegmentLength;
 			double Length = 0;
-			ensureCADKernel(EndIndex > StartIndex);
+			ensureCADKernel(BoundaryIndices[1] > BoundaryIndices[0]);
 
 			OutCurvilinearCoordinates.Add(0);
-			if (EndIndex > StartIndex + 1)
+			if (BoundaryIndices[1] > BoundaryIndices[0] + 1)
 			{
-				PointType StartPoint = ComputePoint(StartIndex, InBoundary.Min);
-				PointType EndPoint = ComputePoint(EndIndex, InBoundary.Max);
-				Length = StartPoint.Distance(PolylinePoints[StartIndex + 1]);
-				LastSegmentLength = EndPoint.Distance(PolylinePoints[EndIndex]);
+				PointType StartPoint = ComputePoint(BoundaryIndices[0], InBoundary.Min);
+				PointType EndPoint = ComputePoint(BoundaryIndices[1], InBoundary.Max);
+				Length = StartPoint.Distance(PolylinePoints[BoundaryIndices[0] + 1]);
+				LastSegmentLength = EndPoint.Distance(PolylinePoints[BoundaryIndices[1]]);
 
 				OutCurvilinearCoordinates.Add(Length);
-				for (int32 Index = StartIndex + 1; Index < EndIndex; ++Index)
+				for (int32 Index = BoundaryIndices[0] + 1; Index < BoundaryIndices[1]; ++Index)
 				{
 					Length += PolylinePoints[Index].Distance(PolylinePoints[Index + 1]);
 					OutCurvilinearCoordinates.Add(Length);
@@ -139,8 +155,8 @@ namespace CADKernel
 			}
 			else
 			{
-				PointType StartPoint = ComputePoint(StartIndex, InBoundary.Min);
-				PointType EndPoint = ComputePoint(EndIndex, InBoundary.Max);
+				PointType StartPoint = ComputePoint(BoundaryIndices[0], InBoundary.Min);
+				PointType EndPoint = ComputePoint(BoundaryIndices[1], InBoundary.Max);
 				Length = StartPoint.Distance(EndPoint);
 				OutCurvilinearCoordinates.Add(Length);
 			}
@@ -162,7 +178,7 @@ namespace CADKernel
 		 * Project a Set of points on a restricted polyline (StartIndex & EndIndex define the polyline boundary)
 		 * the points are projected on all segments of the polyline, the closest are selected
 		 */
-		double ProjectPointToPolyline(int32 StartIndex, int32 EndIndex, const PointType& InPointToProject, PointType& OutProjectedPoint) const
+		double ProjectPointToPolyline(int32 BoundaryIndices[2], const PointType& InPointToProject, PointType& OutProjectedPoint) const
 		{
 			double MinDistance = HUGE_VAL;
 			double UForMinDistance = 0;
@@ -170,7 +186,7 @@ namespace CADKernel
 			double ParamU = 0.;
 			int32 SegmentIndex = 0;
 
-			for (int32 Index = StartIndex; Index <= EndIndex; ++Index)
+			for (int32 Index = BoundaryIndices[0]; Index <= BoundaryIndices[1]; ++Index)
 			{
 				FPoint ProjectPoint = ProjectPointOnSegment(InPointToProject, PolylinePoints[Index], PolylinePoints[Index + 1], ParamU, true);
 				double SquareDistance = FMath::Square(ProjectPoint[0] - InPointToProject[0]);
@@ -199,11 +215,11 @@ namespace CADKernel
 
 	public:
 
-		void GetStartEndIndex(const FLinearBoundary& InBoundary, int32& StartIndex, int32& EndIndex) const
+		void GetStartEndIndex(const FLinearBoundary& InBoundary, int32 BoundaryIndices[2]) const
 		{
 			FDichotomyFinder Finder(PolylineCoordinates);
-			StartIndex = Finder.Find(InBoundary.Min);
-			EndIndex = Finder.Find(InBoundary.Max);
+			BoundaryIndices[0] = Finder.Find(InBoundary.Min);
+			BoundaryIndices[1] = Finder.Find(InBoundary.Max);
 		}
 
 		/**
@@ -352,11 +368,10 @@ namespace CADKernel
 
 		void SamplePolyline(const FLinearBoundary& InBoundary, const double DesiredSegmentLength, TArray<double>& OutCoordinates) const
 		{
-			int32 StartIndex = 0;
-			int32 EndIndex = 0;
+			int32 BoundaryIndices[2];
 
 			TArray<double> CurvilinearCoordinates;
-			double CurveLength = ComputeCurvilinearCoordinatesOfPolyline(InBoundary, CurvilinearCoordinates, StartIndex, EndIndex);
+			double CurveLength = ComputeCurvilinearCoordinatesOfPolyline(InBoundary, CurvilinearCoordinates, BoundaryIndices);
 
 			int32 SegmentNum = (int32)FMath::Max(CurveLength / DesiredSegmentLength + 0.5, 1.0);
 
@@ -373,7 +388,7 @@ namespace CADKernel
 
 			double CurvilinearLength = SectionLength;
 			double LastCoordinate = InBoundary.Min;
-			for (int32 IndexCurvilinear = 1, IndexCoordinate = StartIndex + 1; IndexCurvilinear < CurvilinearCoordinates.Num(); ++IndexCurvilinear, ++IndexCoordinate)
+			for (int32 IndexCurvilinear = 1, IndexCoordinate = BoundaryIndices[0] + 1; IndexCurvilinear < CurvilinearCoordinates.Num(); ++IndexCurvilinear, ++IndexCoordinate)
 			{
 				while (CurvilinearLength < CurvilinearCoordinates[IndexCurvilinear] + SMALL_NUMBER)
 				{
@@ -396,11 +411,9 @@ namespace CADKernel
 		 */
 		double ProjectPointToPolyline(const FLinearBoundary& InBoundary, const PointType& PointOnEdge, PointType& OutProjectedPoint) const
 		{
-			int32 StartIndex = 0;
-			int32 EndIndex = 0;
-			GetStartEndIndex(InBoundary, StartIndex, EndIndex);
-
-			return ProjectPointToPolyline(StartIndex, EndIndex, PointOnEdge, OutProjectedPoint);
+			int32 BoundaryIndices[2];
+			GetStartEndIndex(InBoundary, BoundaryIndices);
+			return ProjectPointToPolyline(BoundaryIndices, PointOnEdge, OutProjectedPoint);
 		}
 
 		/**
@@ -412,15 +425,13 @@ namespace CADKernel
 			OutProjectedPointCoords.Empty(InPointsToProject.Num());
 			OutProjectedPoints.Empty(InPointsToProject.Num());
 
-			int32 StartIndex = 0;
-			int32 EndIndex = 0;
-
-			GetStartEndIndex(InBoundary, StartIndex, EndIndex);
+			int32 BoundaryIndices[2];
+			GetStartEndIndex(InBoundary, BoundaryIndices);
 
 			for (const PointType& Point : InPointsToProject)
 			{
 				PointType ProjectedPoint;
-				double Coordinate = ProjectPointToPolyline(StartIndex, EndIndex, Point, ProjectedPoint);
+				double Coordinate = ProjectPointToPolyline(BoundaryIndices, Point, ProjectedPoint);
 
 				OutProjectedPointCoords.Emplace(Coordinate);
 				OutProjectedPoints.Emplace(ProjectedPoint);
@@ -484,41 +495,40 @@ namespace CADKernel
 		 */
 		void GetSubPolyline(const FLinearBoundary& InBoundary, const EOrientation Orientation, TArray<PointType>& OutPoints) const
 		{
-			int32 StartIndex = 0;
-			int32 EndIndex = 0;
-			GetStartEndIndex(InBoundary, StartIndex, EndIndex);
+			int32 BoundaryIndices[2];
+			GetStartEndIndex(InBoundary, BoundaryIndices);
 
-			int32 NewSize = OutPoints.Num() + EndIndex - StartIndex + 2;
+			int32 NewSize = OutPoints.Num() + BoundaryIndices[1] - BoundaryIndices[0] + 2;
 			OutPoints.Reserve(NewSize);
 
-			int32 PolylineStartIndex = StartIndex;
-			int32 PolylineEndIndex = EndIndex;
-			if (FMath::IsNearlyEqual(PolylineCoordinates[StartIndex + 1], InBoundary.Min, (double)KINDA_SMALL_NUMBER))
+			int32 PolylineStartIndex = BoundaryIndices[0];
+			int32 PolylineEndIndex = BoundaryIndices[1];
+			if (FMath::IsNearlyEqual(PolylineCoordinates[BoundaryIndices[0] + 1], InBoundary.Min, (double)KINDA_SMALL_NUMBER))
 			{
 				PolylineStartIndex++;
 			}
-			if (FMath::IsNearlyEqual(PolylineCoordinates[EndIndex], InBoundary.Max, (double)KINDA_SMALL_NUMBER))
+			if (FMath::IsNearlyEqual(PolylineCoordinates[BoundaryIndices[1]], InBoundary.Max, (double)KINDA_SMALL_NUMBER))
 			{
 				PolylineEndIndex--;
 			}
 
 			if (Orientation)
 			{
-				OutPoints.Emplace(ComputePoint(StartIndex, InBoundary.Min));
+				OutPoints.Emplace(ComputePoint(BoundaryIndices[0], InBoundary.Min));
 				if(PolylineEndIndex - PolylineStartIndex > 0)
 				{
 					OutPoints.Append(PolylinePoints.GetData() + PolylineStartIndex + 1, PolylineEndIndex - PolylineStartIndex);
 				}
-				OutPoints.Emplace(ComputePoint(EndIndex, InBoundary.Max));
+				OutPoints.Emplace(ComputePoint(BoundaryIndices[1], InBoundary.Max));
 			}
 			else
 			{
-				OutPoints.Emplace(ComputePoint(EndIndex, InBoundary.Max));
+				OutPoints.Emplace(ComputePoint(BoundaryIndices[1], InBoundary.Max));
 				for (int32 Index = PolylineEndIndex; Index > PolylineStartIndex; --Index)
 				{
 					OutPoints.Emplace(PolylinePoints[Index]);
 				}
-				OutPoints.Emplace(ComputePoint(StartIndex, InBoundary.Min));
+				OutPoints.Emplace(ComputePoint(BoundaryIndices[0], InBoundary.Min));
 			}
 		}
 
@@ -564,24 +574,19 @@ namespace CADKernel
 			return Length;
 		}
 
-		double ComputeLengthOfSubPolyline(const FLinearBoundary& InBoundary) const
+		double ComputeLengthOfSubPolyline(const int BoundaryIndex[2], const FLinearBoundary& InBoundary) const
 		{
-			int32 StartIndex = 0;
-			int32 EndIndex = 0;
-
-			GetStartEndIndex(InBoundary, StartIndex, EndIndex);
-
 			double Length = 0;
-			if (EndIndex > StartIndex)
+			if (BoundaryIndex[1] > BoundaryIndex[0])
 			{
-				FPoint StartPoint = ComputePoint(StartIndex, InBoundary.Min);
-				FPoint EndPoint = ComputePoint(EndIndex, InBoundary.Max);
-				Length = StartPoint.Distance(PolylinePoints[StartIndex + 1]);
-				Length += EndPoint.Distance(PolylinePoints[EndIndex]);
+				PointType StartPoint = ComputePoint(BoundaryIndex[0], InBoundary.Min);
+				PointType EndPoint = ComputePoint(BoundaryIndex[1], InBoundary.Max);
+				Length = StartPoint.Distance(PolylinePoints[BoundaryIndex[0] + 1]);
+				Length += EndPoint.Distance(PolylinePoints[BoundaryIndex[1]]);
 
-				if (EndIndex > StartIndex + 1)
+				if (BoundaryIndex[1] > BoundaryIndex[0] + 1)
 				{
-					for (int32 Index = StartIndex + 1; Index < EndIndex; ++Index)
+					for (int32 Index = BoundaryIndex[0] + 1; Index < BoundaryIndex[1]; ++Index)
 					{
 						Length += PolylinePoints[Index].Distance(PolylinePoints[Index + 1]);
 					}
@@ -589,14 +594,36 @@ namespace CADKernel
 			}
 			else
 			{
-				FPoint StartPoint = ComputePoint(StartIndex, InBoundary.Min);
-				FPoint EndPoint = ComputePoint(EndIndex, InBoundary.Max);
+				PointType StartPoint = ComputePoint(BoundaryIndex[0], InBoundary.Min);
+				PointType EndPoint = ComputePoint(BoundaryIndex[1], InBoundary.Max);
 				Length = StartPoint.Distance(EndPoint);
 			}
 			return Length;
 		}
 
+		double ComputeLengthOfSubPolyline(const FLinearBoundary& InBoundary) const
+		{
+			int BoundaryIndices[2];
+			GetStartEndIndex(InBoundary, BoundaryIndices);
+			return ComputeLengthOfSubPolyline(BoundaryIndices, InBoundary);
+		}
 
+		void ComputeBoundingBox(const int BoundaryIndex[2], const FLinearBoundary& InBoundary, TAABB<PointType>& Aabb)
+		{
+			Aabb.Empty();
+			if (BoundaryIndex[1] > BoundaryIndex[0])
+			{
+				for (int32 Index = BoundaryIndex[0] + 1; Index <= BoundaryIndex[1]; ++Index)
+				{
+					Aabb += PolylinePoints[Index];
+				}
+			}
+
+			PointType StartPoint = ComputePoint(BoundaryIndex[0], InBoundary.Min);
+			Aabb += StartPoint;
+			PointType EndPoint = ComputePoint(BoundaryIndex[1], InBoundary.Max);
+			Aabb += EndPoint;
+		}
 	};
 
 } // namespace CADKernel
diff --git a/Engine/Source/Runtime/Datasmith/CADKernel/Public/CADKernel/Geo/Sampling/SurfacicPolyline.h b/Engine/Source/Runtime/Datasmith/CADKernel/Public/CADKernel/Geo/Sampling/SurfacicPolyline.h
--- a/Engine/Source/Runtime/Datasmith/CADKernel/Public/CADKernel/Geo/Sampling/SurfacicPolyline.h
+++ b/Engine/Source/Runtime/Datasmith/CADKernel/Public/CADKernel/Geo/Sampling/SurfacicPolyline.h
@@ -33,6 +33,8 @@ namespace CADKernel
 			Serialize(Archive);
 		}
 
+		FSurfacicPolyline(TSharedRef<FSurface> InCarrierSurface, TSharedRef<FCurve> InCurve2D);
+
 		FSurfacicPolyline(TSharedRef<FSurface> InCarrierSurface, TSharedRef<FCurve> InCurve2D, const double Tolerance);
 
 		FSurfacicPolyline(TSharedRef<FSurface> InCarrierSurface, TSharedRef<FCurve> InCurve2D, const double ChordTolerance, const double ParamTolerance, bool bInWithNormals/* = false*/, bool bWithTangent/* = false*/);
@@ -45,10 +47,10 @@ namespace CADKernel
 
 		void Serialize(FCADKernelArchive& Ar)
 		{
-			Ar << Points3D;
-			Ar << Points2D;
-			Ar << Normals;
-			Ar << Coordinates;
+			Ar.Serialize(Points3D);
+			Ar.Serialize(Points2D);
+			Ar.Serialize(Normals);
+			Ar.Serialize(Coordinates);
 			Ar << bWithNormals;
 			Ar << bWithTangent;
 		}
@@ -57,7 +59,9 @@ namespace CADKernel
 
 		TSharedPtr<FEntityGeom> ApplyMatrix(const FMatrixH&) const;
 
-		void CheckIfDegenerated(const double Tolerance3D, const double Tolerance2D, const FLinearBoundary& Boudary, bool& bDegeneration2D, bool& bDegeneration3D, double& Length3D) const;
+		void CheckIfDegenerated(const double Tolerance3D, const FSurfacicTolerance& Tolerances2D, const FLinearBoundary& Boudary, bool& bDegeneration2D, bool& bDegeneration3D, double& Length3D) const;
+
+		void GetExtremities(const FLinearBoundary& InBoundary, const double Tolerance3D, const FSurfacicTolerance& Tolerances2D, FSurfacicCurveExtremity& Extremities) const;
 
 		FPoint Approximate3DPoint(double InCoordinate) const
 		{
@@ -77,6 +81,48 @@ namespace CADKernel
 			return Approximator.ApproximatePoint(InCoordinate);
 		}
 
+		FPoint GetTangentAt(double InCoordinate) const
+		{
+			FDichotomyFinder Finder(Coordinates);
+			int32 Index = Finder.Find(InCoordinate);
+			return Points3D[Index + 1] - Points3D[Index];
+		}
+
+		FPoint2D GetTangent2DAt(double InCoordinate) const
+		{
+			FDichotomyFinder Finder(Coordinates);
+			int32 Index = Finder.Find(InCoordinate);
+			return Points2D[Index + 1] - Points2D[Index];
+		}
+
+		FSurfacicTolerance ComputeTolerance(const double Tolerance3D, const FSurfacicTolerance& MinToleranceIso, const int32 Index) const
+		{
+			double Distance3D = Points3D[Index].Distance(Points3D[Index + 1]);
+			if (FMath::IsNearlyZero(Distance3D, (double) SMALL_NUMBER))
+			{
+				return FPoint2D::FarawayPoint;
+			}
+			else
+			{
+				FSurfacicTolerance Tolerance2D = Points2D[Index] - Points2D[Index + 1];
+				return Max(Abs(Tolerance2D)*Tolerance3D/Distance3D, MinToleranceIso);
+			}
+		};
+
+		double ComputeLinearToleranceAt(const double Tolerance3D, const double MinLinearTolerance, const int32 Index) const
+		{
+			double Distance3D = Points3D[Index].Distance(Points3D[Index + 1]);
+			if (FMath::IsNearlyZero(Distance3D, (double)SMALL_NUMBER))
+			{
+				return HUGE_VALUE;
+			}
+			else
+			{
+				double LinearDistance = Coordinates[Index+1] - Coordinates[Index];
+				return FMath::Max(LinearDistance / Distance3D * Tolerance3D, MinLinearTolerance);
+			}
+		};
+
 		void Approximate2DPoints(const TArray<double>& InCoordinates, TArray<FPoint2D>& OutPoints) const
 		{
 			TPolylineApproximator<FPoint2D> Approximator(Coordinates, Points2D);
@@ -176,12 +222,24 @@ namespace CADKernel
 			return Approximator3D.ProjectPointToPolyline(Boundary, PointOnEdge, ProjectedPoint);
 		}
 
+		double GetCoordinateOfProjectedPoint(const FLinearBoundary& Boundary, const FPoint2D& PointOnEdge, FPoint2D& ProjectedPoint) const
+		{
+			TPolylineApproximator<FPoint2D> Approximator2D(Coordinates, Points2D);
+			return Approximator2D.ProjectPointToPolyline(Boundary, PointOnEdge, ProjectedPoint);
+		}
+
 		void ProjectPoints(const FLinearBoundary& InBoundary, const TArray<FPoint>& InPointsToProject, TArray<double>& ProjectedPointCoordinates, TArray<FPoint>& ProjectedPoints) const
 		{
 			TPolylineApproximator<FPoint> Approximator3D(Coordinates, Points3D);
 			Approximator3D.ProjectPointsToPolyline(InBoundary, InPointsToProject, ProjectedPointCoordinates, ProjectedPoints);
 		}
 
+		void ProjectPoints(const FLinearBoundary& InBoundary, const TArray<FPoint2D>& InPointsToProject, TArray<double>& ProjectedPointCoordinates, TArray<FPoint2D>& ProjectedPoints) const
+		{
+			TPolylineApproximator<FPoint2D> Approximator(Coordinates, Points2D);
+			Approximator.ProjectPointsToPolyline(InBoundary, InPointsToProject, ProjectedPointCoordinates, ProjectedPoints);
+		}
+
 		/**
 		 * Project each point of a coincidental polyline on the Polyline.
 		 */
diff --git a/Engine/Source/Runtime/Datasmith/CADKernel/Public/CADKernel/Geo/Surfaces/BezierSurface.h b/Engine/Source/Runtime/Datasmith/CADKernel/Public/CADKernel/Geo/Surfaces/BezierSurface.h
--- a/Engine/Source/Runtime/Datasmith/CADKernel/Public/CADKernel/Geo/Surfaces/BezierSurface.h
+++ b/Engine/Source/Runtime/Datasmith/CADKernel/Public/CADKernel/Geo/Surfaces/BezierSurface.h
@@ -30,6 +30,7 @@ namespace CADKernel
 			, Poles(InPoles)
 		{
 			ensureCADKernel((UDegre + 1) * (VDegre + 1) == Poles.Num());
+			SetMinToleranceIso();
 		}
 
 		FBezierSurface(FCADKernelArchive& Archive)
@@ -47,7 +48,7 @@ namespace CADKernel
 			Ar << VDegre;
 			Ar << UPoleNum;
 			Ar << VPoleNum;
-			Ar << Poles;
+			Ar.Serialize(Poles);
 		}
 
 		ESurface GetSurfaceType() const
diff --git a/Engine/Source/Runtime/Datasmith/CADKernel/Public/CADKernel/Geo/Surfaces/ConeSurface.h b/Engine/Source/Runtime/Datasmith/CADKernel/Public/CADKernel/Geo/Surfaces/ConeSurface.h
--- a/Engine/Source/Runtime/Datasmith/CADKernel/Public/CADKernel/Geo/Surfaces/ConeSurface.h
+++ b/Engine/Source/Runtime/Datasmith/CADKernel/Public/CADKernel/Geo/Surfaces/ConeSurface.h
@@ -36,6 +36,7 @@ namespace CADKernel
 			, StartRadius(InStartRadius)
 			, ConeAngle(InConeAngle)
 		{
+			SetMinToleranceIso();
 		}
 
 		FConeSurface(FCADKernelArchive& Archive)
@@ -44,6 +45,26 @@ namespace CADKernel
 			Serialize(Archive);
 		}
 
+		virtual void SetMinToleranceIso() const override
+		{
+			FPoint Origin = Matrix.Multiply(FPoint::ZeroPoint);
+
+			double DeltaVR = tan(ConeAngle);
+			double Radius1 = FMath::Abs(StartRadius + Boundary[EIso::IsoV].Max * DeltaVR);
+			double Radius2 = FMath::Abs(StartRadius + Boundary[EIso::IsoV].Min * DeltaVR);
+			double Radius = FMath::Max(Radius2, Radius1);
+
+			FPoint Point2DU{ 1 , 0, 0 };
+			double ScaleU = ComputeScaleAlongAxis(Point2DU, Matrix, Origin);
+			double ToleranceU = Tolerance3D / Radius;
+			ToleranceU /= ScaleU;
+
+			FPoint Point2DV{ 0, 1, 0 };
+			double ToleranceV = Tolerance3D / ComputeScaleAlongAxis(Point2DV, Matrix, Origin);
+
+			MinToleranceIso.Set(ToleranceU, ToleranceV);
+		}
+
 	public:
 
 		virtual void Serialize(FCADKernelArchive& Ar) override
@@ -79,7 +100,6 @@ namespace CADKernel
 			OutCoordinates[EIso::IsoV].Add(InBoundaries.UVBoundaries[EIso::IsoV].Max);
 		}
 
-
 	};
 }
 
diff --git a/Engine/Source/Runtime/Datasmith/CADKernel/Public/CADKernel/Geo/Surfaces/NURBSSurface.h b/Engine/Source/Runtime/Datasmith/CADKernel/Public/CADKernel/Geo/Surfaces/NURBSSurface.h
--- a/Engine/Source/Runtime/Datasmith/CADKernel/Public/CADKernel/Geo/Surfaces/NURBSSurface.h
+++ b/Engine/Source/Runtime/Datasmith/CADKernel/Public/CADKernel/Geo/Surfaces/NURBSSurface.h
@@ -68,6 +68,7 @@ namespace CADKernel
 			, Poles(InPoles)
 			, bIsRational(true)
 		{
+			SetMinToleranceIso();
 			Finalize();
 		}
 
@@ -87,10 +88,11 @@ namespace CADKernel
 			Ar << PoleVNum;
 			Ar << UDegree;
 			Ar << VDegree;
-			Ar << UNodalVector;
-			Ar << VNodalVector;
-			Ar << Weights;
-			Ar << Poles;
+			Ar.Serialize(UNodalVector);
+			Ar.Serialize(VNodalVector);
+			Ar.Serialize(Weights);
+			Ar.Serialize(Poles);
+			Ar << bIsRational;
 		}
 
 		ESurface GetSurfaceType() const
diff --git a/Engine/Source/Runtime/Datasmith/CADKernel/Public/CADKernel/Geo/Surfaces/OffsetSurface.h b/Engine/Source/Runtime/Datasmith/CADKernel/Public/CADKernel/Geo/Surfaces/OffsetSurface.h
--- a/Engine/Source/Runtime/Datasmith/CADKernel/Public/CADKernel/Geo/Surfaces/OffsetSurface.h
+++ b/Engine/Source/Runtime/Datasmith/CADKernel/Public/CADKernel/Geo/Surfaces/OffsetSurface.h
@@ -20,6 +20,7 @@ namespace CADKernel
 			, BaseSurface(InBaseSurface)
 			, Offset(InOffset)
 		{
+			SetMinToleranceIso();
 		}
 
 		FOffsetSurface(FCADKernelArchive& Archive)
@@ -28,6 +29,11 @@ namespace CADKernel
 			Serialize(Archive);
 		}
 
+		virtual void SetMinToleranceIso() const override
+		{
+			MinToleranceIso = BaseSurface->GetIsoTolerances();
+		}
+
 	public:
 
 		virtual void Serialize(FCADKernelArchive& Ar) override
@@ -86,7 +92,6 @@ namespace CADKernel
 		virtual void EvaluatePointGrid(const FCoordinateGrid& Coordinates, FSurfacicSampling& OutPoints, bool bComputeNormals = false) const override;
 
 		virtual void Presample(const FSurfacicBoundary& InBoundaries, FCoordinateGrid& OutCoordinates) override;
-
 	};
 
 } // namespace CADKernel
diff --git a/Engine/Source/Runtime/Datasmith/CADKernel/Public/CADKernel/Geo/Surfaces/PlaneSurface.h b/Engine/Source/Runtime/Datasmith/CADKernel/Public/CADKernel/Geo/Surfaces/PlaneSurface.h
--- a/Engine/Source/Runtime/Datasmith/CADKernel/Public/CADKernel/Geo/Surfaces/PlaneSurface.h
+++ b/Engine/Source/Runtime/Datasmith/CADKernel/Public/CADKernel/Geo/Surfaces/PlaneSurface.h
@@ -19,15 +19,15 @@ namespace CADKernel
 		 * The plane surface is the plane XY.
 		 * The surface is placed at its final position and orientation by the Matrix
 		 */
-		FPlaneSurface(const double InToleranceGeometric, const FMatrixH& InMatrix);
+		FPlaneSurface(const double InToleranceGeometric, const FMatrixH& InMatrix, const FSurfacicBoundary& InBoundary);
 
 		/**
 		 * The plane surface is the plane XY.
 		 * The surface is placed at its final position and orientation by the Matrix,
 		 * The matrix is calculated from the normal plane at its final position and its distance from the origin along the normal.
 		 */
-		FPlaneSurface(const double InToleranceGeometric, double DistanceFromOrigin, FPoint Normal)
-			: FPlaneSurface(InToleranceGeometric, Normal* DistanceFromOrigin, Normal)
+		FPlaneSurface(const double InToleranceGeometric, double InDistanceFromOrigin, FPoint InNormal, const FSurfacicBoundary& InBoundary)
+			: FPlaneSurface(InToleranceGeometric, InNormal* InDistanceFromOrigin, InNormal, InBoundary)
 		{
 		}
 
@@ -36,7 +36,7 @@ namespace CADKernel
 		 * The surface is placed at its final position and orientation by the Matrix 
 		 * The matrix is calculated from the plan origine at its final position and its final normal
 		 */
-		FPlaneSurface(const double InToleranceGeometric, const FPoint& Position, FPoint Normal);
+		FPlaneSurface(const double InToleranceGeometric, const FPoint& Position, FPoint Normal, const FSurfacicBoundary& InBoundary);
 
 		FPlaneSurface(FCADKernelArchive& Archive)
 			: FSurface()
@@ -44,6 +44,18 @@ namespace CADKernel
 			Serialize(Archive);
 		}
 		
+		virtual void SetMinToleranceIso() const override
+		{
+			FPoint Origin = Matrix.Multiply(FPoint::ZeroPoint);
+
+			FPoint Point2DU{ 1 , 0, 0 };
+			FPoint Point2DV{ 0, 1, 0 };
+
+			double ToleranceU = Tolerance3D / ComputeScaleAlongAxis(Point2DU, Matrix, Origin);
+			double ToleranceV = Tolerance3D / ComputeScaleAlongAxis(Point2DV, Matrix, Origin);
+
+			MinToleranceIso.Set(ToleranceU, ToleranceV);
+		}
 
 	public:
 
@@ -98,9 +110,6 @@ namespace CADKernel
 			bOutClosedAlongV = false;
 		}
 
-	private:
-		virtual const void ComputeIsoTolerances() const override;
-
 	};
 
 } // namespace CADKernel
diff --git a/Engine/Source/Runtime/Datasmith/CADKernel/Public/CADKernel/Geo/Surfaces/Surface.h b/Engine/Source/Runtime/Datasmith/CADKernel/Public/CADKernel/Geo/Surfaces/Surface.h
--- a/Engine/Source/Runtime/Datasmith/CADKernel/Public/CADKernel/Geo/Surfaces/Surface.h
+++ b/Engine/Source/Runtime/Datasmith/CADKernel/Public/CADKernel/Geo/Surfaces/Surface.h
@@ -7,6 +7,7 @@
 #include "CADKernel/Geo/GeoEnum.h"
 #include "CADKernel/Geo/GeoPoint.h"
 #include "CADKernel/Math/Boundary.h"
+#include "CADKernel/Math/MatrixH.h"
 #include "CADKernel/Math/Point.h"
 #include "CADKernel/Utils/Cache.h"
 
@@ -27,16 +28,19 @@ namespace CADKernel
 	protected:
 
 		double Tolerance3D;
-		mutable TCache<FSurfacicTolerance> ToleranceIsos;
+
+		/**
+		 * Tolerance along Iso U/V is very costly to compute and not accurate.
+		 * A first approximation is based on the surface boundary along U and along V
+		 * Indeed, defining a Tolerance2D has no sense has the boundary length along an Iso could be very very huge compare to the boundary length along the other Iso like [[0, 1000] [0, 1]]
+		 * The tolerance along an iso is the length of the boundary along this iso divided by 100 000: if the curve length in 3d is 10m, the tolerance is 0.01mm
+		 * In the used, a local tolerance has to be estimated
+		 */
+		mutable TCache<FSurfacicTolerance> MinToleranceIso;
 		mutable FSurfacicBoundary Boundary;
 
 	private:
 
-		/**
-		 * Return the tolerance in the parametric space of the surface along isoU and isoV i.e Distance(Point(U, V), Point(U + ToleranceIsoU, V)) = Tolerance3D
-		 * With Tolerance3D = FSysteme.GeometricalTolerance
-		 */
-		const virtual void ComputeIsoTolerances() const;
 
 		virtual void InitBoundary() const
 		{
@@ -54,6 +58,13 @@ namespace CADKernel
 
 		}
 
+		FSurface(double InToleranceGeometric, const FSurfacicBoundary& InBoundary)
+			: FEntityGeom()
+			, Tolerance3D(InToleranceGeometric)
+			, Boundary(InBoundary)
+		{
+		}
+
 		FSurface(double InToleranceGeometric, double UMin, double UMax, double VMin, double VMax)
 			: FEntityGeom()
 			, Tolerance3D(InToleranceGeometric)
@@ -61,8 +72,12 @@ namespace CADKernel
 		{
 		}
 
-	public:
+		virtual void SetMinToleranceIso() const
+		{
+			MinToleranceIso.Set(Boundary[EIso::IsoU].ComputeMinimalTolerance(), Boundary[EIso::IsoV].ComputeMinimalTolerance());
+		}
 
+	public:
 
 		virtual void Serialize(FCADKernelArchive& Ar) override
 		{
@@ -75,7 +90,7 @@ namespace CADKernel
 			FEntityGeom::Serialize(Ar);
 
 			Ar << Tolerance3D;
-			Ar << ToleranceIsos;
+			Ar << MinToleranceIso;
 			Ar << Boundary;
 		}
 
@@ -102,6 +117,7 @@ namespace CADKernel
 			if (!Boundary.IsValid())
 			{
 				InitBoundary();
+				SetMinToleranceIso();
 			}
 			return Boundary;
 		};
@@ -109,55 +125,41 @@ namespace CADKernel
 		void ExtendBoundaryTo(const FSurfacicBoundary MaxLimit)
 		{
 			Boundary.ExtendTo(MaxLimit);
+			SetMinToleranceIso();
 		}
 
 		virtual TSharedPtr<FEntityGeom> ApplyMatrix(const FMatrixH& InMatrix) const = 0;
 
 		/**
-		 * Return the tolerance in the parametric space of the surface i.e Distance(Point(U, V), Point(U', V')) = Tolerance3D with Distance((U,V), (U', V')) = Tolerance2D
+		 * Tolerance along Iso U/V is very costly to compute and not accurate.
+		 * A first approximation is based on the surface parametric space length along U and along V
+		 * Indeed, defining a Tolerance2D has no sense has the boundary length along an Iso could be very very huge compare to the boundary length along the other Iso like [[0, 1000] [0, 1]]
+		 * @see FBoundary::ComputeMinimalTolerance
+		 * In the used, a local tolerance has to be estimated
 		 */
-		const virtual double Get2DTolerance() const
-		{
-			if (!ToleranceIsos.IsValid())
-			{
-				GetIsoTolerances();
-			}
-			FSurfacicTolerance& Tolerance = ToleranceIsos;
-			return FMath::Min(Tolerance[EIso::IsoU], Tolerance[EIso::IsoV]);
-		}
-
-		const double Get3DTolerance()
+		const FSurfacicTolerance& GetIsoTolerances() const
 		{
-			return Tolerance3D;
+			ensureCADKernel(MinToleranceIso.IsValid());
+			return MinToleranceIso;
 		}
 
 		/**
-		 * Return the tolerance in the parametric space of the surface along isoU and isoV i.e Distance(Point(U, V), Point(U + ToleranceIsoU, V)) = Tolerance3D
+		 * Return the minimum tolerance in the parametric space of the surface along the specified axis 
 		 * With Tolerance3D = FSysteme.GeometricalTolerance
+		 * @see FBoundary::ComputeMinimalTolerance
 		 */
-		const FSurfacicTolerance& GetIsoTolerances() const
+		const double& GetIsoTolerance(EIso Iso) const
 		{
-			if (!ToleranceIsos.IsValid())
-			{
-				ComputeIsoTolerances();
-			}
-			return ToleranceIsos;
+			ensureCADKernel(MinToleranceIso.IsValid());
+			return ((FSurfacicTolerance) MinToleranceIso)[Iso];
 		}
 
-		/**
-		 * Return the tolerance in the parametric space of the surface along the specified axis 
-		 * With Tolerance3D = FSysteme.GeometricalTolerance
-		 * @see GetIsoTolerances
-		 */
-		const double& GetIsoTolerance(EIso Iso) const
+		const double Get3DTolerance()
 		{
-			if (!ToleranceIsos.IsValid())
-			{
-				ComputeIsoTolerances();
-			}
-			return ((FSurfacicTolerance) ToleranceIsos)[Iso];
+			return Tolerance3D;
 		}
 
+
 		virtual void EvaluatePoint(const FPoint2D& InSurfacicCoordinate, FSurfacicPoint& OutPoint3D, int32 InDerivativeOrder = 0) const = 0;
 
 		virtual void EvaluatePoints(const TArray<FPoint2D>& InSurfacicCoordinates, TArray<FSurfacicPoint>& OutPoint3D, int32 InDerivativeOrder = 0) const;
@@ -224,6 +226,20 @@ namespace CADKernel
 			}
 		}
 
+		/**
+		 * This function return the scale of the input Axis.
+		 * This function is useful to estimate tolerance when scales are defined in the Matrix
+		 * @param InAxis a vetor of length 1
+		 * @param InMatrix
+		 * @param InOrigin = InMatrix*FPoint::ZeroPoint
+		 */
+		static double ComputeScaleAlongAxis(const FPoint& InAxis, const FMatrixH& InMatrix, const FPoint& InOrigin)
+		{
+			FPoint Point = InMatrix.Multiply(InAxis);
+			double Length = InOrigin.Distance(Point);
+			return Length;
+		};
+
 
 	};
 } // namespace CADKernel
diff --git a/Engine/Source/Runtime/Datasmith/CADKernel/Public/CADKernel/Geo/Surfaces/TabulatedCylinderSurface.h b/Engine/Source/Runtime/Datasmith/CADKernel/Public/CADKernel/Geo/Surfaces/TabulatedCylinderSurface.h
--- a/Engine/Source/Runtime/Datasmith/CADKernel/Public/CADKernel/Geo/Surfaces/TabulatedCylinderSurface.h
+++ b/Engine/Source/Runtime/Datasmith/CADKernel/Public/CADKernel/Geo/Surfaces/TabulatedCylinderSurface.h
@@ -17,7 +17,7 @@ namespace CADKernel
 		FPoint DirectorVector;
 
 		/**
-		 * A tabulated cylinder is a surface formed by moving a line segment called GuideCurve parallel to
+		 * A tabulated cylinder is a surface formed by moving a curve segment called GuideCurve parallel to
 		 * itself along a vector called DirectorVector. 
 		 *
 		 * The bounds of the cylinder are defined as follow:
@@ -30,6 +30,7 @@ namespace CADKernel
 			, GuideCurve(InGuideCurve)
 			, DirectorVector(InDirectorVector)
 		{
+			SetMinToleranceIso();
 		}
 
 		FTabulatedCylinderSurface(FCADKernelArchive& Archive)
diff --git a/Engine/Source/Runtime/Datasmith/CADKernel/Public/CADKernel/Geo/Surfaces/TorusSurface.h b/Engine/Source/Runtime/Datasmith/CADKernel/Public/CADKernel/Geo/Surfaces/TorusSurface.h
--- a/Engine/Source/Runtime/Datasmith/CADKernel/Public/CADKernel/Geo/Surfaces/TorusSurface.h
+++ b/Engine/Source/Runtime/Datasmith/CADKernel/Public/CADKernel/Geo/Surfaces/TorusSurface.h
@@ -34,6 +34,7 @@ namespace CADKernel
 			, MajorRadius(InMajorRadius)
 			, MinorRadius(InMinorRadius)
 		{
+			SetMinToleranceIso();
 		}
 
 		FTorusSurface(FCADKernelArchive& Archive)
@@ -42,6 +43,22 @@ namespace CADKernel
 			Serialize(Archive);
 		}
 
+		void SetMinToleranceIso() const
+		{
+			double Tolerance2DU = Tolerance3D / MajorRadius;
+			double Tolerance2DV = Tolerance3D / MinorRadius;
+
+			FPoint Origin = Matrix.Multiply(FPoint::ZeroPoint);
+
+			FPoint Point2DU{ 1 , 0, 0 };
+			FPoint Point2DV{ 0, 1, 0 };
+
+			Tolerance2DU /= ComputeScaleAlongAxis(Point2DU, Matrix, Origin);
+			Tolerance2DV /= ComputeScaleAlongAxis(Point2DV, Matrix, Origin);
+
+			MinToleranceIso.Set(Tolerance2DU, Tolerance2DV);
+		}
+		
 	public:
 
 		virtual void Serialize(FCADKernelArchive& Ar) override
diff --git a/Engine/Source/Runtime/Datasmith/CADKernel/Public/CADKernel/Math/Aabb.h b/Engine/Source/Runtime/Datasmith/CADKernel/Public/CADKernel/Math/Aabb.h
--- a/Engine/Source/Runtime/Datasmith/CADKernel/Public/CADKernel/Math/Aabb.h
+++ b/Engine/Source/Runtime/Datasmith/CADKernel/Public/CADKernel/Math/Aabb.h
@@ -21,104 +21,125 @@ namespace CADKernel
 	class FSphere;
 	class FPlane;
 
-	class FAABB
+	template<class PointType>
+	class TAABB
 	{
 
-	private:
-		FPoint MinCorner;
-		FPoint MaxCorner;
+	protected:
+		PointType MinCorner;
+		PointType MaxCorner;
 
 	public:
-		FAABB()
-			: MinCorner(HUGE_VALUE, HUGE_VALUE, HUGE_VALUE)
-			, MaxCorner(-HUGE_VALUE, -HUGE_VALUE, -HUGE_VALUE)
+		TAABB()
+			: MinCorner(PointType::FarawayPoint)
+			, MaxCorner(-PointType::FarawayPoint)
 		{
 		}
 
-		FAABB(const FPoint& InMinCorner, const FPoint& InMaxCorner)
+		TAABB(const PointType& InMinCorner, const PointType& InMaxCorner)
 			: MinCorner(InMinCorner)
 			, MaxCorner(InMaxCorner)
 		{
 		}
 
+		friend FArchive& operator<<(FArchive& Ar, TAABB& AABB)
+		{
+			Ar << AABB.MinCorner;
+			Ar << AABB.MaxCorner;
+			return Ar;
+		}
+
 		bool IsValid() const
 		{
-			return (MinCorner.X < MaxCorner.Y);
+			for (int32 Axis = 0; Axis < PointType::Dimension; Axis++)
+			{
+				if (MinCorner[Axis] > MaxCorner[Axis])
+				{
+					return false;
+				}
+			}
+			return true;
 		}
 
 		void Empty()
 		{
-			MinCorner.Set(HUGE_VALUE, HUGE_VALUE, HUGE_VALUE);
-			MaxCorner.Set(-HUGE_VALUE, -HUGE_VALUE, -HUGE_VALUE);
+			MinCorner = PointType::FarawayPoint;
+			MaxCorner = -PointType::FarawayPoint;
 		}
 
-		bool Contains(const FPoint& Point) const
+		bool Contains(const PointType& Point) const
 		{
-			return ((Point.X > MinCorner.X - SMALL_NUMBER) && (Point.X < MaxCorner.X + SMALL_NUMBER) &&
-				(Point.Y > MinCorner.Y - SMALL_NUMBER) && (Point.Y < MaxCorner.Y + SMALL_NUMBER) &&
-				(Point.Z > MinCorner.Z - SMALL_NUMBER) && (Point.Z < MaxCorner.Z + SMALL_NUMBER));
+			for (int32 Axis = 0; Axis < PointType::Dimension; Axis++)
+			{
+				if ((Point[Axis] < MinCorner[Axis]) || (Point[Axis] > MaxCorner[Axis]))
+				{
+					return false;
+				}
+			}
+			return true;
 		}
 
-		void SetMinDimension(double MinDimension)
+		void SetMinSize(double MinSize)
 		{
-			for (int32 Axis = 0; Axis < 3; Axis++)
+			for (int32 Axis = 0; Axis < PointType::Dimension; Axis++)
 			{
-				double AxisDimension = GetDimension(Axis);
-				if (AxisDimension < MinDimension)
+				double AxisSize = GetSize(Axis);
+				if (AxisSize < MinSize)
 				{
-					double Offset = (MinDimension - AxisDimension) / 2;
+					double Offset = (MinSize - AxisSize) / 2;
 					MinCorner[Axis] -= Offset;
 					MaxCorner[Axis] += Offset;
 				}
 			}
 		}
 
-		double GetMaxDimension() const
+		double GetMaxSize() const
 		{
-			double MaxDimension = 0;
-			for (int32 Index = 0; Index < 3; Index++)
+			double MaxSideSize = 0;
+			for (int32 Index = 0; Index < PointType::Dimension; Index++)
 			{
-				double Dimension = GetDimension(Index);
-				if (Dimension > MaxDimension)
+				double Size = GetSize(Index);
+				if (Size > MaxSideSize)
 				{
-					MaxDimension = Dimension;
+					MaxSideSize = Size;
 				}
 			}
-			return MaxDimension;
+			return MaxSideSize;
 		}
 
-		double GetDimension(int32 Axis) const
+		double GetSize(int32 Axis) const
 		{
 			return MaxCorner[Axis] - MinCorner[Axis];
 		}
 
-		bool Contains(const FAABB& Aabb) const
+		double DiagonalLength() const
 		{
-			return IsValid() && Aabb.IsValid() && Contains(Aabb.MinCorner) && Contains(Aabb.MaxCorner);
+			return MaxCorner.Distance(MinCorner);
 		}
 
-		const FPoint& GetMin() const
+		PointType Diagonal() const
 		{
-			return MinCorner;
+			return MaxCorner - MinCorner;
 		}
 
-		const FPoint& GetMax() const
+		bool Contains(const TAABB& Aabb) const
 		{
-			return MaxCorner;
+			return IsValid() && Aabb.IsValid() && Contains(Aabb.MinCorner) && Contains(Aabb.MaxCorner);
 		}
 
-		FPoint GetCorner(int32 Corner) const
+		const PointType& GetMin() const
 		{
-			return FPoint(
-				Corner & EAABBBoundary::XMin ? MinCorner[0] : MaxCorner[0],
-				Corner & EAABBBoundary::YMin ? MinCorner[1] : MaxCorner[1],
-				Corner & EAABBBoundary::ZMin ? MinCorner[2] : MaxCorner[2]
-			);
+			return MinCorner;
+		}
+
+		const PointType& GetMax() const
+		{
+			return MaxCorner;
 		}
 
-		FAABB& operator+= (const double* Point)
+		TAABB& operator+= (const double* Point)
 		{
-			for (int32 Index = 0; Index < 3; Index++)
+			for (int32 Index = 0; Index < PointType::Dimension; Index++)
 			{
 				if (Point[Index] < MinCorner[Index])
 				{
@@ -132,9 +153,9 @@ namespace CADKernel
 			return *this;
 		}
 
-		FAABB& operator+= (const FPoint& Point)
+		TAABB& operator+= (const PointType& Point)
 		{
-			for (int32 Index = 0; Index < 3; Index++)
+			for (int32 Index = 0; Index < PointType::Dimension; Index++)
 			{
 				if (Point[Index] < MinCorner[Index])
 				{
@@ -148,92 +169,84 @@ namespace CADKernel
 			return *this;
 		}
 
-		void Offset (double Offset)
+
+		TAABB& operator+= (const TArray<PointType>& Points)
+		{
+			for (const PointType& Point : Points)
+			{
+				*this += Point;
+			}
+			return *this;
+		}
+
+
+		void Offset(double Offset)
 		{
-			for (int32 Index = 0; Index < 3; Index++)
+			for (int32 Index = 0; Index < PointType::Dimension; Index++)
 			{
 				MinCorner[Index] -= Offset;
 				MaxCorner[Index] += Offset;
 			}
 		}
 
-		FAABB& operator+= (const FAABB& aabb)
+		TAABB& operator+= (const TAABB& aabb)
 		{
 			*this += aabb.MinCorner;
 			*this += aabb.MaxCorner;
 			return *this;
 		}
 
-		FAABB operator+ (const FPoint& Point) const
+		TAABB operator+ (const PointType& Point) const
 		{
-			FAABB Other = *this;
+			TAABB Other = *this;
 			Other += Point;
 			return Other;
 		}
 
 
-		FAABB operator+ (const FAABB& Aabb) const
+		TAABB operator+ (const TAABB& Aabb) const
 		{
-			FAABB Other = *this;
+			TAABB Other = *this;
 			Other += Aabb;
 			return Other;
 		}
-
 	};
 
-	class CADKERNEL_API FAABB2D
+	class FAABB : public TAABB<FPoint>
 	{
-	private:
-
-		FPoint2D MinCorner;
-		FPoint2D MaxCorner;
 
 	public:
-
-		FAABB2D()
-			: MinCorner(HUGE_VALUE, HUGE_VALUE)
-			, MaxCorner(-HUGE_VALUE, -HUGE_VALUE)
-		{
-		}
-
-		FAABB2D(const FPoint2D& InMinCorner, const FPoint2D& InMaxCorner)
-			: MinCorner(InMinCorner)
-			, MaxCorner(InMaxCorner)
-		{
-		}
-
-		friend FArchive& operator<<(FArchive& Ar, FAABB2D& AABB)
-		{
-			Ar << AABB.MinCorner;
-			Ar << AABB.MaxCorner;
-			return Ar;
-		}
-
-		const FPoint2D& GetMin() const
+		FAABB()
+			: TAABB<FPoint>()
 		{
-			return MinCorner;
 		}
 
-		const FPoint2D& GetMax() const
+		FAABB(const FPoint& InMinCorner, const FPoint& InMaxCorner)
+			: TAABB<FPoint>(InMinCorner, InMaxCorner)
 		{
-			return MaxCorner;
 		}
 
-		bool IsValid() const
+		FPoint GetCorner(int32 Corner) const
 		{
-			return (MinCorner.U < MaxCorner.U);
+			return FPoint(
+				Corner & EAABBBoundary::XMin ? MinCorner[0] : MaxCorner[0],
+				Corner & EAABBBoundary::YMin ? MinCorner[1] : MaxCorner[1],
+				Corner & EAABBBoundary::ZMin ? MinCorner[2] : MaxCorner[2]
+			);
 		}
+	};
 
-		void Empty()
+	class CADKERNEL_API FAABB2D : public TAABB<FPoint2D>
+	{
+	public:
+		FAABB2D()
+			: TAABB<FPoint2D>()
 		{
-			MinCorner.Set(HUGE_VALUE, HUGE_VALUE);
-			MaxCorner.Set(-HUGE_VALUE, -HUGE_VALUE);
 		}
 
-		bool Contains(const FPoint2D& Point) const
+		FAABB2D(const FPoint& InMinCorner, const FPoint& InMaxCorner)
+			: TAABB<FPoint2D>(InMinCorner, InMaxCorner)
 		{
-			return ((Point.U > MinCorner.U - SMALL_NUMBER) && (Point.U < MaxCorner.U + SMALL_NUMBER) &&
-				(Point.V > MinCorner.V - SMALL_NUMBER) && (Point.V < MaxCorner.V + SMALL_NUMBER));
 		}
 
 		FPoint2D GetCorner(int32 CornerIndex) const
@@ -244,85 +257,6 @@ namespace CADKernel
 			);
 		}
 
-		bool Contains(const FAABB2D& Aabb2d) const
-		{
-			return IsValid() && Aabb2d.IsValid() && Contains(Aabb2d.MinCorner) && Contains(Aabb2d.MaxCorner);
-		}
-
-		double DiagonalLength() const
-		{
-			return MinCorner.Distance(MaxCorner);
-		}
-
-		FPoint2D Diagonal() const
-		{
-			return MaxCorner - MinCorner;
-		}
-
-		double GetDimension(int32 Axis) const
-		{
-			return MaxCorner[Axis] - MinCorner[Axis];
-		}
-
-		FPoint2D Center() const
-		{
-			return (MinCorner + MaxCorner) / 2.;
-		}
-
-		void Set(const FPoint2D& InMinCorner, const FPoint2D& InMaxCorner)
-		{
-			MinCorner = InMinCorner;
-			MaxCorner = InMaxCorner;
-		}
-
-		FAABB2D& operator+= (const FAABB2D& Aabb2d)
-		{
-			*this += Aabb2d.GetMin();
-			*this += Aabb2d.GetMax();
-			return *this;
-		}
-
-		FAABB2D& operator+= (TArray<FPoint2D> Points)
-		{
-			for (FPoint Point : Points)
-			{
-				*this += Point;
-			}
-			return *this;
-		}
-
-		FAABB2D& operator+= (const FPoint2D& Point)
-		{
-			if (Point.U < MinCorner.U)
-			{
-				MinCorner.U = Point.U;
-			}
-
-			if (Point.V < MinCorner.V)
-			{
-				MinCorner.V = Point.V;
-			}
-
-			if (Point.U > MaxCorner.U)
-			{
-				MaxCorner.U = Point.U;
-			}
-
-			if (Point.V > MaxCorner.V)
-			{
-				MaxCorner.V = Point.V;
-			}
-
-			return *this;
-		}
-
-		void Offset(double Offset)
-		{
-			MinCorner.U -= Offset;
-			MinCorner.V -= Offset;
-			MaxCorner.U += Offset;
-			MaxCorner.V += Offset;
-		}
 	};
 
 } // namespace CADKernel
diff --git a/Engine/Source/Runtime/Datasmith/CADKernel/Public/CADKernel/Math/BSpline.h b/Engine/Source/Runtime/Datasmith/CADKernel/Public/CADKernel/Math/BSpline.h
--- a/Engine/Source/Runtime/Datasmith/CADKernel/Public/CADKernel/Math/BSpline.h
+++ b/Engine/Source/Runtime/Datasmith/CADKernel/Public/CADKernel/Math/BSpline.h
@@ -33,6 +33,8 @@ namespace CADKernel
 
 		CADKERNEL_API void EvaluatePoint(const FNURBSSurface&, const FPoint2D& InPoint2D, FSurfacicPoint& OutPoint3D, int32 InDerivativeOrder);
 		CADKERNEL_API void EvaluatePointGrid(const FNURBSSurface&, const FCoordinateGrid& Coords, FSurfacicSampling& OutPoints, bool bComputeNormals);
+
+		CADKERNEL_API TSharedRef<FNURBSCurve> DuplicateNurbsCurveWithHigherDegree(int32 degre, const FNURBSCurve& InCurve);
 	}
 }
 
diff --git a/Engine/Source/Runtime/Datasmith/CADKernel/Public/CADKernel/Math/Boundary.h b/Engine/Source/Runtime/Datasmith/CADKernel/Public/CADKernel/Math/Boundary.h
--- a/Engine/Source/Runtime/Datasmith/CADKernel/Public/CADKernel/Math/Boundary.h
+++ b/Engine/Source/Runtime/Datasmith/CADKernel/Public/CADKernel/Math/Boundary.h
@@ -8,6 +8,11 @@
 
 namespace CADKernel
 {
+	/**
+	 * MINIMAL_UNIT_LINEAR_TOLERANCE allows to define the minimal tolerance value of a parametric space
+	 * @see FLinearBoundary::ComputeMinimalTolerance
+	 */
+	#define MINIMAL_UNIT_LINEAR_TOLERANCE 10e-5
 
 	struct CADKERNEL_API FLinearBoundary
 	{
@@ -30,8 +35,7 @@ namespace CADKernel
 
 		friend FArchive& operator<<(FArchive& Ar, FLinearBoundary& Bounds)
 		{
-			Ar << Bounds.Min;
-			Ar << Bounds.Max;
+			Ar.Serialize(&Bounds, sizeof(FLinearBoundary));
 			return Ar;
 		}
 
@@ -96,6 +100,16 @@ namespace CADKernel
 			return (DeltaU < KINDA_SMALL_NUMBER);
 		}
 
+		/**
+		 * Compute the minimal tolerance of the parametric domain i.e. 
+		 * ToleranceMin = Boundary.Length() * MINIMAL_UNIT_LINEAR_TOLERANCE
+		 * e.g. for a curve of 1m with a parametric space define between [0, 1], the parametric tolerance is 0.01
+		 * This is a minimal value that has to be replace with a more accurate value when its possible
+		 */
+		double ComputeMinimalTolerance() const
+		{
+			return Length() * MINIMAL_UNIT_LINEAR_TOLERANCE;
+		}
 
 		/**
 		 * If a coordinate is outside the bounds, set the coordinate at the closed limit
diff --git a/Engine/Source/Runtime/Datasmith/CADKernel/Public/CADKernel/Math/Point.h b/Engine/Source/Runtime/Datasmith/CADKernel/Public/CADKernel/Math/Point.h
--- a/Engine/Source/Runtime/Datasmith/CADKernel/Public/CADKernel/Math/Point.h
+++ b/Engine/Source/Runtime/Datasmith/CADKernel/Public/CADKernel/Math/Point.h
@@ -20,6 +20,8 @@ namespace CADKernel
 
 		/** A zero point (0,0,0)*/
 		static const FPoint ZeroPoint;
+		static const FPoint FarawayPoint;
+		static const int32 Dimension;
 
 		explicit FPoint(double InCoordX = 0., double InCoordY = 0., double InCoordZ = 0.)
 		{
@@ -232,11 +234,37 @@ namespace CADKernel
 			return FMath::Square(Point.X - X) + FMath::Square(Point.Y - Y) + FMath::Square(Point.Z - Z);
 		}
 
-		double ComputeCosinus(const FPoint& Vector2) const;
+		double ComputeCosinus(const FPoint& OtherVector) const
+		{
+			FPoint ThisNormalized = *this;
+			FPoint OtherNormalized = OtherVector;
+
+			ThisNormalized.Normalize();
+			OtherNormalized.Normalize();
+
+			double Cosinus = ThisNormalized * OtherNormalized;
+
+			return FMath::Max(-1.0, FMath::Min(Cosinus, 1.0));
+		}
+
+		double ComputeSinus(const FPoint& OtherVector) const
+		{
+			FPoint ThisNormalized = *this;
+			FPoint OtherNormalized = OtherVector;
 
-		double ComputeSinus(const FPoint& Vector2) const;
+			ThisNormalized.Normalize();
+			OtherNormalized.Normalize();
 
-		double ComputeAngle(const FPoint& Vector2) const;
+			FPoint SinusPoint = ThisNormalized ^ OtherNormalized;
+			double Sinus = SinusPoint.Length();
+			return FMath::Min(Sinus, 1.0);
+		}
+
+		double ComputeAngle(const FPoint& OtherVector) const
+		{
+			double CosAngle = ComputeCosinus(OtherVector);
+			return acos(CosAngle);
+		}
 
 		double SignedAngle(const FPoint& Vector2, const FPoint& Normal) const;
 
@@ -247,13 +275,18 @@ namespace CADKernel
 
 		friend FArchive& operator<<(FArchive& Ar, FPoint& Point)
 		{
-			Ar << Point.X;
-			Ar << Point.Y;
-			Ar << Point.Z;
+			Ar.Serialize(&Point, sizeof(FPoint));
 			return Ar;
 		}
 
-		static int32 Dimension() { return 3; }
+		friend FPoint Abs(const FPoint& Point)
+		{
+			double PointX = FMath::Abs(Point.X);
+			double PointY = FMath::Abs(Point.Y);
+			double PointZ = FMath::Abs(Point.Z);
+			return FPoint(PointX, PointY, PointZ);
+		}
+
 	};
 
 	class CADKERNEL_API FPointH
@@ -267,6 +300,8 @@ namespace CADKernel
 	public:
 		/** A zero point (0,0,0)*/
 		static const FPointH ZeroPoint;
+		static const FPointH FarawayPoint;
+		static const int32 Dimension;
 
 		FPointH()
 			: X(0.)
@@ -324,10 +359,7 @@ namespace CADKernel
 
 		friend FArchive& operator<<(FArchive& Ar, FPointH& Point)
 		{
-			Ar << Point.X;
-			Ar << Point.Y;
-			Ar << Point.Z;
-			Ar << Point.W;
+			Ar.Serialize(&Point, sizeof(FPointH));
 			return Ar;
 		}
 	};
@@ -342,6 +374,8 @@ namespace CADKernel
 
 		/** A zero point (0,0,0)*/
 		static const FPoint2D ZeroPoint;
+		static const FPoint2D FarawayPoint;
+		static const int32 Dimension;
 
 		FPoint2D()
 		{
@@ -519,11 +553,17 @@ namespace CADKernel
 			return *this;
 		}
 
-		friend FArchive& operator<<(FArchive& Ar, FPoint2D& Point)
+		double ComputeCosinus(const FPoint2D& OtherVector) const
 		{
-			Ar << Point.U;
-			Ar << Point.V;
-			return Ar;
+			FPoint2D ThisNormalized = *this;
+			FPoint2D OtherNormalized = OtherVector;
+
+			ThisNormalized.Normalize();
+			OtherNormalized.Normalize();
+
+			double Cosinus = ThisNormalized * OtherNormalized;
+
+			return FMath::Max(-1.0, FMath::Min(Cosinus, 1.0));
 		}
 
 		FPoint2D GetPerpendicularVector()
@@ -531,7 +571,23 @@ namespace CADKernel
 			return FPoint2D(-V, U);
 		}
 
-		static constexpr const int32 Dimension() { return 2; }
+		friend FArchive& operator<<(FArchive& Ar, FPoint2D& Point)
+		{
+			Ar.Serialize(&Point, sizeof(FPoint2D));
+			return Ar;
+		}
+
+		friend FPoint2D Abs(const FPoint2D& Point)
+		{
+			double PointU = FMath::Abs(Point.U);
+			double PointV = FMath::Abs(Point.V);
+			return { PointU, PointV };
+		}
+
+		friend FPoint2D Max(const FPoint2D& PointA, const FPoint2D& PointB)
+		{
+			return { FMath::Max(PointA.U, PointB.U), FMath::Max(PointA.V, PointB.V) };
+		}
 	};
 
 	inline FPoint::FPoint(const FPoint2D& Point)
@@ -550,6 +606,8 @@ namespace CADKernel
 
 		/** A zero point (0,0,0)*/
 		static const FFPoint ZeroPoint;
+		static const FFPoint FarawayPoint;
+		static const int32 Dimension;
 
 		FFPoint(float InX = 0., float InY = 0., float InZ = 0.)
 		{
@@ -582,13 +640,9 @@ namespace CADKernel
 
 		friend FArchive& operator<<(FArchive& Ar, FFPoint& Point)
 		{
-			Ar << Point.X;
-			Ar << Point.Y;
-			Ar << Point.Z;
+			Ar.Serialize(&Point, sizeof(FFPoint));
 			return Ar;
 		}
-
-		static constexpr const int32 Dimension() { return 3; }
 	};
 
 	inline FPoint operator*(double Scale, const FPoint& Point)
diff --git a/Engine/Source/Runtime/Datasmith/CADKernel/Public/CADKernel/Mesh/Meshers/ParametricMesher.h b/Engine/Source/Runtime/Datasmith/CADKernel/Public/CADKernel/Mesh/Meshers/ParametricMesher.h
--- a/Engine/Source/Runtime/Datasmith/CADKernel/Public/CADKernel/Mesh/Meshers/ParametricMesher.h
+++ b/Engine/Source/Runtime/Datasmith/CADKernel/Public/CADKernel/Mesh/Meshers/ParametricMesher.h
@@ -88,7 +88,6 @@ namespace CADKernel
 			FChrono::PrintClockElapse(Log, TEXT("  |   |  "), TEXT("GeneratePoint Cloud "), GlobalGeneratePointCloudDuration);
 			FChrono::PrintClockElapse(Log, TEXT("  |   |  |  "), TEXT("Point Cloud "), GlobalPointCloudDuration);
 			FChrono::PrintClockElapse(Log, TEXT("  |   |  "), TEXT("ThinZones "), GlobalThinZones);
-			//FChrono::PrintClockElapse(Log, TEXT("  |   |  |  "), TEXT("GlobalScaleGrid"), GlobalScaleGrid);
 			FChrono::PrintClockElapse(Log, TEXT("  |   |  "), TEXT("Mesh ThinZones "), GlobalMeshThinZones);
 			FChrono::PrintClockElapse(Log, TEXT("  |   |  "), TEXT("MeshEdges "), GlobalMeshEdges);
 			FChrono::PrintClockElapse(Log, TEXT("  |   |  "), TEXT("TriangulateDuration "), GlobalTriangulateDuration);
diff --git a/Engine/Source/Runtime/Datasmith/CADKernel/Public/CADKernel/Topo/Body.h b/Engine/Source/Runtime/Datasmith/CADKernel/Public/CADKernel/Topo/Body.h
--- a/Engine/Source/Runtime/Datasmith/CADKernel/Public/CADKernel/Topo/Body.h
+++ b/Engine/Source/Runtime/Datasmith/CADKernel/Public/CADKernel/Topo/Body.h
@@ -46,6 +46,7 @@ namespace CADKernel
 		{
 			FTopologicalEntity::Serialize(Ar);
 			SerializeIdents(Ar, Shells);
+			SerializeMetadata(Ar);
 		}
 
 		virtual void SpawnIdent(FDatabase& Database) override
@@ -75,6 +76,13 @@ namespace CADKernel
 
 		void AddShell(TSharedRef<FShell> Shell);
 
+		void RemoveEmptyShell();
+
+		void Empty()
+		{
+			Shells.Empty();
+		}
+
 		const TArray<TSharedPtr<FShell>>& GetShells() const
 		{
 			return Shells;
diff --git a/Engine/Source/Runtime/Datasmith/CADKernel/Public/CADKernel/Topo/Linkable.h b/Engine/Source/Runtime/Datasmith/CADKernel/Public/CADKernel/Topo/Linkable.h
--- a/Engine/Source/Runtime/Datasmith/CADKernel/Public/CADKernel/Topo/Linkable.h
+++ b/Engine/Source/Runtime/Datasmith/CADKernel/Public/CADKernel/Topo/Linkable.h
@@ -115,7 +115,6 @@ namespace CADKernel
 			return 1;
 		}
 
-
 		bool HasTwin() const
 		{
 			return GetTwinsEntityCount() != 1;
@@ -139,7 +138,7 @@ namespace CADKernel
 			return TopologicalLink->GetTwinsEntities();
 		}
 
-		void RemoveFromLink() 
+		virtual void RemoveFromLink() 
 		{
 			if (TopologicalLink.IsValid())
 			{
diff --git a/Engine/Source/Runtime/Datasmith/CADKernel/Public/CADKernel/Topo/Model.h b/Engine/Source/Runtime/Datasmith/CADKernel/Public/CADKernel/Topo/Model.h
--- a/Engine/Source/Runtime/Datasmith/CADKernel/Public/CADKernel/Topo/Model.h
+++ b/Engine/Source/Runtime/Datasmith/CADKernel/Public/CADKernel/Topo/Model.h
@@ -33,6 +33,11 @@ namespace CADKernel
 			: FTopologicalEntity()
 		{
 			Serialize(Archive);
+			//ensureCADKernel(Archive.ArchiveModel == nullptr);
+			if(Archive.ArchiveModel == nullptr)
+			{
+				Archive.ArchiveModel = this;
+			}
 		}
 
 	public:
@@ -42,6 +47,7 @@ namespace CADKernel
 			FEntityGeom::Serialize(Ar);
 			SerializeIdents(Ar, (TArray<TSharedPtr<FEntity>>&) Bodies);
 			SerializeIdents(Ar, (TArray<TSharedPtr<FEntity>>&) Faces);
+			SerializeMetadata(Ar);
 		}
 
 		virtual void SpawnIdent(FDatabase& Database) override
@@ -71,38 +77,120 @@ namespace CADKernel
 			return EEntity::Model;
 		}
 
-		void AddEntity(TSharedRef<FTopologicalEntity> Entity);
+		void AddEntity(TSharedRef<FTopologicalEntity> InEntity);
+
+		void Add(const TSharedPtr<FTopologicalFace>& InFace)
+		{
+			Faces.Add(InFace);
+		}
+
+		void Append(TArray<TSharedPtr<FTopologicalFace>>& InNewFaces)
+		{
+			Faces.Append(InNewFaces);
+		}
+
+		void Append(TArray<TSharedPtr<FBody>>& InNewBody)
+		{
+			Bodies.Append(InNewBody);
+		}
+
+		void Add(const TSharedPtr<FBody>& InBody)
+		{
+			Bodies.Add(InBody);
+		}
+
+		void Empty()
+		{
+			Bodies.Empty();
+			Faces.Empty();
+		}
+
+		void RemoveFace(TSharedPtr<FTopologicalFace> InToplologicalFace)
+		{
+			Faces.Remove(InToplologicalFace);
+		}
+
+		void RemoveBody(TSharedPtr<FBody> InBody)
+		{
+			Bodies.Remove(InBody);
+		}
+
+		void PrintBodyAndShellCount();
+		void RemoveEmptyBodies();
 
-		void Add(TSharedPtr<FTopologicalFace> Face);
-		void Add(TSharedPtr<FBody> Body);
+		void RemoveEntity(TSharedPtr<FTopologicalEntity> InEntity);
 
-		void RemoveEntity(TSharedPtr<FTopologicalEntity> Entity);
-		void RemoveDomain(TSharedPtr<FTopologicalFace> ToplologicalFace);
-		void RemoveBody(TSharedPtr<FBody> Body);
-		bool Contains(TSharedPtr<FTopologicalEntity> Entity);
+		bool Contains(TSharedPtr<FTopologicalEntity> InEntity);
+
+		/**
+		 * Copy the body and face arrays of other model
+		 */
+		void Copy(const TSharedPtr<FModel>& OtherModel)
+		{
+			Bodies.Append(OtherModel->Bodies);
+			Faces.Append(OtherModel->Faces);
+		}
 		
-		const TArray<TSharedPtr<FTopologicalFace>>& GetFaces() const
+		/**
+		 * Copy the body and face arrays of other model
+		 */
+		void Copy(const FModel& OtherModel)
 		{
-			return Faces;
+			Bodies.Append(OtherModel.Bodies);
+			Faces.Append(OtherModel.Faces);
 		}
 
-		const TArray<TSharedPtr<FTopologicalFace>>& GetSurfaceList() const
-		{ 
-			return Faces;
+		/**
+		 * Copy the body and face arrays of other model
+		 * Empty other model arrays
+		 */
+		void Merge(FModel& OtherModel)
+		{
+			Copy(OtherModel);
+			OtherModel.Bodies.Empty();
+			OtherModel.Faces.Empty();
 		}
 
-		const TArray<TSharedPtr<FBody>>& GetBodyList() const
+		const TArray<TSharedPtr<FTopologicalFace>>& GetFaces() const
 		{
-			return Bodies;
+			return Faces;
 		}
 
+		virtual void GetFaces(TArray<TSharedPtr<FTopologicalFace>>& OutFaces) override;
+
 		virtual int32 FaceCount() const override;
 
-		virtual void GetFaces(TArray<TSharedPtr<FTopologicalFace>>& OutFaces) override;
+		const TArray<TSharedPtr<FBody>>& GetBodies() const
+		{
+			return Bodies;
+		}
 
 		virtual void SpreadBodyOrientation() override;
 
 		virtual TSharedPtr<FEntityGeom> ApplyMatrix(const FMatrixH& InMatrix) const override;
+
+		// Topo functions
+
+		/**
+		 * Check topology of each body
+		 */
+		void CheckTopology();
+
+		/**
+		 * Fore each body
+		 */
+		void FixModelTopology(double JoiningTolerance);
+
+		void MergeInto(TSharedPtr<FBody> Body, TArray<TSharedPtr<FTopologicalEntity>>& InEntities);
+		void Split(TSharedPtr<FBody> Body, TArray<TSharedPtr<FBody>>& OutNewBody);
+		void Join(TArray<TSharedPtr<FBody>> Bodies, double Tolerance);
+
+		/**
+		 * Fore each shell of each body, try to stitch topological gap
+		 */
+		void HealModelTopology(double JoiningTolerance);
+
+
 	};
 
 } // namespace CADKernel
diff --git a/Engine/Source/Runtime/Datasmith/CADKernel/Public/CADKernel/Topo/Shell.h b/Engine/Source/Runtime/Datasmith/CADKernel/Public/CADKernel/Topo/Shell.h
--- a/Engine/Source/Runtime/Datasmith/CADKernel/Public/CADKernel/Topo/Shell.h
+++ b/Engine/Source/Runtime/Datasmith/CADKernel/Public/CADKernel/Topo/Shell.h
@@ -12,7 +12,6 @@ namespace CADKernel
 {
 	class FTopologicalFace;
 	class FBody;
-	class FTopologicalFace;
 
 	struct FFaceSubset;
 
@@ -38,9 +37,11 @@ namespace CADKernel
 	class CADKERNEL_API FShell : public FTopologicalEntity, public FMetadataDictionary
 	{
 		friend FEntity;
+		friend class FBody;
 
 	private:
 		TArray<FOrientedFace> TopologicalFaces;
+		TWeakPtr<FBody> HostedBy;
 
 		FShell()
 			: FTopologicalEntity()
@@ -73,6 +74,8 @@ namespace CADKernel
 		{
 			FTopologicalEntity::Serialize(Ar);
 			SerializeIdents(Ar, (TArray<TOrientedEntity<FEntity>>&) TopologicalFaces);
+			SerializeIdent(Ar, HostedBy);
+			SerializeMetadata(Ar);
 		}
 
 		virtual void SpawnIdent(FDatabase& Database) override
@@ -91,7 +94,28 @@ namespace CADKernel
 			ResetMarkersRecursivelyOnEntities((TArray<TOrientedEntity<FEntity>>&) TopologicalFaces);
 		}
 
+		TWeakPtr<FBody>& GetHost()
+		{
+			return HostedBy;
+		}
+
+		void SetHost(TSharedPtr<FBody> Body)
+		{
+			HostedBy = Body;
+		}
+
+		void CompleteMetadata()
+		{
+			if (HostedBy.IsValid())
+			{
+				CompleteDictionary((const FMetadataDictionary&) * HostedBy.Pin());
+			}
+		}
+
+		void Empty(int32 NewSize = 0);
+
 		void Add(TSharedRef<FTopologicalFace> InTopologicalFace, EOrientation InOrientation);
+		void Add(TArray<TSharedPtr<FTopologicalFace>> Faces);
 
 #ifdef CADKERNEL_DEV
 		virtual FInfoEntity& GetInfo(FInfoEntity&) const override;
@@ -109,6 +133,12 @@ namespace CADKernel
 			return TopologicalFaces.Num();
 		}
 
+		void ReplaceFaces(TArray<FOrientedFace>& NewFaces)
+		{
+			Swap(TopologicalFaces, NewFaces);
+			NewFaces.Empty();
+		}
+
 		const TArray<FOrientedFace>& GetFaces() const
 		{
 			return TopologicalFaces;
@@ -116,10 +146,17 @@ namespace CADKernel
 
 		virtual void GetFaces(TArray<TSharedPtr<FTopologicalFace>>& OutFaces) override;
 
+		virtual void Merge(TSharedPtr<FShell>& Shell);
+
 		virtual void SpreadBodyOrientation() override;
 
 		void CheckTopology(TArray<FFaceSubset>& SubShells);
 
+		/**
+		 * @return true if the shell has at least one border edge 
+		 */
+		bool IsOpenShell();
+
 		bool IsInner() const
 		{
 			return ((States & EHaveStates::IsInner) == EHaveStates::IsInner);
diff --git a/Engine/Source/Runtime/Datasmith/CADKernel/Public/CADKernel/Topo/TopologicalEdge.h b/Engine/Source/Runtime/Datasmith/CADKernel/Public/CADKernel/Topo/TopologicalEdge.h
--- a/Engine/Source/Runtime/Datasmith/CADKernel/Public/CADKernel/Topo/TopologicalEdge.h
+++ b/Engine/Source/Runtime/Datasmith/CADKernel/Public/CADKernel/Topo/TopologicalEdge.h
@@ -13,6 +13,8 @@ namespace CADKernel
 {
 	typedef TTopologicalLink<FTopologicalEdge> FEdgeLink;
 
+	class FOrientedEdge;
+
 	enum ECoordinateType : uint8
 	{
 		VertexCoordinate,
@@ -32,6 +34,7 @@ namespace CADKernel
 
 	class FModelMesh;
 	class FRestrictionCurve;
+	class FSurface;
 	class FTopologicalLoop;
 	class FTopologicalVertex;
 
@@ -87,12 +90,14 @@ namespace CADKernel
 	private:
 
 		/**
-		 * Constructors of FEdge cannot be used directly,use FEdge::Make to create an FEdge object.
+		 * MANDATORY
+		 * Constructors of FEdge cannot be used directly, use FEdge::Make to create an FEdge object.
 		 */
 		FTopologicalEdge(const TSharedRef<FRestrictionCurve>& InCurve, const TSharedRef<FTopologicalVertex>& InVertex1, const TSharedRef<FTopologicalVertex>& InVertex2, const FLinearBoundary& InBoundary);
 		FTopologicalEdge(const TSharedRef<FRestrictionCurve>& InCurve, const TSharedRef<FTopologicalVertex>& InVertex1, const TSharedRef<FTopologicalVertex>& InVertex2);
 		FTopologicalEdge(const TSharedRef<FRestrictionCurve>& InCurve, const FLinearBoundary& InBoundary);
 		FTopologicalEdge(const TSharedRef<FRestrictionCurve>& InCurve);
+		FTopologicalEdge(const TSharedRef<FSurface>& InSurface, const FPoint2D& InCoordinateVertex1, const TSharedRef<FTopologicalVertex>& InVertex1, const FPoint2D& InCoordinateVertex2, const TSharedRef<FTopologicalVertex>& InVertex2);
 
 		FTopologicalEdge(FCADKernelArchive& Archive)
 			: TLinkable<FTopologicalEdge, FEdgeLink>()
@@ -100,12 +105,12 @@ namespace CADKernel
 			Serialize(Archive);
 		}
 
-		void SetBoundary(const TSharedRef<FTopologicalLoop>& NewBoundary)
+		void SetLoop(const TSharedRef<FTopologicalLoop>& NewBoundary)
 		{
 			Loop = NewBoundary;
 		}
 
-		void RemoveBoundary()
+		void RemoveLoop()
 		{
 			Loop.Reset();
 		}
@@ -124,6 +129,7 @@ namespace CADKernel
 		static TSharedPtr<FTopologicalEdge> Make(const TSharedRef<FRestrictionCurve>& InCurve, const TSharedRef<FTopologicalVertex>& InVertex1, const TSharedRef<FTopologicalVertex>& InVertex2);
 		static TSharedPtr<FTopologicalEdge> Make(const TSharedRef<FRestrictionCurve>& InCurve, const FLinearBoundary& InBoundary);
 		static TSharedPtr<FTopologicalEdge> Make(const TSharedRef<FRestrictionCurve>& InCurve);
+		static TSharedPtr<FTopologicalEdge> Make(const TSharedRef<FSurface>& InSurface, const FPoint2D& InCoordinateVertex1, const TSharedRef<FTopologicalVertex>& InVertex1, const FPoint2D& InCoordinateVertex2, const TSharedRef<FTopologicalVertex>& InVertex2);
 
 		virtual ~FTopologicalEdge() = default;
 
@@ -175,7 +181,7 @@ namespace CADKernel
 		 */
 		bool CheckIfDegenerated();
 
-		void Link(const TSharedRef<FTopologicalEdge>& OtherEdge);
+		void Link(const TSharedRef<FTopologicalEdge>& OtherEdge, double JoiningTolerance);
 
 		TSharedRef<const FTopologicalEdge> GetLinkActiveEdge() const
 		{
@@ -279,12 +285,12 @@ namespace CADKernel
 			return EndVertex.ToSharedRef();
 		}
 
-		TSharedPtr<FTopologicalVertex> GetOtherVertex(const TSharedRef<FTopologicalVertex>& Vertex) const
+		TSharedPtr<FTopologicalVertex> GetOtherVertex(const TSharedRef<FTopologicalVertex>& Vertex)
 		{
 			return (Vertex->GetLink() == StartVertex->GetLink() ? EndVertex : (Vertex->GetLink() == EndVertex->GetLink() ? StartVertex : TSharedPtr<FTopologicalVertex>()));
 		}
 
-		TSharedPtr<FTopologicalVertex> GetOtherVertex(const TSharedRef<const FTopologicalVertex>& Vertex) const
+		const TSharedPtr<FTopologicalVertex> GetOtherVertex(const TSharedRef<FTopologicalVertex>& Vertex) const
 		{
 			return (Vertex->GetLink() == StartVertex->GetLink() ? EndVertex : (Vertex->GetLink() == EndVertex->GetLink() ? StartVertex : TSharedPtr<FTopologicalVertex>()));
 		}
@@ -352,6 +358,8 @@ namespace CADKernel
 			return EndVertex->GetCoordinates();
 		}
 
+
+
 		void GetTangentsAtExtremities(FPoint& StartTangent, FPoint& EndTangent, bool bForward) const;
 
 		// ======   Meshing Function   ======
@@ -512,10 +520,27 @@ namespace CADKernel
 			Curve->ApproximatePolyline(Polyline);
 		}
 
+		FPoint GetTangentAt(const double InCoordinate) const
+		{
+			return Curve->GetTangentAt(InCoordinate);
+		}
+
+		FPoint2D GetTangent2DAt(const double InCoordinate) const
+		{
+			return Curve->GetTangent2DAt(InCoordinate);
+		}
+
+		/**
+		 * Return the tangent at the input vertex
+		 */
+		FPoint GetTangentAt(const TSharedRef<FTopologicalVertex>& InVertex);
+		FPoint2D GetTangent2DAt(const TSharedRef<FTopologicalVertex>& InVertex);
+
 		/**
-		 * Project Point on the 3D polyline and return the coordinate of the projected point
+		 * Project Point (2D or 3D) on the polyline (2D or 3D) and return the coordinate of the projected point
 		 */
-		double ProjectPoint(const FPoint& InPointToProject, FPoint& OutProjectedPoint) const
+		template<class PointType>
+		double ProjectPoint(const PointType& InPointToProject, PointType& OutProjectedPoint) const
 		{
 			return Curve->GetCoordinateOfProjectedPoint(Boundary, InPointToProject, OutProjectedPoint);
 		}
@@ -561,6 +586,29 @@ namespace CADKernel
 		 */
 		void ComputeEdge2DProperties(FEdge2DProperties& SlopCharacteristics);
 
+		void GetExtremities(FSurfacicCurveExtremity& Extremities) const
+		{
+			Curve->GetExtremities(Boundary, Extremities);
+		}
+
+
+		// ======   Geometrical Functions   ======
+
+		/**
+		 * Split the edge at input coordinate. The initial edge will be used for the first edge, the second edge will be return.
+		 * @param SplittingCoordinate
+		 * @param NewVertexCoordinate
+		 * @param bKeepStartVertexConnectivity: if true the new edge is connected to endVertex, otherwise the new edge is connected to start vertex
+		 * @param OutNewEdge, the second edge
+		 */
+		TSharedPtr<FTopologicalVertex> SplitAt(double SplittingCoordinate, const FPoint& NewVertexCoordinate, bool bKeepStartVertexConnectivity, TSharedPtr<FTopologicalEdge>& OutNewEdge);
+
+		/**
+		 * Extend the Edge to the NewVertex. 
+		 * NewVertex become the new extremity of the edge. The old vertex is unconnected of the edge.
+		 */
+		bool ExtendTo(bool bStartExtremity, const FPoint2D& NewExtremityCoordinate, TSharedRef<FTopologicalVertex> NewVertex);
+
 		// ======   State Functions   ======
 
 		/**
@@ -572,6 +620,16 @@ namespace CADKernel
 			return FHaveStates::IsDegenerated();
 		}
 
+		virtual void SetAsDegenerated() const override
+		{
+			if (TopologicalLink.IsValid())
+			{
+				printf("WTF");
+			}
+			return FHaveStates::SetAsDegenerated();
+		}
+
+
 		bool IsThinPeak() const
 		{
 			return ((States & EHaveStates::ThinPeak) == EHaveStates::ThinPeak);
@@ -595,7 +653,13 @@ namespace CADKernel
 			return GetTwinsEntityCount() == 1;
 		}
 
-		static TSharedPtr<FTopologicalEdge> CreateEdgeToMerge2Edges(const TArray<TWeakPtr<FTopologicalEdge>>& Edges, const TArray<EOrientation>& edgeDirections);
+		/**
+		 * Merge successive edges of a face in a single edge.
+		 * Edges must not be topological linked to other faces i.e. edges must be border edges
+		 * The merged edges are deleted
+		 * @return TSharedPtr<FTopologicalEdge>() if failed
+		 */
+		static TSharedPtr<FTopologicalEdge> CreateEdgeByMergingEdges(TArray<FOrientedEdge>& Edges, const TSharedRef<FTopologicalVertex> StartVertex, const TSharedRef<FTopologicalVertex> EndVertex);
 	};
 
 	struct CADKERNEL_API FEdge2DProperties
diff --git a/Engine/Source/Runtime/Datasmith/CADKernel/Public/CADKernel/Topo/TopologicalEntity.h b/Engine/Source/Runtime/Datasmith/CADKernel/Public/CADKernel/Topo/TopologicalEntity.h
--- a/Engine/Source/Runtime/Datasmith/CADKernel/Public/CADKernel/Topo/TopologicalEntity.h
+++ b/Engine/Source/Runtime/Datasmith/CADKernel/Public/CADKernel/Topo/TopologicalEntity.h
@@ -19,6 +19,10 @@ namespace CADKernel
 
 		virtual void GetFaces(TArray<TSharedPtr<FTopologicalFace>>& OutFaces) {}
 
+		/**
+		 * Each face of model is set by its orientation. This allow to make oriented mesh and to keep the face orientation in topological function.
+		 * Marker2 of spread face is set. It must be reset after the process
+		 */
 		virtual void SpreadBodyOrientation() {}
 
 		virtual void Serialize(FCADKernelArchive& Ar) override
diff --git a/Engine/Source/Runtime/Datasmith/CADKernel/Public/CADKernel/Topo/TopologicalFace.h b/Engine/Source/Runtime/Datasmith/CADKernel/Public/CADKernel/Topo/TopologicalFace.h
--- a/Engine/Source/Runtime/Datasmith/CADKernel/Public/CADKernel/Topo/TopologicalFace.h
+++ b/Engine/Source/Runtime/Datasmith/CADKernel/Public/CADKernel/Topo/TopologicalFace.h
@@ -36,11 +36,13 @@ namespace CADKernel
 	class CADKERNEL_API FTopologicalFace : public FTopologicalEntity, public FMetadataDictionary
 	{
 		friend class FEntity;
+		//friend class FShell;
 
 	protected:
 
 		TSharedPtr<FSurface> CarrierSurface;
 		TArray<TSharedPtr<FTopologicalLoop>> Loops;
+		TWeakPtr<FShell> HostedBy;
 
 		mutable TCache<FSurfacicBoundary> Boundary;
 
@@ -97,6 +99,8 @@ namespace CADKernel
 			FTopologicalEntity::Serialize(Ar);
 			SerializeIdent(Ar, CarrierSurface);
 			SerializeIdents(Ar, (TArray<TSharedPtr<FEntity>>&) Loops);
+			SerializeIdent(Ar, HostedBy);
+			SerializeMetadata(Ar);
 		}
 
 		virtual void SpawnIdent(FDatabase& Database) override;
@@ -152,6 +156,29 @@ namespace CADKernel
 			}
 		}
 
+		void ResetHost()
+		{
+			HostedBy.Reset();
+		}
+
+		TWeakPtr<FShell>& GetHost()
+		{
+			return HostedBy;
+		}
+
+		void SetHost(TSharedPtr<FShell> Shell)
+		{
+			HostedBy = Shell;
+		}
+
+		void CompleteMetadata()
+		{
+			if (HostedBy.IsValid())
+			{
+				CompleteDictionary((const FMetadataDictionary&)*HostedBy.Pin());
+			}
+		}
+
 		// ======   Loop Functions   ======
 
 		void RemoveLoop(const TSharedPtr<FTopologicalLoop>& Loop);
@@ -190,7 +217,6 @@ namespace CADKernel
 		 */
 		const void Get2DLoopSampling(TArray<TArray<FPoint2D>>& OutLoopSamplings) const;
 
-
 		// ======   Loop edge Functions   ======
 
 		/**
@@ -496,6 +522,10 @@ namespace CADKernel
 		TArray<TSharedPtr<FTopologicalFace>> Faces;
 		int32 BorderEdgeCount = 0;
 		int32 NonManifoldEdgeCount = 0;
+		TWeakPtr<FShell> MainShell;
+		FString MainName;
+		uint32 MainColor;
+		uint32 MainMaterial;
 	};
 
 
diff --git a/Engine/Source/Runtime/Datasmith/CADKernel/Public/CADKernel/Topo/TopologicalLoop.h b/Engine/Source/Runtime/Datasmith/CADKernel/Public/CADKernel/Topo/TopologicalLoop.h
--- a/Engine/Source/Runtime/Datasmith/CADKernel/Public/CADKernel/Topo/TopologicalLoop.h
+++ b/Engine/Source/Runtime/Datasmith/CADKernel/Public/CADKernel/Topo/TopologicalLoop.h
@@ -31,6 +31,12 @@ namespace CADKernel
 			: TOrientedEntity()
 		{
 		}
+
+		bool operator==(const FOrientedEdge& Edge) const
+		{
+			return Entity == Edge.Entity;
+		}
+
 	};
 
 	class CADKERNEL_API FTopologicalLoop : public FTopologicalEntity
@@ -46,7 +52,7 @@ namespace CADKernel
 	protected:
 		TArray<FOrientedEdge> Edges;
 
-		TWeakPtr<FTopologicalFace> Surface;
+		TWeakPtr<FTopologicalFace> Face;
 		bool bExternalLoop;
 
 		FTopologicalLoop(const TArray<TSharedPtr<FTopologicalEdge>>& Edges, const TArray<EOrientation>& EdgeDirections);
@@ -61,25 +67,25 @@ namespace CADKernel
 
 		void SetSurface(TSharedRef<FTopologicalFace> NewDomain)
 		{
-			Surface = NewDomain;
+			Face = NewDomain;
 		}
 
 		void ResetSurface()
 		{
-			Surface.Reset();
+			Face.Reset();
 		}
 
 	public:
 
 		~FTopologicalLoop() = default;
 
-		static TSharedPtr<FTopologicalLoop> Make(const TArray<TSharedPtr<FTopologicalEdge>>& EdgeList, const TArray<EOrientation>& EdgeDirections);
+		static TSharedPtr<FTopologicalLoop> Make(const TArray<TSharedPtr<FTopologicalEdge>>& EdgeList, const TArray<EOrientation>& EdgeDirections, const double GeometricTolerance);
 
 		virtual void Serialize(FCADKernelArchive& Ar) override
 		{
 			FTopologicalEntity::Serialize(Ar);
 			SerializeIdents(Ar, (TArray<TOrientedEntity<FEntity>>&) Edges);
-			SerializeIdent(Ar, Surface);
+			SerializeIdent(Ar, Face);
 			Ar << bExternalLoop;
 		}
 
@@ -125,7 +131,7 @@ namespace CADKernel
 
 		TSharedRef<FTopologicalFace> GetFace() const
 		{
-			return Surface.Pin().ToSharedRef();
+			return Face.Pin().ToSharedRef();
 		}
 
 		void Orient();
@@ -138,9 +144,17 @@ namespace CADKernel
 
 		void ReplaceEdge(TSharedPtr<FTopologicalEdge>& OldEdge, TSharedPtr<FTopologicalEdge>& NewEdge);
 		void ReplaceEdge(TSharedPtr<FTopologicalEdge>& Edge, TArray<TSharedPtr<FTopologicalEdge>>& NewEdges);
+		void ReplaceEdges(TArray<FOrientedEdge>& Candidates, TSharedPtr<FTopologicalEdge>& NewEdge);
+
+		/**
+		 * The Edge is split in two edges : Edge + NewEdge 
+		 * @param bNewEdgeIsFirst == true => StartVertex Connected to Edge, EndVertexConnected to NewEdge
+		 * According to the direction of Edge, if bNewEdgeIsFirst == true, NewEdge is added in the loop after (EOrientation::Front) or before (EOrientation::Back)
+		 */
+		void SplitEdge(TSharedPtr<FTopologicalEdge> Edge, TSharedPtr<FTopologicalEdge> NewEdge, bool bNewEdgeIsFirst);
 
 		void RemoveEdge(TSharedPtr<FTopologicalEdge>& Edge);
-		void ReplaceEdgesWithMergedEdge(TArray<TSharedPtr<FTopologicalEdge>>& OldEdges, TSharedPtr<FTopologicalVertex>& MiddleVertex, TSharedPtr<FTopologicalEdge>& NewEdge);
+		//void ReplaceEdgesWithMergedEdge(TArray<TSharedPtr<FTopologicalEdge>>& OldEdges, TSharedPtr<FTopologicalVertex>& MiddleVertex, TSharedPtr<FTopologicalEdge>& NewEdge);
 
 		EOrientation GetDirection(TSharedPtr<FTopologicalEdge>& Edge, bool bAllowLinkedEdge = false) const;
 
@@ -173,6 +187,6 @@ namespace CADKernel
 
 		void ComputeBoundaryProperties(const TArray<int32>& StartSideIndex, TArray<FEdge2DProperties>& OutSideProperties) const;
 
-		void EnsureLogicalClosing();
+		void EnsureLogicalClosing(const double GeometricTolerance);
 	};
 }
diff --git a/Engine/Source/Runtime/Datasmith/CADKernel/Public/CADKernel/Topo/TopologicalVertex.h b/Engine/Source/Runtime/Datasmith/CADKernel/Public/CADKernel/Topo/TopologicalVertex.h
--- a/Engine/Source/Runtime/Datasmith/CADKernel/Public/CADKernel/Topo/TopologicalVertex.h
+++ b/Engine/Source/Runtime/Datasmith/CADKernel/Public/CADKernel/Topo/TopologicalVertex.h
@@ -40,19 +40,11 @@ namespace CADKernel
 		virtual void Serialize(FCADKernelArchive& Ar) override
 		{
 			TLinkable<FTopologicalVertex, FVertexLink>::Serialize(Ar);
-			Ar << Coordinates;
+			Ar.Serialize(Coordinates);
 			SerializeIdents(Ar, (TArray<TSharedPtr<FEntity>>&) ConnectedEdges);
 		}
 
 		virtual void SpawnIdent(FDatabase& Database) override;
-		//{
-		//	if (!FEntity::SetId(Database))
-		//	{
-		//		return;
-		//	}
-		//	TopologicalLink->SpawnIdent(Database);
-		//	((TSharedPtr<FEntity>) Mesh)->SpawnIdent(Database);
-		//}
 
 		virtual void ResetMarkersRecursively() override
 		{
@@ -137,6 +129,19 @@ namespace CADKernel
 
 		void Link(TSharedRef<FTopologicalVertex> InEntity);
 
+		void UnlinkTo(TSharedRef<FTopologicalVertex> Entity);
+
+		virtual void RemoveFromLink() override
+		{
+			if (TopologicalLink.IsValid())
+			{
+				TopologicalLink->RemoveEntity(StaticCastSharedRef<FTopologicalVertex>(AsShared()));
+				TopologicalLink->ComputeBarycenter();
+				TopologicalLink.Reset();
+			}
+		}
+
+
 		void DeleteIfIsolated()
 		{
 			if (ConnectedEdges.Num() == 0)
@@ -156,11 +161,57 @@ namespace CADKernel
 		void AddConnectedEdge(TSharedRef<FTopologicalEdge> Edge);
 		void RemoveConnectedEdge(TSharedRef<FTopologicalEdge> Edge);
 
-		const TArray<TWeakPtr<FTopologicalEdge>>& GetConnectedEdges() const
+		/**
+		 * Mandatory: to browse all the connected edges, you have to browse the connected edges of all the twin vertices
+		 * for (TWeakPtr<FTopologicalVertex> TwinVertex : Vertex->GetTwinsEntities())
+		 * {
+		 *    for (TWeakPtr<FTopologicalEdge> ConnectedEdge : TwinVertex.Pin()->GetDirectConnectedEdges())
+		 *    {
+		 *       ...
+		 *    }
+		 *  }
+		 */
+		const TArray<TWeakPtr<FTopologicalEdge>>& GetDirectConnectedEdges() const 
 		{
 			return ConnectedEdges;
 		}
 
+		const void GetConnectedEdges(TArray<TWeakPtr<FTopologicalEdge>>& OutConnectedEdges) const
+		{
+			if (!TopologicalLink.IsValid())
+			{
+				OutConnectedEdges = ConnectedEdges;
+			}
+			else
+			{
+				OutConnectedEdges.Reserve(100);
+				for (const TWeakPtr<FTopologicalVertex>& Vertex : GetLink()->GetTwinsEntities())
+				{
+					OutConnectedEdges.Append(Vertex.Pin()->ConnectedEdges);
+				}
+			}
+		}
+
+		const int32 ConnectedEdgeCount()
+		{
+			if (!TopologicalLink.IsValid())
+			{
+				return ConnectedEdges.Num();
+			}
+			else
+			{
+				int32 Count = 0;
+				for (const TWeakPtr<FTopologicalVertex>& Vertex : GetLink()->GetTwinsEntities())
+				{
+					Count +=Vertex.Pin()->ConnectedEdges.Num();
+				}
+				return Count;
+			}
+		}
+
+		/**
+		 * 
+		 */
 		void GetConnectedEdges(TSharedPtr<FTopologicalVertex> OtherVertex, TArray<TSharedPtr<FTopologicalEdge>>& Edges) const;
 	};
 
diff --git a/Engine/Source/Runtime/Datasmith/CADKernel/Public/CADKernel/UI/Display.h b/Engine/Source/Runtime/Datasmith/CADKernel/Public/CADKernel/UI/Display.h
--- a/Engine/Source/Runtime/Datasmith/CADKernel/Public/CADKernel/UI/Display.h
+++ b/Engine/Source/Runtime/Datasmith/CADKernel/Public/CADKernel/UI/Display.h
@@ -116,6 +116,11 @@ namespace CADKernel
 		DrawPoint(Point, Property);
 	}
 
+	CADKERNEL_API void DisplayProductTree(const TSharedPtr<FEntity>& RootId);
+	CADKERNEL_API void DisplayProductTree(const TSharedPtr<FModel>& Model);
+	CADKERNEL_API void DisplayProductTree(const TSharedPtr<FBody>& Body);
+	CADKERNEL_API void DisplayProductTree(const TSharedPtr<FShell>& Shell);
+
 	CADKERNEL_API void DisplayAABB(const FAABB& aabb, FIdent Ident = 0);
 	CADKERNEL_API void DisplayAABB2D(const FAABB2D& aabb, FIdent Ident = 0);
 
diff --git a/Engine/Source/Runtime/Datasmith/CADKernel/Public/CADKernel/UI/Visu.h b/Engine/Source/Runtime/Datasmith/CADKernel/Public/CADKernel/UI/Visu.h
--- a/Engine/Source/Runtime/Datasmith/CADKernel/Public/CADKernel/UI/Visu.h
+++ b/Engine/Source/Runtime/Datasmith/CADKernel/Public/CADKernel/UI/Visu.h
@@ -34,6 +34,8 @@ namespace CADKernel
 		GreenCurve,				
 		PinkPoint,
 		PinkCurve,
+		OrangePoint,
+		OrangeCurve,
 		Last,
 	} ;
 
