Issue 78710: fix one compile warnings

diff --git a/Engine/Source/Runtime/Datasmith/CADKernel/Private/CADKernel/Topo/Joiner.cpp b/Engine/Source/Runtime/Datasmith/CADKernel/Private/CADKernel/Topo/Joiner.cpp
--- a/Engine/Source/Runtime/Datasmith/CADKernel/Private/CADKernel/Topo/Joiner.cpp
+++ b/Engine/Source/Runtime/Datasmith/CADKernel/Private/CADKernel/Topo/Joiner.cpp
@@ -533,7 +533,7 @@ void FJoiner::StitchParallelEdges(TArray<TSharedPtr<FTopologicalVertex>>& Vertic
 				bool bSecondEdgeBorder = Edge->IsBorder();
 
 				// Process only if at least one edge is Border
-				if (!bSecondEdgeBorder && !bSecondEdgeBorder)
+				if (!bFirstEdgeBorder && !bSecondEdgeBorder)
 				{
 					continue;
 				}
