CADKernel Source Code v0.03
The review of the source code is in progress.

diff --git a/Engine/Plugins/Enterprise/DatasmithCADImporter/Source/CADTools/Private/CADKernelTools.cpp b/Engine/Plugins/Enterprise/DatasmithCADImporter/Source/CADTools/Private/CADKernelTools.cpp
--- a/Engine/Plugins/Enterprise/DatasmithCADImporter/Source/CADTools/Private/CADKernelTools.cpp
+++ b/Engine/Plugins/Enterprise/DatasmithCADImporter/Source/CADTools/Private/CADKernelTools.cpp
@@ -32,23 +32,14 @@ uint32 FCADKernelTools::GetFaceTessellation(const TSharedRef<CADKernel::FFaceMes
 
 	CADLibrary::FTessellationData& Tessellation = OutBodyMesh.Faces.Emplace_GetRef();
 
-	const TSharedRef<CADKernel::FMetadataDictionary>& HaveMetadata = StaticCastSharedRef<CADKernel::FMetadataDictionary>(StaticCastSharedRef<CADKernel::FTopologicalFace>(FaceMesh->GetGeometricEntity()));
-	Tessellation.PatchId = HaveMetadata->GetPatchId();
+	const CADKernel::FTopologicalFace& HaveMetadata = (const CADKernel::FTopologicalFace&) FaceMesh->GetGeometricEntity();
+	Tessellation.PatchId = HaveMetadata.GetPatchId();
 
 	Tessellation.PositionIndices = FaceMesh->VerticesGlobalIndex;
 	Tessellation.VertexIndices = FaceMesh->TrianglesVerticesIndex;
 
-	Tessellation.NormalArray.Reserve(FaceMesh->Normals.Num());
-	for (const CADKernel::FPoint& Normal : FaceMesh->Normals)
-	{
-		Tessellation.NormalArray.Emplace((float)Normal.X, (float)Normal.Y, (float)Normal.Z);
-	}
-
-	Tessellation.TexCoordArray.Reserve(FaceMesh->UVMap.Num());
-	for (const CADKernel::FPoint2D& TexCoord : FaceMesh->UVMap)
-	{
-		Tessellation.TexCoordArray.Emplace((float)TexCoord.U, (float)TexCoord.V);
-	}
+	Tessellation.NormalArray = FaceMesh->Normals;
+	Tessellation.TexCoordArray = FaceMesh->UVMap;
 
 	return Tessellation.VertexIndices.Num() / 3;
 }
diff --git a/Engine/Source/Runtime/Datasmith/CADKernel/Private/CADKernel/Core/Session.cpp b/Engine/Source/Runtime/Datasmith/CADKernel/Private/CADKernel/Core/Session.cpp
--- a/Engine/Source/Runtime/Datasmith/CADKernel/Private/CADKernel/Core/Session.cpp
+++ b/Engine/Source/Runtime/Datasmith/CADKernel/Private/CADKernel/Core/Session.cpp
@@ -72,3 +72,9 @@ void FSession::AddDatabase(const TArray<uint8>& InRawData)
 	FCADKernelArchive Archive = FCADKernelArchive(*this, InRawData);
 	Database.Deserialize(Archive);
 }
+
+void FSession::SetGeometricTolerance(double NewTolerance)
+{
+	ensure(Database.GetModel()->EntityCount() == 0);
+	GeometricTolerance = NewTolerance;
+}
\ No newline at end of file
diff --git a/Engine/Source/Runtime/Datasmith/CADKernel/Private/CADKernel/Geo/Curves/NURBSCurve.cpp b/Engine/Source/Runtime/Datasmith/CADKernel/Private/CADKernel/Geo/Curves/NURBSCurve.cpp
--- a/Engine/Source/Runtime/Datasmith/CADKernel/Private/CADKernel/Geo/Curves/NURBSCurve.cpp
+++ b/Engine/Source/Runtime/Datasmith/CADKernel/Private/CADKernel/Geo/Curves/NURBSCurve.cpp
@@ -34,6 +34,17 @@ FNURBSCurve::FNURBSCurve(int32 InDegre, const TArray<double>& InNodalVector, con
 	Finalize();
 }
 
+FNURBSCurve::FNURBSCurve(FNurbsCurveData& NurbsCurveData)
+	: FCurve(NurbsCurveData.Dimension)
+	, Degree(NurbsCurveData.Degree)
+	, NodalVector(NurbsCurveData.NodalVector)
+	, Weights(NurbsCurveData.Weights)
+	, Poles(NurbsCurveData.Poles)
+	, bIsRational(NurbsCurveData.bIsRational)
+{
+	Finalize();
+}
+
 FNURBSCurve::FNURBSCurve(const TSharedRef<FNURBSCurve>& Nurbs)
 	: FNURBSCurve(Nurbs->Degree, Nurbs->NodalVector, Nurbs->Poles, Nurbs->Weights, Nurbs->Dimension)
 {
diff --git a/Engine/Source/Runtime/Datasmith/CADKernel/Private/CADKernel/Geo/Sampling/SurfacicSampling.h b/Engine/Source/Runtime/Datasmith/CADKernel/Private/CADKernel/Geo/Sampling/SurfacicSampling.h
--- a/Engine/Source/Runtime/Datasmith/CADKernel/Private/CADKernel/Geo/Sampling/SurfacicSampling.h
+++ b/Engine/Source/Runtime/Datasmith/CADKernel/Private/CADKernel/Geo/Sampling/SurfacicSampling.h
@@ -12,7 +12,7 @@ namespace CADKernel
 		bool bWithNormals = false;
 		TArray<FPoint2D> Points2D;
 		TArray<FPoint> Points3D;
-		TArray<FPoint> Normals;
+		TArray<FVector> Normals;
 
 		int32 Count()
 		{
@@ -41,7 +41,7 @@ namespace CADKernel
 
 		void NormalizeNormals()
 		{
-			for (FPoint& Normal : Normals)
+			for (FVector& Normal : Normals)
 			{
 				Normal.Normalize();
 			}
diff --git a/Engine/Source/Runtime/Datasmith/CADKernel/Private/CADKernel/Geo/Surfaces/NURBSSurface.cpp b/Engine/Source/Runtime/Datasmith/CADKernel/Private/CADKernel/Geo/Surfaces/NURBSSurface.cpp
--- a/Engine/Source/Runtime/Datasmith/CADKernel/Private/CADKernel/Geo/Surfaces/NURBSSurface.cpp
+++ b/Engine/Source/Runtime/Datasmith/CADKernel/Private/CADKernel/Geo/Surfaces/NURBSSurface.cpp
@@ -16,7 +16,7 @@ TSharedPtr<FEntityGeom> FNURBSSurface::ApplyMatrix(const FMatrixH& InMatrix) con
 		TransformedPoles.Emplace(InMatrix.Multiply(Pole));
 	}
 
-	return FEntity::MakeShared<FNURBSSurface>(Tolerance3D, PoleUNum, PoleVNum, UDegree, VDegree, UNodalVector, VNodalVector, TransformedPoles, Weights);
+	return FEntity::MakeShared<FNURBSSurface>(Tolerance3D, PoleUCount, PoleVCount, UDegree, VDegree, UNodalVector, VNodalVector, TransformedPoles, Weights);
 }
 
 #ifdef CADKERNEL_DEV
@@ -25,7 +25,7 @@ FInfoEntity& FNURBSSurface::GetInfo(FInfoEntity& Info) const
 	return FSurface::GetInfo(Info)
 		.Add(TEXT("Degre"), UDegree, VDegree)
 		.Add(TEXT("Is Rational"), bIsRational)
-		.Add(TEXT("Poles Num"), PoleUNum, PoleVNum)
+		.Add(TEXT("Poles Num"), PoleUCount, PoleVCount)
 		.Add(TEXT("Nodal Vector U"), UNodalVector)
 		.Add(TEXT("Nodal Vector V"), VNodalVector)
 		.Add(TEXT("Poles"), Poles)
@@ -88,4 +88,89 @@ void FNURBSSurface::Finalize()
 	SetMinToleranceIso();
 }
 
+void FNURBSSurface::FillNurbs(FNurbsSurfaceHomogeneousData& NurbsData)
+{
+	bIsRational = NurbsData.bIsRational;
+
+	PoleUCount = NurbsData.PoleUCount;
+	PoleVCount = NurbsData.PoleVCount;
+
+	UDegree = NurbsData.UDegree;
+	VDegree = NurbsData.VDegree;
+
+	Swap(UNodalVector, NurbsData.UNodalVector);
+	Swap(VNodalVector, NurbsData.VNodalVector);
+
+	int32 PoleCount = PoleUCount * PoleVCount;
+	
+	double* RawPolesPtr = NurbsData.HomogeneousPoles.GetData();
+	const int32 Dimension = bIsRational ? 4 : 3;
+	Poles.Reserve(PoleCount);
+
+	const int32 OffSet = Dimension * PoleVCount;
+	for (int32 Vndex = 0; Vndex < PoleVCount; ++Vndex)
+	{
+		int32 RawPolesOffset = Dimension * Vndex;
+		for (int32 Undex = 0; Undex < PoleUCount; ++Undex, RawPolesOffset += OffSet)
+		{
+			Poles.Emplace(RawPolesPtr + RawPolesOffset);
+		}
+	}
+
+	if(bIsRational)
+	{
+		Weights.Reserve(PoleCount);
+		for (int32 Vndex = 3; Vndex < (PoleVCount * Dimension); Vndex += Dimension)
+		{
+			int32 RawPolesOffset = Vndex;
+			for (int32 Undex = 0; Undex < PoleUCount; ++Undex, RawPolesOffset += OffSet)
+			{
+				Weights.Emplace(RawPolesPtr[RawPolesOffset]);
+			}
+		}
+	}
+	Finalize();
+}
+
+void FNURBSSurface::FillNurbs(FNurbsSurfaceData& NurbsData)
+{
+	bIsRational = (NurbsData.Weights.Num() > 0);
+
+	PoleUCount = NurbsData.PoleUCount;
+	PoleVCount = NurbsData.PoleVCount;
+
+	UDegree = NurbsData.UDegree;
+	VDegree = NurbsData.VDegree;
+
+	Swap(UNodalVector, NurbsData.UNodalVector);
+	Swap(VNodalVector, NurbsData.VNodalVector);
+
+	int32 PoleCount = PoleUCount * PoleVCount;
+
+	Poles.Reserve(PoleCount);
+	for (int32 Vndex = 0; Vndex < PoleVCount; ++Vndex)
+	{
+		int32 Index = Vndex;
+		for (int32 Undex = 0; Undex < PoleUCount; ++Undex, Index += PoleVCount)
+		{
+			Poles.Emplace(NurbsData.Poles[Index]);
+		}
+	}
+
+	if (bIsRational)
+	{
+		Weights.Reserve(PoleCount);
+		for (int32 Vndex = 0; Vndex < PoleVCount; ++Vndex)
+		{
+			int32 Index = Vndex;
+			for (int32 Undex = 0; Undex < PoleUCount; ++Undex, Index += PoleVCount)
+			{
+				Weights.Emplace(NurbsData.Weights[Index]);
+			}
+		}
+	}
+
+	Finalize();
+}
+
 
diff --git a/Engine/Source/Runtime/Datasmith/CADKernel/Private/CADKernel/Geo/Surfaces/Surface.cpp b/Engine/Source/Runtime/Datasmith/CADKernel/Private/CADKernel/Geo/Surfaces/Surface.cpp
--- a/Engine/Source/Runtime/Datasmith/CADKernel/Private/CADKernel/Geo/Surfaces/Surface.cpp
+++ b/Engine/Source/Runtime/Datasmith/CADKernel/Private/CADKernel/Geo/Surfaces/Surface.cpp
@@ -102,7 +102,7 @@ namespace CADKernel
 			{
 				Polyline.Normals.Emplace(Point.GradientU ^ Point.GradientV);
 			}
-			for (FPoint Normal : Polyline.Normals)
+			for (FVector Normal : Polyline.Normals)
 			{
 				Normal.Normalize();
 			}
@@ -141,13 +141,13 @@ namespace CADKernel
 			Polyline.Normals.Emplace(Point.GradientU ^ Point.GradientV);
 		}
 
-		for (FPoint& Normal : Polyline.Normals)
+		for (FVector& Normal : Polyline.Normals)
 		{
 			Normal.Normalize();
 		}
 	}
 
-	void FSurface::EvaluateNormals(const TArray<FPoint2D>& InPoints2D, TArray<FPoint>& Normals) const
+	void FSurface::EvaluateNormals(const TArray<FPoint2D>& InPoints2D, TArray<FVector>& Normals) const
 	{
 		int32 DerivativeOrder = 1;
 		TArray<FSurfacicPoint> Point3D;
@@ -160,7 +160,7 @@ namespace CADKernel
 			Normals.Emplace(Point.GradientU ^ Point.GradientV);
 		}
 
-		for (FPoint& Normal : Normals)
+		for (FVector& Normal : Normals)
 		{
 			Normal.Normalize();
 		}
diff --git a/Engine/Source/Runtime/Datasmith/CADKernel/Private/CADKernel/Geo/Surfaces/TorusSurface.cpp b/Engine/Source/Runtime/Datasmith/CADKernel/Private/CADKernel/Geo/Surfaces/TorusSurface.cpp
--- a/Engine/Source/Runtime/Datasmith/CADKernel/Private/CADKernel/Geo/Surfaces/TorusSurface.cpp
+++ b/Engine/Source/Runtime/Datasmith/CADKernel/Private/CADKernel/Geo/Surfaces/TorusSurface.cpp
@@ -116,8 +116,8 @@ void FTorusSurface::EvaluatePointGrid(const FCoordinateGrid& Coordinates, FSurfa
 
 	if (bComputeNormals)
 	{
-		FPoint Center = Matrix.Column(3);
-		for (FPoint& Normal : OutPoints.Normals)
+		FVector Center = Matrix.Column(3);
+		for (FVector& Normal : OutPoints.Normals)
 		{
 			 Normal = Matrix.PointRotation(Normal, Center);
 		}
diff --git a/Engine/Source/Runtime/Datasmith/CADKernel/Private/CADKernel/Math/MatrixH.cpp b/Engine/Source/Runtime/Datasmith/CADKernel/Private/CADKernel/Math/MatrixH.cpp
--- a/Engine/Source/Runtime/Datasmith/CADKernel/Private/CADKernel/Math/MatrixH.cpp
+++ b/Engine/Source/Runtime/Datasmith/CADKernel/Private/CADKernel/Math/MatrixH.cpp
@@ -168,6 +168,19 @@ namespace CADKernel
 		return Result;
 	}
 
+	FVector FMatrixH::PointRotation(const FVector& PointToRotate, const FVector& Origin) const
+	{
+		FVector Result = Origin;
+		for (int32 Index = 0; Index < 3; Index++)
+		{
+			for (int32 Jndex = 0; Jndex < 3; Jndex++)
+			{
+				Result[Index] += Get(Index, Jndex) * (PointToRotate[Jndex] - Origin[Jndex]);
+			}
+		}
+		return Result;
+	}
+
 	FPoint2D FMatrixH::PointRotation(const FPoint2D& PointToRotate, const FPoint2D& Origin) const
 	{
 		FPoint2D Result = Origin;
diff --git a/Engine/Source/Runtime/Datasmith/CADKernel/Private/CADKernel/Mesh/Criteria/CriteriaGrid.cpp b/Engine/Source/Runtime/Datasmith/CADKernel/Private/CADKernel/Mesh/Criteria/CriteriaGrid.cpp
--- a/Engine/Source/Runtime/Datasmith/CADKernel/Private/CADKernel/Mesh/Criteria/CriteriaGrid.cpp
+++ b/Engine/Source/Runtime/Datasmith/CADKernel/Private/CADKernel/Mesh/Criteria/CriteriaGrid.cpp
@@ -11,7 +11,7 @@
 using namespace CADKernel;
 
 
-const FPoint& FCriteriaGrid::GetPoint(int32 UIndex, int32 VIndex, bool bIsInternalU, bool bIsInternalV, FPoint* OutNormal) const
+const FPoint& FCriteriaGrid::GetPoint(int32 UIndex, int32 VIndex, bool bIsInternalU, bool bIsInternalV, FVector* OutNormal) const
 {
 	int32 Index = GetIndex(UIndex, VIndex, bIsInternalU, bIsInternalV);
 	ensureCADKernel(Grid.Points3D.IsValidIndex(Index));
diff --git a/Engine/Source/Runtime/Datasmith/CADKernel/Private/CADKernel/Mesh/Criteria/CriteriaGrid.h b/Engine/Source/Runtime/Datasmith/CADKernel/Private/CADKernel/Mesh/Criteria/CriteriaGrid.h
--- a/Engine/Source/Runtime/Datasmith/CADKernel/Private/CADKernel/Mesh/Criteria/CriteriaGrid.h
+++ b/Engine/Source/Runtime/Datasmith/CADKernel/Private/CADKernel/Mesh/Criteria/CriteriaGrid.h
@@ -35,7 +35,7 @@ namespace CADKernel
 			return TrueVIndex * TrueUcoorindateCount + TrueUIndex;
 		}
 
-		const FPoint& GetPoint(int32 UIndex, int32 VIndex, bool bIsInterU, bool bIsInterV, FPoint* normal = nullptr) const;
+		const FPoint& GetPoint(int32 UIndex, int32 VIndex, bool bIsInterU, bool bIsInterV, FVector* normal = nullptr) const;
 		void Init();
 
 	public:
@@ -68,27 +68,27 @@ namespace CADKernel
 			return Grid.Points3D;
 		}
 
-		const TArray<FPoint>& GetGridNormals() const
+		const TArray<FVector>& GetGridNormals() const
 		{
 			return Grid.Normals;
 		}
 
-		const FPoint& GetPoint(int32 iU, int32 iV, FPoint* normal = nullptr) const
+		const FPoint& GetPoint(int32 iU, int32 iV, FVector* normal = nullptr) const
 		{
 			return GetPoint(iU, iV, false, false, normal);
 		}
 
-		const FPoint& GetIntermediateU(int32 iU, int32 iV, FPoint* normal = nullptr) const
+		const FPoint& GetIntermediateU(int32 iU, int32 iV, FVector* normal = nullptr) const
 		{
 			return GetPoint(iU, iV, true, false, normal);
 		}
 
-		const FPoint& GetIntermediateV(int32 iU, int32 iV, FPoint* normal = nullptr) const
+		const FPoint& GetIntermediateV(int32 iU, int32 iV, FVector* normal = nullptr) const
 		{
 			return GetPoint(iU, iV, false, true, normal);
 		}
 
-		const FPoint& GetIntermediateUV(int32 iU, int32 iV, FPoint* normal = nullptr) const
+		const FPoint& GetIntermediateUV(int32 iU, int32 iV, FVector* normal = nullptr) const
 		{
 			return GetPoint(iU, iV, true, true, normal);
 		}
diff --git a/Engine/Source/Runtime/Datasmith/CADKernel/Private/CADKernel/Mesh/Meshers/IsoTriangulator.cpp b/Engine/Source/Runtime/Datasmith/CADKernel/Private/CADKernel/Mesh/Meshers/IsoTriangulator.cpp
--- a/Engine/Source/Runtime/Datasmith/CADKernel/Private/CADKernel/Mesh/Meshers/IsoTriangulator.cpp
+++ b/Engine/Source/Runtime/Datasmith/CADKernel/Private/CADKernel/Mesh/Meshers/IsoTriangulator.cpp
@@ -319,12 +319,14 @@ void FIsoTriangulator::FillMeshNodes()
 
 	for (FLoopNode& Node : LoopNodes)
 	{
-		Mesh->UVMap.Emplace(Node.Get2DPoint(EGridSpace::Scaled, Grid));
+		const FPoint2D& UVCoordinate = Node.Get2DPoint(EGridSpace::Scaled, Grid);
+		Mesh->UVMap.Emplace(UVCoordinate.U, UVCoordinate.V);
 	}
 
 	for (FIsoInnerNode& Node : InnerNodes)
 	{
-		Mesh->UVMap.Emplace(Node.Get2DPoint(EGridSpace::Scaled, Grid));
+		const FPoint2D& UVCoordinate = Node.Get2DPoint(EGridSpace::Scaled, Grid);
+		Mesh->UVMap.Emplace(UVCoordinate.U, UVCoordinate.V);
 	}
 }
 
@@ -2307,7 +2309,7 @@ void FIsoTriangulator::MeshCycle(const EGridSpace Space, const TArray<FIsoSegmen
 				const FPoint& NodePoint3D = Node->Get3DPoint(Grid);
 				FPoint NodeNormal;
 				double PointCriteria = FMath::Abs(CotangentCriteria(StartPoint3D, EndPoint3D, NodePoint3D, NodeNormal));
-				double CosAngle = FMath::Abs(NodeNormal.ComputeCosinus(Node->GetNormal(Grid)));
+				double CosAngle = FMath::Abs(ComputeCosinus(NodeNormal, Node->GetNormal(Grid)));
 
 				// the criteria is weighted according to the cosinus of the angle between the normal of the candidate triangle and the normal at the tested point
 				if (CosAngle > SMALL_NUMBER)
diff --git a/Engine/Source/Runtime/Datasmith/CADKernel/Private/CADKernel/Mesh/Meshers/IsoTriangulator/IsoNode.h b/Engine/Source/Runtime/Datasmith/CADKernel/Private/CADKernel/Mesh/Meshers/IsoTriangulator/IsoNode.h
--- a/Engine/Source/Runtime/Datasmith/CADKernel/Private/CADKernel/Mesh/Meshers/IsoTriangulator/IsoNode.h
+++ b/Engine/Source/Runtime/Datasmith/CADKernel/Private/CADKernel/Mesh/Meshers/IsoTriangulator/IsoNode.h
@@ -212,7 +212,7 @@ namespace CADKernel
 		virtual const FPoint2D& Get2DPoint(EGridSpace Space, const FGrid& Grid) const = 0;
 
 		virtual const FPoint& Get3DPoint(const FGrid& Grid) const = 0;
-		virtual const FPoint& GetNormal(const FGrid& Grid) const = 0;
+		virtual const FVector& GetNormal(const FGrid& Grid) const = 0;
 
 		/**
 		 * Only for display purpose as it return a copy of the point
@@ -280,7 +280,7 @@ namespace CADKernel
 			return Grid.GetLoops3D()[LoopIndex][Index];
 		}
 
-		virtual const FPoint& GetNormal(const FGrid& Grid) const override
+		virtual const FVector& GetNormal(const FGrid& Grid) const override
 		{
 			return Grid.GetLoopNormals()[LoopIndex][Index];
 		}
@@ -427,7 +427,7 @@ namespace CADKernel
 			return Grid.GetInner3DPoint(Index);
 		}
 
-		virtual const FPoint& GetNormal(const FGrid& Grid) const override
+		virtual const FVector& GetNormal(const FGrid& Grid) const override
 		{
 			return Grid.GetPointNormal(Index);
 		}
diff --git a/Engine/Source/Runtime/Datasmith/CADKernel/Private/CADKernel/Mesh/Meshers/ParametricMesher.cpp b/Engine/Source/Runtime/Datasmith/CADKernel/Private/CADKernel/Mesh/Meshers/ParametricMesher.cpp
--- a/Engine/Source/Runtime/Datasmith/CADKernel/Private/CADKernel/Mesh/Meshers/ParametricMesher.cpp
+++ b/Engine/Source/Runtime/Datasmith/CADKernel/Private/CADKernel/Mesh/Meshers/ParametricMesher.cpp
@@ -674,11 +674,11 @@ void FParametricMesher::GenerateEdgeElements(TSharedRef<FTopologicalEdge> Edge)
 	int32 StartVertexNodeIndex = Edge->GetStartVertex()->GetOrCreateMesh(GetMeshModel())->GetMesh();
 	int32 EndVertexNodeIndex = Edge->GetEndVertex()->GetOrCreateMesh(GetMeshModel())->GetMesh();
 
-	TSharedRef<FTopologicalEdge> ActiveEdge = StaticCastSharedRef<FTopologicalEdge>(EdgeMesh->GetGeometricEntity());
+	const FTopologicalEdge& ActiveEdge = (const FTopologicalEdge&) EdgeMesh->GetGeometricEntity();
 
 	TArray<double> CuttingPointCoordinates;
-	CuttingPointCoordinates.Reserve(ActiveEdge->GetCuttingPoints().Num());
-	for (const FCuttingPoint& CuttingPoint : ActiveEdge->GetCuttingPoints())
+	CuttingPointCoordinates.Reserve(ActiveEdge.GetCuttingPoints().Num());
+	for (const FCuttingPoint& CuttingPoint : ActiveEdge.GetCuttingPoints())
 	{
 		CuttingPointCoordinates.Add(CuttingPoint.Coordinate);
 	}
@@ -688,7 +688,7 @@ void FParametricMesher::GenerateEdgeElements(TSharedRef<FTopologicalEdge> Edge)
 	CuttingPointCoordinates.Pop();
 
 	TArray<FPoint>& Coordinates = EdgeMesh->GetNodeCoordinates();
-	ActiveEdge->ApproximatePoints(CuttingPointCoordinates, Coordinates);
+	ActiveEdge.ApproximatePoints(CuttingPointCoordinates, Coordinates);
 
 	EdgeMesh->RegisterCoordinates();
 	EdgeMesh->Mesh(StartVertexNodeIndex, EndVertexNodeIndex);
diff --git a/Engine/Source/Runtime/Datasmith/CADKernel/Private/CADKernel/Mesh/Structure/FaceMesh.cpp b/Engine/Source/Runtime/Datasmith/CADKernel/Private/CADKernel/Mesh/Structure/FaceMesh.cpp
--- a/Engine/Source/Runtime/Datasmith/CADKernel/Private/CADKernel/Mesh/Structure/FaceMesh.cpp
+++ b/Engine/Source/Runtime/Datasmith/CADKernel/Private/CADKernel/Mesh/Structure/FaceMesh.cpp
@@ -13,8 +13,9 @@ using namespace CADKernel;
 
 void FFaceMesh::GetNodeIdToCoordinates(TMap<int32, const FPoint*>& NodeIdToCoordinates) const
 {
-	const TSharedRef<FTopologicalFace> Face = StaticCastSharedRef<FTopologicalFace>(GetGeometricEntity());
-	for (const TSharedPtr<FTopologicalLoop>& Loop : Face->GetLoops())
+	const FTopologicalFace& Face = (const FTopologicalFace&) GetGeometricEntity();
+
+	for (const TSharedPtr<FTopologicalLoop>& Loop : Face.GetLoops())
 	{
 		const TArray<FOrientedEdge>& Edges = Loop->GetEdges();
 
@@ -60,7 +61,7 @@ void FFaceMesh::InverseOrientation()
 		Swap(TrianglesVerticesIndex[Index], TrianglesVerticesIndex[Index+1]);
 	}
 
-	for (FPoint& Normal : Normals)
+	for (FVector& Normal : Normals)
 	{
 		Normal *= -1.;
 	}
diff --git a/Engine/Source/Runtime/Datasmith/CADKernel/Private/CADKernel/Mesh/Structure/Grid.cpp b/Engine/Source/Runtime/Datasmith/CADKernel/Private/CADKernel/Mesh/Structure/Grid.cpp
--- a/Engine/Source/Runtime/Datasmith/CADKernel/Private/CADKernel/Mesh/Structure/Grid.cpp
+++ b/Engine/Source/Runtime/Datasmith/CADKernel/Private/CADKernel/Mesh/Structure/Grid.cpp
@@ -1095,7 +1095,7 @@ namespace CADKernel
 			TArray<FPoint>& Loop3D = FaceLoops3D.Emplace_GetRef();
 			Loop3D.Reserve(LoopNodeCount);
 
-			TArray<FPoint>& LoopNormals = NormalsOfFaceLoops.Emplace_GetRef();
+			TArray<FVector>& LoopNormals = NormalsOfFaceLoops.Emplace_GetRef();
 			LoopNormals.Reserve(LoopNodeCount);
 
 			TArray<int32>& LoopIds = NodeIdsOfFaceLoops.Emplace_GetRef();
diff --git a/Engine/Source/Runtime/Datasmith/CADKernel/Private/CADKernel/Mesh/Structure/Grid.h b/Engine/Source/Runtime/Datasmith/CADKernel/Private/CADKernel/Mesh/Structure/Grid.h
--- a/Engine/Source/Runtime/Datasmith/CADKernel/Private/CADKernel/Mesh/Structure/Grid.h
+++ b/Engine/Source/Runtime/Datasmith/CADKernel/Private/CADKernel/Mesh/Structure/Grid.h
@@ -88,7 +88,7 @@ namespace CADKernel
 		/**
 		 * Surface Normal at each boundary nodes
 		 */
-		TArray<TArray<FPoint>> NormalsOfFaceLoops;
+		TArray<TArray<FVector>> NormalsOfFaceLoops;
 
 		TArray<TArray<int32>> NodeIdsOfFaceLoops;
 
@@ -146,7 +146,7 @@ namespace CADKernel
 		/**
 		 * Surface Normal at each inner nodes
 		 */
-		TArray<FPoint> Normals;
+		TArray<FVector> Normals;
 
 	public:
 		FGridChronos Chronos;
@@ -372,7 +372,7 @@ namespace CADKernel
 		/**
 		 * @return the normal of the surface at the point at the Index of the grid
 		 */
-		const FPoint& GetPointNormal(int32 IndexU, int32 IndexV) const
+		const FVector& GetPointNormal(int32 IndexU, int32 IndexV) const
 		{
 			return Normals[GobalIndex(IndexU, IndexV)];
 		}
@@ -380,7 +380,7 @@ namespace CADKernel
 		/**
 		 * @return the normal of the surface at the point at the Index of the grid
 		 */
-		const FPoint& GetPointNormal(int32 Index) const
+		const FVector& GetPointNormal(int32 Index) const
 		{
 			return Normals[Index];
 		}
@@ -432,7 +432,7 @@ namespace CADKernel
 		/**
 		 * @return the array of normal of the points of the grid
 		 */
-		TArray<FPoint>& GetNormals()
+		TArray<FVector>& GetNormals()
 		{
 			return Normals;
 		}
@@ -489,7 +489,7 @@ namespace CADKernel
 		/**
 		 * @return the array of array of normal of the points of the loops
 		 */
-		const TArray<TArray<FPoint>>& GetLoopNormals() const
+		const TArray<TArray<FVector>>& GetLoopNormals() const
 		{
 			return NormalsOfFaceLoops;
 		}
diff --git a/Engine/Source/Runtime/Datasmith/CADKernel/Private/CADKernel/Mesh/Structure/Mesh.cpp b/Engine/Source/Runtime/Datasmith/CADKernel/Private/CADKernel/Mesh/Structure/Mesh.cpp
--- a/Engine/Source/Runtime/Datasmith/CADKernel/Private/CADKernel/Mesh/Structure/Mesh.cpp
+++ b/Engine/Source/Runtime/Datasmith/CADKernel/Private/CADKernel/Mesh/Structure/Mesh.cpp
@@ -10,7 +10,7 @@ using namespace CADKernel;
 
 int32 FMesh::RegisterCoordinates()
 {
-	ModelMesh.Pin()->RegisterCoordinates(NodeCoordinates, StartNodeId, MeshModelIndex);
+	ModelMesh.RegisterCoordinates(NodeCoordinates, StartNodeId, MeshModelIndex);
 	LastNodeIndex = StartNodeId + (int32)NodeCoordinates.Num();
 	return StartNodeId;
 }
diff --git a/Engine/Source/Runtime/Datasmith/CADKernel/Private/CADKernel/Mesh/Structure/ModelMesh.cpp b/Engine/Source/Runtime/Datasmith/CADKernel/Private/CADKernel/Mesh/Structure/ModelMesh.cpp
--- a/Engine/Source/Runtime/Datasmith/CADKernel/Private/CADKernel/Mesh/Structure/ModelMesh.cpp
+++ b/Engine/Source/Runtime/Datasmith/CADKernel/Private/CADKernel/Mesh/Structure/ModelMesh.cpp
@@ -134,3 +134,13 @@ const TArray<TSharedPtr<FMesh>>& FModelMesh::GetMeshes() const
 	}
 	return (TArray<TSharedPtr<FMesh>>&) VertexMeshes;
 }
+
+int32 FModelMesh::GetTriangleCount() const
+{
+	int32 TriangleCount = 0;
+	for (const TSharedPtr<FFaceMesh>& FaceMesh : FaceMeshes)
+	{
+		TriangleCount += FaceMesh->TrianglesVerticesIndex.Num() / 3;
+	}
+	return TriangleCount;
+}
diff --git a/Engine/Source/Runtime/Datasmith/CADKernel/Private/CADKernel/Topo/TopologicalLoop.cpp b/Engine/Source/Runtime/Datasmith/CADKernel/Private/CADKernel/Topo/TopologicalLoop.cpp
--- a/Engine/Source/Runtime/Datasmith/CADKernel/Private/CADKernel/Topo/TopologicalLoop.cpp
+++ b/Engine/Source/Runtime/Datasmith/CADKernel/Private/CADKernel/Topo/TopologicalLoop.cpp
@@ -142,8 +142,14 @@ void FTopologicalLoop::Orient()
 	{
 		// if the slop of the selected segment is not close to the BBox side (closed of 0 of 4), so the angle between the neighboring segments of the local extrema is not closed to 4 and allows to defined the orientation
 		// Pic case: the slop is compute between previous and nex segment of the extrema 
-		int32 Index2 = Index == 0 ? PointCount - 1 : Index - 1;
-		double Slop = ComputePositiveSlope(LoopSampling[Index], LoopSampling[Index + 1], LoopSampling[Index2]);
+		int32 PreviousIndex = Index == 0 ? PointCount - 1 : Index - 1;
+		int32 NextIndex = Index+1;
+		if(NextIndex  == LoopSampling.Num()) 
+		{
+			NextIndex = 0;
+		}
+
+		double Slop = ComputePositiveSlope(LoopSampling[Index], LoopSampling[NextIndex], LoopSampling[PreviousIndex]);
 		if (Slop > 4.2)
 		{
 			WrongOrientationNum++;
@@ -151,7 +157,7 @@ void FTopologicalLoop::Orient()
 		else if (Slop > 3.8)
 		{
 			// the angle between the neighboring segments of the local extrema is too closed to 4 (PI) so the slop of the segment Index, Index+1 is compared to the AABB side to define the orientation
-			Slop = ComputeUnorientedSlope(LoopSampling[Index], LoopSampling[Index + 1], ReferenceSlop);
+			Slop = ComputeUnorientedSlope(LoopSampling[Index], LoopSampling[NextIndex], ReferenceSlop);
 			// slop should be closed to [0, 0.2] or [3.8, 4]
 			if (Slop > 3)
 			{
@@ -173,7 +179,7 @@ void FTopologicalLoop::Orient()
 			}
 			{
 				F3DDebugSession _(*FString::Printf(TEXT("Previous")));
-				::DisplaySegment(LoopSampling[Index2], LoopSampling[Index], EVisuProperty::NonManifoldEdge);
+				::DisplaySegment(LoopSampling[PreviousIndex], LoopSampling[Index], EVisuProperty::NonManifoldEdge);
 			}
 			{
 				F3DDebugSession _(*FString::Printf(TEXT("Node %f"), Slop));
diff --git a/Engine/Source/Runtime/Datasmith/CADKernel/Private/CADKernel/UI/Display.cpp b/Engine/Source/Runtime/Datasmith/CADKernel/Private/CADKernel/UI/Display.cpp
--- a/Engine/Source/Runtime/Datasmith/CADKernel/Private/CADKernel/UI/Display.cpp
+++ b/Engine/Source/Runtime/Datasmith/CADKernel/Private/CADKernel/UI/Display.cpp
@@ -32,12 +32,16 @@ namespace CADKernel
 {
 	void Open3DDebugSession(FString name, const TArray<FIdent>& IdArray)
 	{
+#ifdef CADKERNEL_DEV
 		FSystem::Get().GetVisu()->Open3DDebugSession(*name, IdArray);
+#endif
 	}
 
 	void Close3DDebugSession()
 	{
+#ifdef CADKERNEL_DEV
 		FSystem::Get().GetVisu()->Close3DDebugSession();
+#endif
 	}
 
 	void Wait(bool bMakeWait)
@@ -53,31 +57,42 @@ namespace CADKernel
 
 	void Open3DDebugSegment(FIdent Ident)
 	{
+#ifdef CADKERNEL_DEV
 		FSystem::Get().GetVisu()->Open3DDebugSegment(Ident);
+#endif
 	}
 
 	void Close3DDebugSegment()
 	{
+#ifdef CADKERNEL_DEV
 		FSystem::Get().GetVisu()->Close3DDebugSegment();
+#endif
 	}
 
 	void FlushVisu()
 	{
+#ifdef CADKERNEL_DEV
 		FSystem::Get().GetVisu()->UpdateViewer();
+#endif
 	}
 
 	void DrawElement(int32 Dimension, TArray<FPoint>& Points, EVisuProperty Property)
 	{
+#ifdef CADKERNEL_DEV
 		FSystem::Get().GetVisu()->DrawElement(Dimension, Points, Property);
+#endif
 	}
 
 	void DrawMesh(const TSharedPtr<FMesh>& Mesh)
 	{
+#ifdef CADKERNEL_DEV
 		FSystem::Get().GetVisu()->DrawMesh(Mesh->GetId());
+#endif
 	}
 
 	void Draw(const FLinearBoundary& Boundary, const TSharedPtr<FRestrictionCurve>& Curve, EVisuProperty Property)
 	{
+#ifdef CADKERNEL_DEV
 		if (Boundary.IsDegenerated())
 		{
 			return;
@@ -106,10 +121,12 @@ namespace CADKernel
 
 			DrawQuadripode(Height, Base, Point.Point, Point.Gradient);
 		}
+#endif
 	}
 
 	void Draw2D(const TSharedPtr<FCurve>& Curve, const FLinearBoundary& Boundary, EVisuProperty Property)
 	{
+#ifdef CADKERNEL_DEV
 		if (Boundary.IsDegenerated())
 		{
 			return;
@@ -137,10 +154,12 @@ namespace CADKernel
 
 			DrawQuadripode(Height, Base, Point.Point, Point.Gradient);
 		}
+#endif
 	}
 
 	void Draw3D(const TSharedPtr<FCurve>& Curve, const FLinearBoundary& Boundary, EVisuProperty Property)
 	{
+#ifdef CADKERNEL_DEV
 		if (Boundary.IsDegenerated())
 		{
 			return;
@@ -168,10 +187,12 @@ namespace CADKernel
 
 			DrawQuadripode(Height, Base, Point.Point, Point.Gradient);
 		}
+#endif
 	}
 
 	void Draw(const TSharedPtr<FCurve>& Curve, const FLinearBoundary& Boundary, EVisuProperty Property)
 	{
+#ifdef CADKERNEL_DEV
 		if (Curve->GetDimension() == 3)
 		{
 			Draw3D(Curve, Boundary, Property);
@@ -180,17 +201,21 @@ namespace CADKernel
 		{
 			Draw2D(Curve, Boundary, Property);
 		}
+#endif
 	}
 
 
 	void Draw(const TSharedPtr<FCurve>& Curve, EVisuProperty Property)
 	{
+#ifdef CADKERNEL_DEV
 		const FLinearBoundary& Bounds = Curve->GetBoundary();
 		Draw(Curve, Bounds, Property);
+#endif
 	}
 
 	void DrawQuadripode(double Height, double Base, FPoint& Center, FPoint& InDirection)
 	{
+#ifdef CADKERNEL_DEV
 		FPoint Direction = InDirection;
 		Direction.Normalize();
 
@@ -224,10 +249,12 @@ namespace CADKernel
 		Polygone.Add(Point4);
 
 		Draw(Polygone);
+#endif
 	}
 
 	void DisplayEntity(const TSharedPtr<FEntity>& Entity)
 	{
+#ifdef CADKERNEL_DEV
 		if (!Entity.IsValid())
 		{
 			FMessage::Printf(Log, TEXT("Entity not found"));
@@ -285,10 +312,12 @@ namespace CADKernel
 		FDuration DisplayDuration = FChrono::Elapse(StartTime);
 		FChrono::PrintClockElapse(Log, TEXT("  "), TEXT("Display "), DisplayDuration);
 
+#endif
 	}
 
 	void DisplayEntity2D(const TSharedPtr<FEntity>& Entity)
 	{
+#ifdef CADKERNEL_DEV
 		if (!Entity.IsValid())
 		{
 			FMessage::Printf(Log, TEXT("Entity not found"));
@@ -322,10 +351,12 @@ namespace CADKernel
 		FDuration DisplayDuration = FChrono::Elapse(StartTime);
 		FChrono::PrintClockElapse(Log, TEXT("  "), TEXT("Display "), DisplayDuration);
 
+#endif
 	}
 
 	void Display(const FPlane& Plane, FIdent Ident)
 	{
+#ifdef CADKERNEL_DEV
 		FPoint Normal = Plane.GetNormal();
 		if (Normal.Length() < SMALL_NUMBER)
 		{
@@ -362,10 +393,12 @@ namespace CADKernel
 
 		F3DDebugSegment G(Ident);
 		DrawElement(2, Points);
+#endif
 	}
 
 	void DisplayEdgeCriteriaGrid(int32 EdgeId, const TArray<FPoint>& Points3D)
 	{
+#ifdef CADKERNEL_DEV
 		FString Name = FString::Printf(TEXT("Edge Grid %d"), EdgeId);
 		F3DDebugSession G(Name);
 		{
@@ -382,10 +415,12 @@ namespace CADKernel
 				DisplayPoint(Points3D[Index], EVisuProperty::ControlPoint);
 			}
 		}
+#endif
 	}
 
 	void DisplayAABB(const FAABB& Aabb, FIdent Ident)
 	{
+#ifdef CADKERNEL_DEV
 		TFunction<void(int32, FPoint&, FPoint&)> GetWire = [&](int32 Segment, FPoint& A, FPoint& B)
 		{
 			int32 Corner1;
@@ -431,10 +466,12 @@ namespace CADKernel
 			GetWire(SegmentIndex, Points[0], Points[1]);
 			Draw(Points);
 		}
+#endif
 	}
 
 	void DisplayAABB2D(const FAABB2D& aabb, FIdent Ident)
 	{
+#ifdef CADKERNEL_DEV
 		FPoint A(aabb.GetMin().U, aabb.GetMin().V, 0.);
 		FPoint B(aabb.GetMin().U, aabb.GetMax().V, 0.);
 		FPoint C(aabb.GetMax().U, aabb.GetMax().V, 0.);
@@ -447,40 +484,50 @@ namespace CADKernel
 		Points.Add(D);
 		Points.Add(A);
 		Draw(Points);
+#endif
 	}
 
 	void Display(const TSharedPtr<FTopologicalVertex>& Vertex)
 	{
+#ifdef CADKERNEL_DEV
 		F3DDebugSegment G(Vertex->GetId());
 		DrawPoint(Vertex->GetCoordinates());
+#endif
 	}
 
 	void Display(const TSharedPtr<FCurve>& Curve)
 	{
+#ifdef CADKERNEL_DEV
 		F3DDebugSegment G(Curve->GetId());
 		Draw(Curve, Curve->GetBoundary());
+#endif
 	}
 
 	void Display(const TSharedPtr<FShell>& Shell)
 	{
+#ifdef CADKERNEL_DEV
 		FProgress Progress(Shell->GetFaces().Num(), TEXT("Display Shell"));
 		for (const FOrientedFace& Face : Shell->GetFaces())
 		{
 			Display(Face.Entity);
 		}
+#endif
 	}
 
 	void Draw(const TSharedPtr<FShell>& Shell)
 	{
+#ifdef CADKERNEL_DEV
 		FProgress Progress(Shell->GetFaces().Num(), TEXT("Display Shell"));
 		for (const FOrientedFace& Face : Shell->GetFaces())
 		{
 			Draw(Face.Entity);
 		}
+#endif
 	}
 
 	void Display(const TSharedPtr<FBody>& Body)
 	{
+#ifdef CADKERNEL_DEV
 		FProgress Progress(Body->GetShells().Num(), TEXT("Display Body"));
 
 		F3DDebugSegment GraphicSegment(Body->GetId());
@@ -492,10 +539,12 @@ namespace CADKernel
 			}
 			Draw(Shell);
 		}
+#endif
 	}
 
 	void Display(const TSharedPtr<FSurface>& Surface)
 	{
+#ifdef CADKERNEL_DEV
 		FProgress Progress(TEXT("Display Surface"));
 
 		int32 IsoUNum = FSystem::Get().GetVisu()->GetParameters()->IsoUNumber;
@@ -528,10 +577,12 @@ namespace CADKernel
 
 		DrawIsos(IsoVNum, EIso::IsoV);
 		DrawIsos(IsoUNum, EIso::IsoU);
+#endif
 	}
 
 	void Display2D(const TSharedPtr<FSurface>& Surface)
 	{
+#ifdef CADKERNEL_DEV
 		double StepU, StepV;
 
 		F3DDebugSegment G(Surface->GetId());
@@ -556,10 +607,12 @@ namespace CADKernel
 			FPoint2D EndPoint(Bounds.UVBoundaries[EIso::IsoU].Max, Bounds.UVBoundaries[EIso::IsoV].Min + (iIso)*StepV);
 			DrawSegment(StartPoint, EndPoint, EVisuProperty::Iso);
 		}
+#endif
 	}
 
 	void DisplayIsoCurve(const TSharedPtr<FSurface>& Surface, double Coordinate, EIso IsoType)
 	{
+#ifdef CADKERNEL_DEV
 		const double VisuSag = FSystem::Get().GetVisu()->GetParameters()->ChordError;
 
 		F3DDebugSegment G(Surface->GetId());
@@ -570,12 +623,14 @@ namespace CADKernel
 		FIsoCurve3DSamplerOnChord Sampler(*Surface.Get(), VisuSag, Polyline);
 		Sampler.Set(IsoType, Coordinate, CurveBounds);
 		Draw(Polyline.GetPoints(), EVisuProperty::Iso);
+#endif
 	}
 
 
 
 	void DisplayControlPolygon(const TSharedPtr<FCurve>& Entity)
 	{
+#ifdef CADKERNEL_DEV
 		TFunction<void(const TArray<FPoint>&)> DisplayHull = [](const TArray<FPoint>& Poles)
 		{
 			for (const FPoint& Pole : Poles)
@@ -604,10 +659,12 @@ namespace CADKernel
 			DisplayHull(Nurbs->GetPoles());
 			return;
 		}
+#endif
 	}
 
 	void DisplayControlPolygon(const TSharedPtr<FSurface>& Entity)
 	{
+#ifdef CADKERNEL_DEV
 		TFunction<void(const TArray<FPoint>&, int32, int32)> DisplayHull = [](const TArray<FPoint>& Poles, int32 PoleUNum, int32 PoleVNum)
 		{
 			for (int32 Index = 0; Index < Poles.Num(); Index++)
@@ -648,22 +705,28 @@ namespace CADKernel
 			DisplayHull(Nurbs->GetPoles(), Nurbs->GetPoleCount(EIso::IsoU), Nurbs->GetPoleCount(EIso::IsoV));
 			return;
 		}
+#endif
 	}
 
 	void Display(const TSharedPtr<FTopologicalFace>& Face)
 	{
+#ifdef CADKERNEL_DEV
 		F3DDebugSegment GraphicSegment(Face->GetId());
 		Draw(Face);
+#endif
 	}
 
 	void Display2D(const TSharedPtr<FTopologicalFace>& Face)
 	{
+#ifdef CADKERNEL_DEV
 		F3DDebugSegment GraphicSegment(Face->GetId());
 		Draw2D(Face);
+#endif
 	}
 
 	void Draw(const TSharedPtr<FTopologicalFace>& Face)
 	{
+#ifdef CADKERNEL_DEV
 		{
 			for (const TSharedPtr<FTopologicalLoop>& Loop : Face->GetLoops())
 			{
@@ -687,10 +750,12 @@ namespace CADKernel
 		}
 
 		DrawIsoCurves(Face);
+#endif
 	}
 
 	void Draw2D(const TSharedPtr<FTopologicalFace>& Face)
 	{
+#ifdef CADKERNEL_DEV
 		TArray<TArray<FPoint2D>> BoundaryApproximation;
 		Face->Get2DLoopSampling(BoundaryApproximation);
 
@@ -745,11 +810,13 @@ namespace CADKernel
 
 		DrawIsos(IsoUCount, EIso::IsoU);
 		DrawIsos(IsoVCount, EIso::IsoV);
+#endif
 	}
 
 
 	void DrawIsoCurves(const TSharedPtr<FTopologicalFace>& Face)
 	{
+#ifdef CADKERNEL_DEV
 		TArray<TArray<FPoint2D>> BoundaryApproximation;
 		Face->Get2DLoopSampling(BoundaryApproximation);
 
@@ -812,30 +879,38 @@ namespace CADKernel
 
 		DrawIsos(IsoUCount, EIso::IsoU);
 		DrawIsos(IsoVCount, EIso::IsoV);
+#endif
 	}
 
 	void Display(const TSharedPtr<FTopologicalEdge>& Edge, EVisuProperty Property)
 	{
+#ifdef CADKERNEL_DEV
 		F3DDebugSegment GraphicSegment(Edge->GetId());
 		Draw(Edge, Property);
+#endif
 	}
 
 	void Display2D(const TSharedPtr<FTopologicalEdge>& Edge, EVisuProperty Property)
 	{
+#ifdef CADKERNEL_DEV
 		F3DDebugSegment GraphicSegment(Edge->GetId());
 		TArray<FPoint2D> Polyline;
 		Edge->GetCurve()->GetDiscretizationPoints(Edge->GetBoundary(), EOrientation::Front, Polyline);
 		DisplayPolyline(Polyline, Property);
+#endif
 	}
 
 	void Draw(const TSharedPtr<FTopologicalEdge>& Edge, EVisuProperty Property)
 	{
+#ifdef CADKERNEL_DEV
 		const FLinearBoundary& Boundary = Edge->GetBoundary();
 		Draw(Boundary, Edge->GetCurve(), Property);
+#endif
 	}
 
 	void Display(const TSharedPtr<FModel>& Model)
 	{
+#ifdef CADKERNEL_DEV
 		FProgress MainProgress(2, TEXT("Display model"));
 		{
 			const TArray<TSharedPtr<FBody>>& Bodies = Model->GetBodies();
@@ -854,10 +929,12 @@ namespace CADKernel
 				Display(Surface);
 			}
 		}
+#endif
 	}
 
 	void DisplayProductTree(const TSharedPtr<FModel>& Model)
 	{
+#ifdef CADKERNEL_DEV
 		F3DDebugSession GraphicSession(FString::Printf(TEXT("%s %d"), Model->GetTypeName(), Model->GetId()), { Model->GetId() });
 
 		FProgress MainProgress(2, TEXT("Display Model"));
@@ -878,11 +955,12 @@ namespace CADKernel
 				Display(Face);
 			}
 		}
-
+#endif
 	}
 
 	void DisplayProductTree(const TSharedPtr<FBody>& Body)
 	{
+#ifdef CADKERNEL_DEV
 #ifdef CORETECHBRIDGE_DEBUG
 		F3DDebugSession GraphicSession(FString::Printf(TEXT("%s %s KioId: %d Id: %d"), Body->GetTypeName(), Body->GetName(), Body->GetKioId(), Body->GetId()), { Body->GetId() });
 #else
@@ -898,11 +976,12 @@ namespace CADKernel
 			}
 			DisplayProductTree(Shell);
 		}
-
+#endif
 	}
 
 	void DisplayProductTree(const TSharedPtr<FShell>& Shell)
 	{
+#ifdef CADKERNEL_DEV
 #ifdef CORETECHBRIDGE_DEBUG
 		F3DDebugSession GraphicSession(FString::Printf(TEXT("%s %s KioId: %d Id: %d"), Shell->GetTypeName(), Shell->GetName(), Shell->GetKioId(), Shell->GetId()), { Shell->GetId() });
 #else
@@ -910,10 +989,12 @@ namespace CADKernel
 #endif
 		F3DDebugSegment Segment(Shell->GetId());
 		Draw(Shell);
+#endif
 	}
 
 	void DisplayProductTree(const TSharedPtr<FEntity>& Entity)
 	{
+#ifdef CADKERNEL_DEV
 		if (!Entity.IsValid())
 		{
 			FMessage::Printf(Log, TEXT("Entity not found"));
@@ -942,14 +1023,12 @@ namespace CADKernel
 
 		FDuration DisplayDuration = FChrono::Elapse(StartTime);
 		FChrono::PrintClockElapse(Log, TEXT("  "), TEXT("Display "), DisplayDuration);
+#endif
 	}
 
-
-
-
-
 	void Display(const TSharedPtr<FGroup>& Group)
 	{
+#ifdef CADKERNEL_DEV
 		TArray<TSharedPtr<FEntity>> Entities;
 
 		Group->GetValidEntities(Entities);
@@ -959,12 +1038,14 @@ namespace CADKernel
 		{
 			DisplayEntity(Entity);
 		}
+#endif
 	}
 
 	void DisplayMesh(const TSharedPtr<FMesh>& Mesh)
 	{
+#ifdef CADKERNEL_DEV
 		ensureCADKernel(Mesh.IsValid());
-		switch (Mesh->GetGeometricEntity()->GetEntityType())
+		switch (Mesh->GetGeometricEntity().GetEntityType())
 		{
 		case EEntity::TopologicalFace:
 			DisplayMesh(StaticCastSharedRef<FFaceMesh>(Mesh.ToSharedRef()));
@@ -978,10 +1059,12 @@ namespace CADKernel
 		default:
 			break;
 		}
+#endif
 	}
 
 	void DisplayMesh(const TSharedRef<FFaceMesh>& Mesh)
 	{
+#ifdef CADKERNEL_DEV
 		TMap<int32, const FPoint*> NodeIdToCoordinates;
 		Mesh->GetNodeIdToCoordinates(NodeIdToCoordinates);
 
@@ -1025,10 +1108,12 @@ namespace CADKernel
 				DrawSegment(*NodeIdToCoordinates[VertexIndices[Index]], *NodeIdToCoordinates[VertexIndices[Index]] + Normal, EVisuProperty::EdgeMesh);
 			}
 		}
+#endif
 	}
 
 	void DisplayMesh(const TSharedRef<FEdgeMesh>& Mesh)
 	{
+#ifdef CADKERNEL_DEV
 		const TSharedRef<FModelMesh> MeshModel = Mesh->GetMeshModel();
 
 		const TArray<int32>& NodeIds = Mesh->EdgeVerticesIndex;
@@ -1075,15 +1160,20 @@ namespace CADKernel
 				DrawPoint(NodeCoordinates[Index], EVisuProperty::NodeMesh);
 			}
 		}
+#endif
 	}
 
 	void DisplayMesh(const TSharedRef<FVertexMesh>& Mesh)
 	{
+#ifdef CADKERNEL_DEV
 		DisplayPoint(Mesh->GetNodeCoordinates()[0], EVisuProperty::NodeMesh, Mesh->GetId());
+#endif
 	}
 
+
 	void Display(const TSharedPtr<FModelMesh>& MeshModel)
 	{
+#ifdef CADKERNEL_DEV
 		for (TSharedPtr<FMesh> Mesh : MeshModel->GetMeshes())
 		{
 			if (Mesh.IsValid())
@@ -1091,52 +1181,65 @@ namespace CADKernel
 				DisplayMesh(Mesh.ToSharedRef());
 			}
 		}
+#endif
 	}
 
 	void DisplaySegment(const FPoint& Point1, const FPoint& Point2, FIdent Ident, EVisuProperty Property)
 	{
+#ifdef CADKERNEL_DEV
 		F3DDebugSegment G(Ident);
 		DrawSegment(Point1, Point2, Property);
+#endif
 	}
 
 	void DisplaySegment(const FPoint2D& Point1, const FPoint2D& Point2, FIdent Ident, EVisuProperty Property, bool bWithOrientation)
 	{
+#ifdef CADKERNEL_DEV
 		F3DDebugSegment G(Ident);
 		if (bWithOrientation)
 		{
 			DrawSegmentOrientation(Point1, Point2, Property);
 		}
 		DrawSegment(Point1, Point2, Property);
+#endif
 	}
 
 	void DisplaySegment(const FPointH& Point1, const FPointH& Point2, FIdent Ident, EVisuProperty Property)
 	{
+#ifdef CADKERNEL_DEV
 		F3DDebugSegment G(Ident);
 		DrawSegment(Point1, Point2, Property);
+#endif
 	}
 
 	void DisplayLoop(const TSharedPtr<FTopologicalFace>& Surface)
 	{
+#ifdef CADKERNEL_DEV
 		for (const TSharedPtr<FTopologicalLoop>& Loop : Surface->GetLoops())
 		{
 			Display(Loop);
 		}
+#endif
 	}
 
 
 	void Display(const TSharedPtr<FTopologicalLoop>& Loop)
 	{
+#ifdef CADKERNEL_DEV
 		for (const FOrientedEdge& Edge : Loop->GetEdges())
 		{
 			Display(Edge.Entity);
 		}
+#endif
 	}
 
 	void Display2D(const TSharedPtr<FTopologicalLoop>& Loop)
 	{
+#ifdef CADKERNEL_DEV
 		for (const FOrientedEdge& Edge : Loop->GetEdges())
 		{
 			Display2D(Edge.Entity);
 		}
+#endif
 	}
 } // namespace CADKernel
diff --git a/Engine/Source/Runtime/Datasmith/CADKernel/Public/CADKernel/Core/Session.h b/Engine/Source/Runtime/Datasmith/CADKernel/Public/CADKernel/Core/Session.h
--- a/Engine/Source/Runtime/Datasmith/CADKernel/Public/CADKernel/Core/Session.h
+++ b/Engine/Source/Runtime/Datasmith/CADKernel/Public/CADKernel/Core/Session.h
@@ -18,8 +18,6 @@ namespace CADKernel
 		double GeometricTolerance;
 		FDatabase Database;
 
-		TSharedPtr<FModel> Model;
-
 	public:
 #ifdef CADKERNEL_DEV
 		static TSharedPtr<FSession> Session;
@@ -42,6 +40,11 @@ namespace CADKernel
 			return Database;
 		}
 
+		/**
+		 * Tolerance must not be modified as soon as a geometric entity has been build.
+		 */
+		void SetGeometricTolerance(double NewTolerance);
+
 		double GetGeometricTolerance() const
 		{
 			return GeometricTolerance;
@@ -87,7 +90,6 @@ namespace CADKernel
 
 		void Clear()
 		{
-			Model.Reset();
 			Database.Empty();
 		}
 
diff --git a/Engine/Source/Runtime/Datasmith/CADKernel/Public/CADKernel/Geo/Curves/NURBSCurve.h b/Engine/Source/Runtime/Datasmith/CADKernel/Public/CADKernel/Geo/Curves/NURBSCurve.h
--- a/Engine/Source/Runtime/Datasmith/CADKernel/Public/CADKernel/Geo/Curves/NURBSCurve.h
+++ b/Engine/Source/Runtime/Datasmith/CADKernel/Public/CADKernel/Geo/Curves/NURBSCurve.h
@@ -6,6 +6,17 @@
 
 namespace CADKernel
 {
+	struct FNurbsCurveData
+	{
+		bool bIsRational;
+		int32 Dimension = 0;
+
+		int32 Degree = 0;
+		TArray<double> NodalVector;
+
+		TArray<double> Weights;
+		TArray<FPoint> Poles;
+	};
 
 	class CADKERNEL_API FNURBSCurve : public FCurve
 	{
@@ -30,6 +41,7 @@ namespace CADKernel
 
 		FNURBSCurve(int32 InDegree, const TArray<double>& InNodalVector, const TArray<FPoint>& InPoles, int8 InDimension = 3);
 		FNURBSCurve(int32 InDegree, const TArray<double>& InNodalVector, const TArray<FPoint>& InPoles, const TArray<double>& InWeights, int8 InDimension = 3);
+		FNURBSCurve(FNurbsCurveData& NurbsCurveData);
 		FNURBSCurve(const TSharedRef<FNURBSCurve>& Nurbs);
 
 		FNURBSCurve(FCADKernelArchive& Archive)
diff --git a/Engine/Source/Runtime/Datasmith/CADKernel/Public/CADKernel/Geo/Sampling/SurfacicPolyline.h b/Engine/Source/Runtime/Datasmith/CADKernel/Public/CADKernel/Geo/Sampling/SurfacicPolyline.h
--- a/Engine/Source/Runtime/Datasmith/CADKernel/Public/CADKernel/Geo/Sampling/SurfacicPolyline.h
+++ b/Engine/Source/Runtime/Datasmith/CADKernel/Public/CADKernel/Geo/Sampling/SurfacicPolyline.h
@@ -22,7 +22,7 @@ namespace CADKernel
 		TArray<double> Coordinates;
 		TArray<FPoint2D> Points2D;
 		TArray<FPoint> Points3D;
-		TArray<FPoint> Normals;
+		TArray<FVector> Normals;
 		TArray<FPoint> Tangents;
 
 		bool bWithNormals;
@@ -264,7 +264,7 @@ namespace CADKernel
 			return Points3D;
 		}
 
-		const TArray<FPoint>& GetNormals() const
+		const TArray<FVector>& GetNormals() const
 		{
 			return Normals;
 		}
diff --git a/Engine/Source/Runtime/Datasmith/CADKernel/Public/CADKernel/Geo/Surfaces/NURBSSurface.h b/Engine/Source/Runtime/Datasmith/CADKernel/Public/CADKernel/Geo/Surfaces/NURBSSurface.h
--- a/Engine/Source/Runtime/Datasmith/CADKernel/Public/CADKernel/Geo/Surfaces/NURBSSurface.h
+++ b/Engine/Source/Runtime/Datasmith/CADKernel/Public/CADKernel/Geo/Surfaces/NURBSSurface.h
@@ -7,13 +7,46 @@
 namespace CADKernel
 {
 
+	struct FNurbsSurfaceHomogeneousData
+	{
+		bool bSwapUV;
+		int32 PoleUCount;
+		int32 PoleVCount;
+
+		int32 UDegree;
+		int32 VDegree;
+
+		TArray<double> UNodalVector;
+		TArray<double> VNodalVector;
+
+		bool bIsRational;
+		TArray<double> HomogeneousPoles;
+	};
+
+	struct FNurbsSurfaceData
+	{
+		bool bSwapUV;
+		int32 PoleUCount;
+		int32 PoleVCount;
+
+		int32 UDegree;
+		int32 VDegree;
+
+		TArray<double> UNodalVector;
+		TArray<double> VNodalVector;
+
+		// if Weights.Num() == 0  => bIsRational = false
+		TArray<double> Weights;
+		TArray<FPoint> Poles;
+	};
+
 	class CADKERNEL_API FNURBSSurface : public FSurface
 	{
 		friend FEntity;
 
 	protected:
-		int32 PoleUNum;
-		int32 PoleVNum;
+		int32 PoleUCount;
+		int32 PoleVCount;
 
 		int32 UDegree;
 		int32 VDegree;
@@ -37,10 +70,10 @@ namespace CADKernel
 		 * @param NodalVectorU Its size is the number of poles in U + the surface degree in U + 1 (PoleUNum + UDegre + 1)
 		 * @param NodalVectorV Its size is the number of poles in V + the surface degree in V + 1 (PoleVNum + VDegre + 1)
 		 */
-		FNURBSSurface(const double InToleranceGeometric, int32 InPoleUNum, int32 InPoleVNum, int32 InDegreU, int32 InDegreV, const TArray<double>& InNodalVectorU, const TArray<double>& InNodalVectorV, const TArray<FPoint>& InPoles)
+		FNURBSSurface(const double InToleranceGeometric, int32 InPoleUCount, int32 InPoleVCount, int32 InDegreU, int32 InDegreV, const TArray<double>& InNodalVectorU, const TArray<double>& InNodalVectorV, const TArray<FPoint>& InPoles)
 			: FSurface(InToleranceGeometric)
-			, PoleUNum(InPoleUNum)
-			, PoleVNum(InPoleVNum)
+			, PoleUCount(InPoleUCount)
+			, PoleVCount(InPoleVCount)
 			, UDegree(InDegreU)
 			, VDegree(InDegreV)
 			, UNodalVector(InNodalVectorU)
@@ -56,10 +89,10 @@ namespace CADKernel
 		 * @param NodalVectorU Its size is the number of poles in U + the surface degree in U + 1 (PoleUNum + UDegre + 1)
 		 * @param NodalVectorV Its size is the number of poles in V + the surface degree in V + 1 (PoleVNum + VDegre + 1)
 		 */
-		FNURBSSurface(const double InToleranceGeometric, int32 InPoleUNum, int32 InPoleVNum, int32 InDegreU, int32 InDegreV, const TArray<double>& InNodalVectorU, const TArray<double>& InNodalVectorV, const TArray<FPoint>& InPoles, const TArray<double>& InWeights)
+		FNURBSSurface(const double InToleranceGeometric, int32 InPoleUCount, int32 InPoleVCount, int32 InDegreU, int32 InDegreV, const TArray<double>& InNodalVectorU, const TArray<double>& InNodalVectorV, const TArray<FPoint>& InPoles, const TArray<double>& InWeights)
 			: FSurface(InToleranceGeometric)
-			, PoleUNum(InPoleUNum)
-			, PoleVNum(InPoleVNum)
+			, PoleUCount(InPoleUCount)
+			, PoleVCount(InPoleVCount)
 			, UDegree(InDegreU)
 			, VDegree(InDegreV)
 			, UNodalVector(InNodalVectorU)
@@ -72,6 +105,17 @@ namespace CADKernel
 			Finalize();
 		}
 
+		/**
+		 * Build a Non uniform B-Spline surface
+		 * @param NodalVectorU Its size is the number of poles in U + the surface degree in U + 1 (PoleUNum + UDegre + 1)
+		 * @param NodalVectorV Its size is the number of poles in V + the surface degree in V + 1 (PoleVNum + VDegre + 1)
+		 */
+		FNURBSSurface(const double InToleranceGeometric, FNurbsSurfaceHomogeneousData NurbsData)
+			: FSurface(InToleranceGeometric)
+		{
+			FillNurbs(NurbsData);
+		}
+
 		FNURBSSurface(FCADKernelArchive& Archive)
 			: FSurface()
 		{
@@ -84,8 +128,8 @@ namespace CADKernel
 		virtual void Serialize(FCADKernelArchive& Ar) override
 		{
 			FSurface::Serialize(Ar);
-			Ar << PoleUNum;
-			Ar << PoleVNum;
+			Ar << PoleUCount;
+			Ar << PoleVCount;
 			Ar << UDegree;
 			Ar << VDegree;
 			Ar.Serialize(UNodalVector);
@@ -121,10 +165,10 @@ namespace CADKernel
 			switch (Iso)
 			{
 			case EIso::IsoU:
-				return PoleUNum;
+				return PoleUCount;
 			case EIso::IsoV:
 			default:
-				return PoleVNum;
+				return PoleVCount;
 			}
 		}
 
@@ -180,6 +224,8 @@ namespace CADKernel
 
 	private:
 		void Finalize();
+		void FillNurbs(FNurbsSurfaceHomogeneousData& NurbsData);
+		void FillNurbs(FNurbsSurfaceData& NurbsData);
 	};
 
 } // namespace CADKernel
diff --git a/Engine/Source/Runtime/Datasmith/CADKernel/Public/CADKernel/Geo/Surfaces/Surface.h b/Engine/Source/Runtime/Datasmith/CADKernel/Public/CADKernel/Geo/Surfaces/Surface.h
--- a/Engine/Source/Runtime/Datasmith/CADKernel/Public/CADKernel/Geo/Surfaces/Surface.h
+++ b/Engine/Source/Runtime/Datasmith/CADKernel/Public/CADKernel/Geo/Surfaces/Surface.h
@@ -171,7 +171,7 @@ namespace CADKernel
 		virtual void EvaluatePointGrid(const FCoordinateGrid& Coordinates, FSurfacicSampling& OutPoints, bool bComputeNormals = false) const;
 		void EvaluateGrid(FGrid& Grid) const;
 
-		virtual void EvaluateNormals(const TArray<FPoint2D>& Points2D, TArray<FPoint>& Normals) const;
+		virtual void EvaluateNormals(const TArray<FPoint2D>& Points2D, TArray<FVector>& Normals) const;
 
 		/**
 		 * Divide the parametric space in the desired number of regular subdivisions and compute the associated PointGrid
diff --git a/Engine/Source/Runtime/Datasmith/CADKernel/Public/CADKernel/Math/MatrixH.h b/Engine/Source/Runtime/Datasmith/CADKernel/Public/CADKernel/Math/MatrixH.h
--- a/Engine/Source/Runtime/Datasmith/CADKernel/Public/CADKernel/Math/MatrixH.h
+++ b/Engine/Source/Runtime/Datasmith/CADKernel/Public/CADKernel/Math/MatrixH.h
@@ -91,6 +91,7 @@ namespace CADKernel
 		 */
 		FPoint PointRotation(const FPoint& PointToRotate, const FPoint& Origin) const;
 		FPoint2D PointRotation(const FPoint2D& PointToRotate, const FPoint2D& Origin) const;
+		FVector PointRotation(const FVector& PointToRotate, const FVector& Origin) const;
 
 		void Inverse();
 
diff --git a/Engine/Source/Runtime/Datasmith/CADKernel/Public/CADKernel/Math/Point.h b/Engine/Source/Runtime/Datasmith/CADKernel/Public/CADKernel/Math/Point.h
--- a/Engine/Source/Runtime/Datasmith/CADKernel/Public/CADKernel/Math/Point.h
+++ b/Engine/Source/Runtime/Datasmith/CADKernel/Public/CADKernel/Math/Point.h
@@ -69,6 +69,11 @@ namespace CADKernel
 			return FPoint(X + Point.X, Y + Point.Y, Z + Point.Z);
 		}
 
+		FPoint operator+(const FVector& Point) const
+		{
+			return FPoint(X + Point.X, Y + Point.Y, Z + Point.Z);
+		}
+
 		FPoint operator-(const FPoint& Point) const
 		{
 			return FPoint(X - Point.X, Y - Point.Y, Z - Point.Z);
@@ -115,6 +120,14 @@ namespace CADKernel
 			return *this;
 		}
 
+		FPoint& operator+=(const FVector& Point)
+		{
+			X += Point.X;
+			Y += Point.Y;
+			Z += Point.Z;
+			return *this;
+		}
+
 		FPoint& operator-=(const FPoint& Point)
 		{
 			X -= Point.X;
@@ -165,6 +178,11 @@ namespace CADKernel
 			return (Compare < 0);
 		}
 
+		operator FVector() 
+		{ 
+			return FVector(X, Y, Z);
+		}
+
 		void SetMin(const FPoint& Point)
 		{
 			X = FMath::Min(X, Point.X);
@@ -669,4 +687,14 @@ inline uint32 GetTypeHash(const CADKernel::FPoint2D& Point)
 inline uint32 GetTypeHash(const CADKernel::FPointH& Point)
 {
 	return HashCombine(GetTypeHash(Point.X), HashCombine(GetTypeHash(Point.Y), HashCombine(GetTypeHash(Point.Z), GetTypeHash(Point.W))));
-}
\ No newline at end of file
+}
+
+inline float ComputeCosinus(FVector Vector, FVector OtherVector)
+{
+	Vector.Normalize();
+	OtherVector.Normalize();
+
+	float Cosinus = Vector.X * OtherVector.X + Vector.Y * OtherVector.Y + Vector.Z * OtherVector.Z;
+
+	return FMath::Max(-1.0f, FMath::Min(Cosinus, 1.0f));
+}
diff --git a/Engine/Source/Runtime/Datasmith/CADKernel/Public/CADKernel/Mesh/Structure/FaceMesh.h b/Engine/Source/Runtime/Datasmith/CADKernel/Public/CADKernel/Mesh/Structure/FaceMesh.h
--- a/Engine/Source/Runtime/Datasmith/CADKernel/Public/CADKernel/Mesh/Structure/FaceMesh.h
+++ b/Engine/Source/Runtime/Datasmith/CADKernel/Public/CADKernel/Mesh/Structure/FaceMesh.h
@@ -21,10 +21,10 @@ namespace CADKernel
 		TArray<int32> VerticesGlobalIndex;
 
 		/** Normals of the surface mesh vertices*/
-		TArray<FPoint> Normals;
+		TArray<FVector> Normals;
 
 		/** UV coordinates of the surface mesh vertices*/
-		TArray<FPoint2D> UVMap;
+		TArray<FVector2D> UVMap;
 
 	public:
 		FFaceMesh(TSharedRef<FModelMesh> InMeshModel, TSharedRef<FTopologicalEntity> InTopologicalEntity)
@@ -56,6 +56,7 @@ namespace CADKernel
 		void InverseOrientation();
 
 		void GetNodeIdToCoordinates(TMap<int32, const FPoint*>& NodeIdToCoordinates) const;
+
 	};
 }
 
diff --git a/Engine/Source/Runtime/Datasmith/CADKernel/Public/CADKernel/Mesh/Structure/Mesh.h b/Engine/Source/Runtime/Datasmith/CADKernel/Public/CADKernel/Mesh/Structure/Mesh.h
--- a/Engine/Source/Runtime/Datasmith/CADKernel/Public/CADKernel/Mesh/Structure/Mesh.h
+++ b/Engine/Source/Runtime/Datasmith/CADKernel/Public/CADKernel/Mesh/Structure/Mesh.h
@@ -13,8 +13,8 @@ namespace CADKernel
 	class CADKERNEL_API FMesh : public FEntityGeom
 	{
 	protected:
-		TWeakPtr<FModelMesh> ModelMesh;
-		TWeakPtr<FTopologicalEntity> TopologicalEntity;
+		FModelMesh& ModelMesh;
+		FTopologicalEntity& TopologicalEntity;
 
 		int32 StartNodeId;
 		int32 LastNodeIndex;
@@ -26,8 +26,8 @@ namespace CADKernel
 
 		FMesh(TSharedRef<FModelMesh> InMeshModel, TSharedRef<FTopologicalEntity> InTopologicalEntity)
 			: FEntityGeom()
-			, ModelMesh(InMeshModel)
-			, TopologicalEntity(InTopologicalEntity)
+			, ModelMesh(*InMeshModel)
+			, TopologicalEntity(*InTopologicalEntity)
 		{
 		}
 
@@ -69,20 +69,19 @@ namespace CADKernel
 			return MeshModelIndex;
 		}
 
-		TSharedRef<FModelMesh> GetMeshModel()
+		FModelMesh& GetMeshModel()
 		{
-			ensureCADKernel(ModelMesh.Pin().IsValid());
-			return ModelMesh.Pin().ToSharedRef();
+			return ModelMesh;
 		}
 
-		const TSharedRef<FModelMesh> GetMeshModel() const
+		const FModelMesh& GetMeshModel() const
 		{
-			return ModelMesh.Pin().ToSharedRef();
+			return ModelMesh;
 		}
 
-		const TSharedRef<FTopologicalEntity> GetGeometricEntity() const
+		const FTopologicalEntity& GetGeometricEntity() const
 		{
-			return TopologicalEntity.Pin().ToSharedRef();
+			return TopologicalEntity;
 		}
 	};
 }
diff --git a/Engine/Source/Runtime/Datasmith/CADKernel/Public/CADKernel/Mesh/Structure/ModelMesh.h b/Engine/Source/Runtime/Datasmith/CADKernel/Public/CADKernel/Mesh/Structure/ModelMesh.h
--- a/Engine/Source/Runtime/Datasmith/CADKernel/Public/CADKernel/Mesh/Structure/ModelMesh.h
+++ b/Engine/Source/Runtime/Datasmith/CADKernel/Public/CADKernel/Mesh/Structure/ModelMesh.h
@@ -20,7 +20,6 @@ namespace CADKernel
 	{
 		friend FEntity;
 
-		//TSharedRef<FSession> Session;
 		TArray<TSharedPtr<FCriterion>> Criteria;
 
 		TArray<TArray<FPoint>*> GlobalPointCloud;
@@ -36,14 +35,25 @@ namespace CADKernel
 		double MaxAngle = PI;
 		double Sag = HUGE_VALUE;
 
-		FModelMesh(/*TSharedRef<FSession>& InSession*/)
+		FModelMesh()
 			: FEntityGeom()
-			//, Session(InSession)
 			, LastIdUsed(0)
 		{
 		}
 
 	public:
+		
+		int32 GetFaceCount() const
+		{
+			return FaceMeshes.Num();
+		}
+
+		int32 GetVertexCount() const
+		{
+			return LastIdUsed;
+		}
+
+		int32 GetTriangleCount() const;
 
 		virtual void SpawnIdent(FDatabase& Database) override
 		{
@@ -138,6 +148,11 @@ namespace CADKernel
 
 		const TArray<TSharedPtr<FMesh>>& GetMeshes() const;
 
+		const TArray<TSharedPtr<FFaceMesh>>& GetFaceMeshes() const 
+		{
+			return FaceMeshes;
+		}
+
 #ifdef CADKERNEL_DEV
 		virtual FInfoEntity& GetInfo(FInfoEntity&) const override;
 #endif
diff --git a/Engine/Source/Runtime/Datasmith/CADKernel/Public/CADKernel/Topo/Model.h b/Engine/Source/Runtime/Datasmith/CADKernel/Public/CADKernel/Topo/Model.h
--- a/Engine/Source/Runtime/Datasmith/CADKernel/Public/CADKernel/Topo/Model.h
+++ b/Engine/Source/Runtime/Datasmith/CADKernel/Public/CADKernel/Topo/Model.h
@@ -151,6 +151,11 @@ namespace CADKernel
 			OtherModel.Faces.Empty();
 		}
 
+		int32 EntityCount() const
+		{
+			return Bodies.Num() + Faces.Num();
+		}
+
 		const TArray<TSharedPtr<FTopologicalFace>>& GetFaces() const
 		{
 			return Faces;
