[NO NON-PII COMMIT MESSAGE]

diff --git a/Engine/Source/Runtime/Datasmith/CADKernel/Private/CADKernel/Mesh/Meshers/IsoTriangulator.cpp b/Engine/Source/Runtime/Datasmith/CADKernel/Private/CADKernel/Mesh/Meshers/IsoTriangulator.cpp
--- a/Engine/Source/Runtime/Datasmith/CADKernel/Private/CADKernel/Mesh/Meshers/IsoTriangulator.cpp
+++ b/Engine/Source/Runtime/Datasmith/CADKernel/Private/CADKernel/Mesh/Meshers/IsoTriangulator.cpp
@@ -1351,7 +1351,7 @@ void FIsoTriangulator::FindCandidateSegmentsToLinkInnerAndLoop()
 
 	int32 CellIndex = -1;
 
-	FIsoInnerNode* Cell[4];
+	FIsoInnerNode* Cell[4] = {nullptr, nullptr, nullptr, nullptr};
 	TFunction<void(int32 CellIndex)> GetCellIsoNode = [&](int32 CellIndex)
 	{
 		int32 Index = CellIndex;
