Fix build error introduced in CL 17116813.

#rnx

diff --git a/Engine/Source/Runtime/Datasmith/DirectLink/Private/DirectLinkEndpoint.cpp b/Engine/Source/Runtime/Datasmith/DirectLink/Private/DirectLinkEndpoint.cpp
--- a/Engine/Source/Runtime/Datasmith/DirectLink/Private/DirectLinkEndpoint.cpp
+++ b/Engine/Source/Runtime/Datasmith/DirectLink/Private/DirectLinkEndpoint.cpp
@@ -440,7 +440,7 @@ FEndpoint::EOpenStreamResult FEndpoint::OpenStream(const FSourceHandle& SourceId
 		FRWScopeLock _(SharedState.StreamsLock, SLT_Write);
 		FStreamPort StreamPort = ++SharedState.StreamPortIdGenerator;
 
-		FDirectLinkMsg_OpenStreamRequest* Request = NewMessage<FDirectLinkMsg_OpenStreamRequest>();
+		FDirectLinkMsg_OpenStreamRequest* Request = FMessageEndpoint::MakeMessage<FDirectLinkMsg_OpenStreamRequest>();
 		Request->bRequestFromSource = bRequestFromSource;
 		Request->RequestFromStreamPort = StreamPort;
 		Request->SourceGuid = SourceId;
@@ -644,7 +644,7 @@ void FInternalThreadState::Handle_EndpointLifecycle(const FDirectLinkMsg_Endpoin
 			if (!bIsUpToDate)
 			{
 				UE_CLOG(SharedState.bDebugLog, LogDirectLinkNet, Log, TEXT("Endpoint '%s': Send FDirectLinkMsg_QueryEndpointState"), *SharedState.NiceName);
-				MessageEndpoint->Send(NewMessage<FDirectLinkMsg_QueryEndpointState>(), RemoteEndpointAddress);
+				MessageEndpoint->Send(FMessageEndpoint::MakeMessage<FDirectLinkMsg_QueryEndpointState>(), RemoteEndpointAddress);
 			}
 			break;
 		}
@@ -703,7 +703,7 @@ void FInternalThreadState::Handle_OpenStreamRequest(const FDirectLinkMsg_OpenStr
 
 	const FMessageAddress& RemoteEndpointAddress = Context->GetSender();
 
-	FDirectLinkMsg_OpenStreamAnswer* Answer = NewMessage<FDirectLinkMsg_OpenStreamAnswer>();
+	FDirectLinkMsg_OpenStreamAnswer* Answer = FMessageEndpoint::MakeMessage<FDirectLinkMsg_OpenStreamAnswer>();
 	Answer->RecipientStreamPort = Message.RequestFromStreamPort;
 
 	auto DenyConnection = [&](const FString& Reason)
@@ -855,7 +855,7 @@ void FInternalThreadState::ReplicateState(const FMessageAddress& RemoteEndpointA
 {
 	if (MessageEndpoint.IsValid())
 	{
-		FDirectLinkMsg_EndpointState* EndpointStateMessage = NewMessage<FDirectLinkMsg_EndpointState>();
+		FDirectLinkMsg_EndpointState* EndpointStateMessage = FMessageEndpoint::MakeMessage<FDirectLinkMsg_EndpointState>();
 		*EndpointStateMessage = ThisDescription;
 
 		if (RemoteEndpointAddress.IsValid())
@@ -1296,7 +1296,7 @@ void FSharedState::CloseStreamInternal(FStreamDescription& Stream, const FRWScop
 	if (bNotifyRemote && Stream.RemoteAddress.IsValid())
 	{
 		UE_CLOG(bDebugLog, LogDirectLinkNet, Log, TEXT("Endpoint '%s': Stream removed"), *NiceName, *Stream.SourcePoint.ToString());
-		FDirectLinkMsg_CloseStreamRequest* Request = NewMessage<FDirectLinkMsg_CloseStreamRequest>();
+		FDirectLinkMsg_CloseStreamRequest* Request = FMessageEndpoint::MakeMessage<FDirectLinkMsg_CloseStreamRequest>();
 		Request->RecipientStreamPort = Stream.RemoteStreamPort;
 
 		UE_CLOG(bDebugLog, LogDirectLinkNet, Log, TEXT("Endpoint '%s': Send FDirectLinkMsg_CloseStreamRequest"), *NiceName);
diff --git a/Engine/Source/Runtime/Datasmith/DirectLink/Private/DirectLinkScenePipe.cpp b/Engine/Source/Runtime/Datasmith/DirectLink/Private/DirectLinkScenePipe.cpp
--- a/Engine/Source/Runtime/Datasmith/DirectLink/Private/DirectLinkScenePipe.cpp
+++ b/Engine/Source/Runtime/Datasmith/DirectLink/Private/DirectLinkScenePipe.cpp
@@ -156,7 +156,7 @@ void FPipeBase::SendInternal(MessageType* Message, int32 ByteSizeHint)
 
 void FScenePipeToNetwork::SetupScene(FSetupSceneArg& SetupSceneArg)
 {
-	FDirectLinkMsg_DeltaMessage* Message = NewMessage<FDirectLinkMsg_DeltaMessage>(
+	FDirectLinkMsg_DeltaMessage* Message = FMessageEndpoint::MakeMessage<FDirectLinkMsg_DeltaMessage>(
 		FDirectLinkMsg_DeltaMessage::SetupScene, RemoteStreamPort, 0, 0
 	);
 
@@ -177,7 +177,7 @@ void FScenePipeToNetwork::OpenDelta(FOpenDeltaArg& OpenDeltaArg)
 
 	NextMessageNumber = 0;
 
-	FDirectLinkMsg_DeltaMessage* Message = NewMessage<FDirectLinkMsg_DeltaMessage>(
+	FDirectLinkMsg_DeltaMessage* Message = FMessageEndpoint::MakeMessage<FDirectLinkMsg_DeltaMessage>(
 		FDirectLinkMsg_DeltaMessage::OpenDelta,
 		RemoteStreamPort, BatchNumber, NextMessageNumber++
 	);
@@ -219,7 +219,7 @@ void FScenePipeToNetwork::OnSetElement(FSetElementArg& SetElementArg)
 
 void FScenePipeToNetwork::SendSetElementBuffer()
 {
-	FDirectLinkMsg_DeltaMessage* Message = NewMessage<FDirectLinkMsg_DeltaMessage>(
+	FDirectLinkMsg_DeltaMessage* Message = FMessageEndpoint::MakeMessage<FDirectLinkMsg_DeltaMessage>(
 		FDirectLinkMsg_DeltaMessage::SetElements,
 		RemoteStreamPort, BatchNumber, NextMessageNumber++
 	);
@@ -233,7 +233,7 @@ void FScenePipeToNetwork::SendSetElementBuffer()
 
 void FScenePipeToNetwork::RemoveElements(FRemoveElementsArg& RemoveElementsArg)
 {
-	FDirectLinkMsg_DeltaMessage* Message = NewMessage<FDirectLinkMsg_DeltaMessage>(
+	FDirectLinkMsg_DeltaMessage* Message = FMessageEndpoint::MakeMessage<FDirectLinkMsg_DeltaMessage>(
 		FDirectLinkMsg_DeltaMessage::RemoveElements,
 		RemoteStreamPort, BatchNumber, NextMessageNumber++
 	);
@@ -248,7 +248,7 @@ void FScenePipeToNetwork::RemoveElements(FRemoveElementsArg& RemoveElementsArg)
 void FScenePipeToNetwork::OnCloseDelta(FCloseDeltaArg& CloseDeltaArg)
 {
 	SendSetElementBuffer();
-	FDirectLinkMsg_DeltaMessage* Message = NewMessage<FDirectLinkMsg_DeltaMessage>(
+	FDirectLinkMsg_DeltaMessage* Message = FMessageEndpoint::MakeMessage<FDirectLinkMsg_DeltaMessage>(
 		FDirectLinkMsg_DeltaMessage::CloseDelta,
 		RemoteStreamPort, BatchNumber, NextMessageNumber++
 	);
@@ -315,7 +315,7 @@ void FScenePipeFromNetwork::OnOpenHaveList(const FSceneIdentifier& HaveSceneId,
 	BatchNumber = SyncCycle;
 	NextMessageNumber = 0;
 
-	FDirectLinkMsg_HaveListMessage* Message = NewMessage<FDirectLinkMsg_HaveListMessage>(
+	FDirectLinkMsg_HaveListMessage* Message = FMessageEndpoint::MakeMessage<FDirectLinkMsg_HaveListMessage>(
 		FDirectLinkMsg_HaveListMessage::EKind::OpenHaveList,
 		RemoteStreamPort, BatchNumber, NextMessageNumber++
 	);
@@ -335,7 +335,7 @@ void FScenePipeFromNetwork::OnHaveElement(FSceneGraphId NodeId, FElementHash Hav
 {
 	if (BufferedHaveListContent == nullptr)
 	{
-		BufferedHaveListContent = NewMessage<FDirectLinkMsg_HaveListMessage>(
+		BufferedHaveListContent = FMessageEndpoint::MakeMessage<FDirectLinkMsg_HaveListMessage>(
 			FDirectLinkMsg_HaveListMessage::EKind::HaveListElement,
 			RemoteStreamPort, 0, 0
 		);
@@ -369,7 +369,7 @@ void FScenePipeFromNetwork::OnCloseHaveList()
 {
 	SendHaveElements();
 
-	FDirectLinkMsg_HaveListMessage* Message = NewMessage<FDirectLinkMsg_HaveListMessage>(
+	FDirectLinkMsg_HaveListMessage* Message = FMessageEndpoint::MakeMessage<FDirectLinkMsg_HaveListMessage>(
 		FDirectLinkMsg_HaveListMessage::EKind::CloseHaveList,
 		RemoteStreamPort, BatchNumber, NextMessageNumber++
 	);
@@ -384,7 +384,7 @@ void FScenePipeFromNetwork::DelegateDeltaMessage(FDirectLinkMsg_DeltaMessage& Me
 	UE_LOG(LogDirectLinkNet, Verbose, TEXT("Delta message transmited: b:%d m:%d k:%d"), Message.BatchCode, Message.MessageCode, Message.Kind);
 
 	// acknowledge message (allows the sender to track communication progress)
-	auto* AckMessage = NewMessage<FDirectLinkMsg_HaveListMessage>(
+	auto* AckMessage = FMessageEndpoint::MakeMessage<FDirectLinkMsg_HaveListMessage>(
 		FDirectLinkMsg_HaveListMessage::EKind::AckDeltaMessage,
 		RemoteStreamPort, Message.BatchCode, Message.MessageCode
 	);
