Datasmith: cleanup unused Displacement mapping code

diff --git a/Engine/Plugins/Enterprise/DatasmithFBXImporter/Source/DatasmithVREDTranslator/Private/DatasmithVREDImporter.cpp b/Engine/Plugins/Enterprise/DatasmithFBXImporter/Source/DatasmithVREDTranslator/Private/DatasmithVREDImporter.cpp
--- a/Engine/Plugins/Enterprise/DatasmithFBXImporter/Source/DatasmithVREDTranslator/Private/DatasmithVREDImporter.cpp
+++ b/Engine/Plugins/Enterprise/DatasmithFBXImporter/Source/DatasmithVREDTranslator/Private/DatasmithVREDImporter.cpp
@@ -65,7 +65,7 @@ struct FCombinedAnimBlock
 	// Returns the combined, min start time for all blocks we have
 	float GetStartSeconds()
 	{
- 		float MinStartTime = FLT_MAX;
+		float MinStartTime = FLT_MAX;
 
 		for (const auto& Pair : NodeNameToBlock)
 		{
@@ -1081,7 +1081,6 @@ TSharedPtr<IDatasmithTextureElement> CreateTextureElement(IDatasmithMasterMateri
 		TextureModeEntry(TEXT("TexTransparent"), EDatasmithTextureMode::Other),
 		TextureModeEntry(TEXT("TexScatter"), EDatasmithTextureMode::Other),
 		TextureModeEntry(TEXT("TexRoughness"), EDatasmithTextureMode::Other),
-		TextureModeEntry(TEXT("TexDisplacement"), EDatasmithTextureMode::Displace),
 		TextureModeEntry(TEXT("TexFresnel"), EDatasmithTextureMode::Other),
 		TextureModeEntry(TEXT("TexRotation"), EDatasmithTextureMode::Other),
 		TextureModeEntry(TEXT("TexIndexOfRefraction"), EDatasmithTextureMode::Other),
diff --git a/Engine/Plugins/Enterprise/DatasmithGLTFImporter/Source/DatasmithGLTFTranslator/Private/DatasmithGLTFMaterialElement.cpp b/Engine/Plugins/Enterprise/DatasmithGLTFImporter/Source/DatasmithGLTFTranslator/Private/DatasmithGLTFMaterialElement.cpp
--- a/Engine/Plugins/Enterprise/DatasmithGLTFImporter/Source/DatasmithGLTFTranslator/Private/DatasmithGLTFMaterialElement.cpp
+++ b/Engine/Plugins/Enterprise/DatasmithGLTFImporter/Source/DatasmithGLTFTranslator/Private/DatasmithGLTFMaterialElement.cpp
@@ -56,7 +56,6 @@ void FDatasmithGLTFMaterialElement::Finalize()
 	ConnectInput(Opacity, MaterialExpressions, MaterialElement->GetOpacity());
 	ConnectInput(Refraction, MaterialExpressions, MaterialElement->GetRefraction());
 	ConnectInput(Normal, MaterialExpressions, MaterialElement->GetNormal());
-	ConnectInput(WorldDisplacement, MaterialExpressions, MaterialElement->GetWorldDisplacement());
 	ConnectInput(AmbientOcclusion, MaterialExpressions, MaterialElement->GetAmbientOcclusion());
 
 	bIsFinal = true;
diff --git a/Engine/Plugins/Enterprise/DatasmithImporter/Source/DatasmithImporter/Private/DatasmithMaterialExpressions.cpp b/Engine/Plugins/Enterprise/DatasmithImporter/Source/DatasmithImporter/Private/DatasmithMaterialExpressions.cpp
--- a/Engine/Plugins/Enterprise/DatasmithImporter/Source/DatasmithImporter/Private/DatasmithMaterialExpressions.cpp
+++ b/Engine/Plugins/Enterprise/DatasmithImporter/Source/DatasmithImporter/Private/DatasmithMaterialExpressions.cpp
@@ -2129,7 +2129,6 @@ void FDatasmithMaterialExpressions::CreateDatasmithMaterialHelper(UPackage* Pack
 	UTexture* EmissiveTexture = FDatasmithImporterUtils::FindAsset< UTexture >( AssetsContext, ShaderElement->GetEmitTexture() );
 	UTexture* RefraTexture = FDatasmithImporterUtils::FindAsset< UTexture >( AssetsContext, ShaderElement->GetTransparencyTexture() );
 	UTexture* MaskTexture = FDatasmithImporterUtils::FindAsset< UTexture >( AssetsContext, ShaderElement->GetMaskTexture() );
-	UTexture* DisplaceTexture = FDatasmithImporterUtils::FindAsset< UTexture >( AssetsContext, ShaderElement->GetDisplaceTexture() );
 	UTexture* MetalTexture = FDatasmithImporterUtils::FindAsset< UTexture >( AssetsContext, ShaderElement->GetMetalTexture() );
 
 	// roughness
@@ -2433,15 +2432,6 @@ void FDatasmithMaterialExpressions::CreateDatasmithMaterialHelper(UPackage* Pack
 	}
 
 	ModulateNormalAmount(UnrealMaterial, ShaderElement->GetBumpAmount());
-
-	// Displace mapping. #jira UE-113727
-// 	if (AddCompExpression(ShaderElement->GetDisplaceComp(), UnrealMaterial, EDatasmithTextureSlot::DISPLACE, AssetsContext) == nullptr)
-// 	{
-// 		if (DisplaceTexture != nullptr)
-// 		{
-// 			AddTextureExpression(DisplaceTexture, ShaderElement->GetDisplaceTextureSampler(), 0.0, 0.0, UnrealMaterial, EDatasmithTextureSlot::DISPLACE);
-// 		}
-// 	}
 }
 
 UMaterialExpressionMaterialFunctionCall* FDatasmithMaterialExpressions::BlendFunctions(UMaterial* UnrealMaterial, const FDatasmithAssetsImportContext& AssetsContext,
@@ -3164,7 +3154,6 @@ void FDatasmithMaterialExpressions::CreateUEPbrMaterialGraph(const TSharedPtr< I
 	}
 
 	ConnectExpression( MaterialElement.ToSharedRef(), MaterialExpressions, MaterialElement->GetNormal().GetExpression(), GetMaterialOrFunctionSlot( UnrealMaterialOrFunction, EDatasmithTextureSlot::NORMAL ), MaterialElement->GetNormal().GetOutputIndex() );
-// #jira UE-113727 // 	ConnectExpression( MaterialElement.ToSharedRef(), MaterialExpressions, MaterialElement->GetWorldDisplacement().GetExpression(), GetMaterialOrFunctionSlot( UnrealMaterialOrFunction, EDatasmithTextureSlot::DISPLACE ), MaterialElement->GetWorldDisplacement().GetOutputIndex() );
 
 	// Connect expressions to any UMaterialExpressionCustomOutput since these aren't part of the predefined material outputs
 	for ( int32 ExpressionIndex = 0; ExpressionIndex < MaterialElement->GetExpressionsCount(); ++ExpressionIndex )
diff --git a/Engine/Plugins/Enterprise/DatasmithImporter/Source/DatasmithImporter/Private/DatasmithMaterialImporter.cpp b/Engine/Plugins/Enterprise/DatasmithImporter/Source/DatasmithImporter/Private/DatasmithMaterialImporter.cpp
--- a/Engine/Plugins/Enterprise/DatasmithImporter/Source/DatasmithImporter/Private/DatasmithMaterialImporter.cpp
+++ b/Engine/Plugins/Enterprise/DatasmithImporter/Source/DatasmithImporter/Private/DatasmithMaterialImporter.cpp
@@ -141,7 +141,6 @@ namespace DatasmithMaterialImporterUtils
 		Hash = HashCombine( Hash, ComputeExpressionInputHash( &MaterialElement->GetEmissiveColor() ) );
 		Hash = HashCombine( Hash, ComputeExpressionInputHash( &MaterialElement->GetOpacity() ) );
 		Hash = HashCombine( Hash, ComputeExpressionInputHash( &MaterialElement->GetNormal() ) );
-		Hash = HashCombine( Hash, ComputeExpressionInputHash( &MaterialElement->GetWorldDisplacement() ) );
 		Hash = HashCombine( Hash, ComputeExpressionInputHash( &MaterialElement->GetRefraction() ) );
 		Hash = HashCombine( Hash, ComputeExpressionInputHash( &MaterialElement->GetAmbientOcclusion() ) );
 		Hash = HashCombine( Hash, ComputeExpressionInputHash( &MaterialElement->GetMaterialAttributes() ) );
diff --git a/Engine/Plugins/Enterprise/DatasmithImporter/Source/DatasmithImporter/Private/DatasmithTextureImporter.cpp b/Engine/Plugins/Enterprise/DatasmithImporter/Source/DatasmithImporter/Private/DatasmithTextureImporter.cpp
--- a/Engine/Plugins/Enterprise/DatasmithImporter/Source/DatasmithImporter/Private/DatasmithTextureImporter.cpp
+++ b/Engine/Plugins/Enterprise/DatasmithImporter/Source/DatasmithImporter/Private/DatasmithTextureImporter.cpp
@@ -234,10 +234,6 @@ UTexture* FDatasmithTextureImporter::CreateTexture(const TSharedPtr<IDatasmithTe
 		TextureFact->CompressionSettings = TC_Normalmap;
 		TextureFact->bFlipNormalMapGreenChannel = true;
 		break;
-	case EDatasmithTextureMode::Displace:
-		TextureFact->CompressionSettings = TC_Displacementmap;
-		TextureFact->LODGroup = TEXTUREGROUP_World;
-		break;
 	}
 
 	const EDatasmithColorSpace DatasmithColorSpace = TextureElement->GetSRGB();
@@ -301,7 +297,7 @@ UTexture* FDatasmithTextureImporter::CreateTexture(const TSharedPtr<IDatasmithTe
 		{
 			bUpdateResource |= Texture2D->AddressX != (TextureAddress)TextureElement->GetTextureAddressX();
 			Texture2D->AddressX = (TextureAddress)TextureElement->GetTextureAddressX();
-			
+
 			bUpdateResource |= Texture2D->AddressY != (TextureAddress)TextureElement->GetTextureAddressY();
 			Texture2D->AddressY = (TextureAddress)TextureElement->GetTextureAddressY();
 		}
@@ -471,10 +467,6 @@ bool UDatasmithTexturePipeline::ExecutePreImportPipeline(UInterchangeBaseNodeCon
 		CompressionSettings = TC_Normalmap;
 		bFlipNormalMapGreenChannel = true;
 		break;
-	case EDatasmithTextureMode::Displace:
-		LODGroup = TEXTUREGROUP_World;
-		CompressionSettings = TC_Displacementmap;
-		break;
 	}
 
 	const TOptional< float > RGBCurve = [ this ]() -> TOptional< float >
diff --git a/Engine/Plugins/Enterprise/DatasmithImporter/Source/DatasmithImporter/Private/Utility/DatasmithImporterImpl.cpp b/Engine/Plugins/Enterprise/DatasmithImporter/Source/DatasmithImporter/Private/Utility/DatasmithImporterImpl.cpp
--- a/Engine/Plugins/Enterprise/DatasmithImporter/Source/DatasmithImporter/Private/Utility/DatasmithImporterImpl.cpp
+++ b/Engine/Plugins/Enterprise/DatasmithImporter/Source/DatasmithImporter/Private/Utility/DatasmithImporterImpl.cpp
@@ -280,10 +280,6 @@ void FDatasmithImporterImpl::SetTexturesMode( FDatasmithImportContext& ImportCon
 					{
 						TextureElement->SetTextureMode(EDatasmithTextureMode::Specular);
 					}
-					else if ( FCString::Strlen( ShaderElement->GetDisplaceTexture() ) > 0 && ShaderElement->GetDisplaceTexture() == TextureName)
-					{
-						TextureElement->SetTextureMode(EDatasmithTextureMode::Displace);
-					}
 					else if ( FCString::Strlen( ShaderElement->GetNormalTexture() ) > 0 && ShaderElement->GetNormalTexture() == TextureName)
 					{
 						if (!ShaderElement->GetNormalTextureSampler().bInvert)
@@ -570,7 +566,7 @@ UActorComponent* FDatasmithImporterImpl::PublicizeComponent(UActorComponent& Sou
 		if ( &SourceComponent != DestinationComponent )
 		{
 			CopyObject( SourceComponent, *DestinationComponent, ReusableBuffer );
-	
+
 			if ( DestinationComponent->GetFName() != SourceComponent.GetFName() )
 			{
 				DestinationComponent->Rename( *SourceComponent.GetName() );
@@ -607,14 +603,14 @@ USceneComponent* FDatasmithImporterImpl::FinalizeSceneComponent(FDatasmithImport
 	}
 	else
 	{
-	
+
 		if ( !DestinationComponent )
 		{
-			// Improvement suggestion. We should build a cache in finalize scene components (It could be done in O(n*k)). 
+			// Improvement suggestion. We should build a cache in finalize scene components (It could be done in O(n*k)).
 			// Currently this can be highly inefficient on the first import of a dataprep scene where some actors would have a absurd amount components.
 			//  O(n^2*k) where n is the number of components on the actor and k is the avergage lenght of a datasmith ID (we need to hash the string representing the ID to convert it to a FName).
-			// 
-			// Look at components of the actor, we might find the scene component we are looking for. 
+			//
+			// Look at components of the actor, we might find the scene component we are looking for.
 			for ( UActorComponent* Component : DestinationActor.GetInstanceComponents() )
 			{
 				if ( Component && Component->IsA( SourceComponent.GetClass() ) )
@@ -651,7 +647,7 @@ USceneComponent* FDatasmithImporterImpl::FinalizeSceneComponent(FDatasmithImport
 		if ( Child && Child->GetOuter() == SourceComponent.GetOuter() )
 		{
 			if ( USceneComponent* DestinationChild = FinalizeSceneComponent( ImportContext, *Child, DestinationActor, AttachParentForChildren, ReferencesToRemap, ReusableBuffer, ComponentsToApplyMigratedTemplate))
-			{ 
+			{
 				if ( AttachParentForChildren )
 				{
 					DestinationChild->AttachToComponent(AttachParentForChildren, FAttachmentTransformRules::KeepRelativeTransform);
diff --git a/Engine/Plugins/Enterprise/DatasmithImporter/Source/DatasmithImporter/Public/DatasmithMaterialExpressions.h b/Engine/Plugins/Enterprise/DatasmithImporter/Source/DatasmithImporter/Public/DatasmithMaterialExpressions.h
--- a/Engine/Plugins/Enterprise/DatasmithImporter/Source/DatasmithImporter/Public/DatasmithMaterialExpressions.h
+++ b/Engine/Plugins/Enterprise/DatasmithImporter/Source/DatasmithImporter/Public/DatasmithMaterialExpressions.h
@@ -47,11 +47,6 @@ enum class EDatasmithTextureSlot
 	OPACITYMASK,
 	NORMAL,
 	WORLDPOSITIONOFFSET,
-
-	// no longer meaningful. #jira UE-113727
-// 	DISPLACE,
-// 	TESSELLATIONMULTIPLIER,
-
 	SUBSURFACECOLOR,
 	COATSPECULAR,
 	COATROUGHNESS,
diff --git a/Engine/Source/Developer/Datasmith/DatasmithFacade/Private/DatasmithFacadeUEPbrMaterial.cpp b/Engine/Source/Developer/Datasmith/DatasmithFacade/Private/DatasmithFacadeUEPbrMaterial.cpp
--- a/Engine/Source/Developer/Datasmith/DatasmithFacade/Private/DatasmithFacadeUEPbrMaterial.cpp
+++ b/Engine/Source/Developer/Datasmith/DatasmithFacade/Private/DatasmithFacadeUEPbrMaterial.cpp
@@ -356,12 +356,6 @@ FDatasmithFacadeExpressionInput FDatasmithFacadeUEPbrMaterial::GetNormal() const
 	return FDatasmithFacadeExpressionInput( &UEPbrMaterial->GetNormal(), UEPbrMaterial );
 }
 
-FDatasmithFacadeExpressionInput FDatasmithFacadeUEPbrMaterial::GetWorldDisplacement() const
-{
-	TSharedPtr<IDatasmithUEPbrMaterialElement> UEPbrMaterial = GetDatasmithUEPbrMaterialElement();
-	return FDatasmithFacadeExpressionInput( &UEPbrMaterial->GetWorldDisplacement(), UEPbrMaterial );
-}
-
 FDatasmithFacadeExpressionInput FDatasmithFacadeUEPbrMaterial::GetRefraction() const
 {
 	TSharedPtr<IDatasmithUEPbrMaterialElement> UEPbrMaterial = GetDatasmithUEPbrMaterialElement();
diff --git a/Engine/Source/Developer/Datasmith/DatasmithFacade/Public/DatasmithFacadeUEPbrMaterial.h b/Engine/Source/Developer/Datasmith/DatasmithFacade/Public/DatasmithFacadeUEPbrMaterial.h
--- a/Engine/Source/Developer/Datasmith/DatasmithFacade/Public/DatasmithFacadeUEPbrMaterial.h
+++ b/Engine/Source/Developer/Datasmith/DatasmithFacade/Public/DatasmithFacadeUEPbrMaterial.h
@@ -328,7 +328,6 @@ public:
 	FDatasmithFacadeExpressionInput GetEmissiveColor() const;
 	FDatasmithFacadeExpressionInput GetOpacity() const;
 	FDatasmithFacadeExpressionInput GetNormal() const;
-	FDatasmithFacadeExpressionInput GetWorldDisplacement() const;
 	FDatasmithFacadeExpressionInput GetRefraction() const;
 	FDatasmithFacadeExpressionInput GetAmbientOcclusion() const;
 	FDatasmithFacadeExpressionInput GetMaterialAttributes() const;
@@ -363,7 +362,7 @@ public:
 
 	/** Reset all expression to their default values and remove all connections */
 	void ResetExpressionGraph()
-	{ 
+	{
 		constexpr bool bRemoveAllExpressions = false;
 		GetDatasmithUEPbrMaterialElement()->ResetExpressionGraph( bRemoveAllExpressions );
 	}
diff --git a/Engine/Source/Programs/Enterprise/Datasmith/DatasmithFacadeCSharp/Private/DatasmithFacadeCSharp.cpp b/Engine/Source/Programs/Enterprise/Datasmith/DatasmithFacadeCSharp/Private/DatasmithFacadeCSharp.cpp
--- a/Engine/Source/Programs/Enterprise/Datasmith/DatasmithFacadeCSharp/Private/DatasmithFacadeCSharp.cpp
+++ b/Engine/Source/Programs/Enterprise/Datasmith/DatasmithFacadeCSharp/Private/DatasmithFacadeCSharp.cpp
@@ -1837,18 +1837,6 @@ SWIGEXPORT void * SWIGSTDCALL CSharp_FDatasmithFacadeUEPbrMaterial_GetNormal(voi
 }
 
 
-SWIGEXPORT void * SWIGSTDCALL CSharp_FDatasmithFacadeUEPbrMaterial_GetWorldDisplacement(void * jarg1) {
-  void * jresult ;
-  FDatasmithFacadeUEPbrMaterial *arg1 = (FDatasmithFacadeUEPbrMaterial *) 0 ;
-  SwigValueWrapper< FDatasmithFacadeExpressionInput > result;
-  
-  arg1 = (FDatasmithFacadeUEPbrMaterial *)jarg1; 
-  result = ((FDatasmithFacadeUEPbrMaterial const *)arg1)->GetWorldDisplacement();
-  jresult = new FDatasmithFacadeExpressionInput((const FDatasmithFacadeExpressionInput &)result); 
-  return jresult;
-}
-
-
 SWIGEXPORT void * SWIGSTDCALL CSharp_FDatasmithFacadeUEPbrMaterial_GetRefraction(void * jarg1) {
   void * jresult ;
   FDatasmithFacadeUEPbrMaterial *arg1 = (FDatasmithFacadeUEPbrMaterial *) 0 ;
diff --git a/Engine/Source/Programs/Enterprise/Datasmith/DatasmithFacadeCSharp/Public/DatasmithFacadeCSharpPINVOKE.cs b/Engine/Source/Programs/Enterprise/Datasmith/DatasmithFacadeCSharp/Public/DatasmithFacadeCSharpPINVOKE.cs
--- a/Engine/Source/Programs/Enterprise/Datasmith/DatasmithFacadeCSharp/Public/DatasmithFacadeCSharpPINVOKE.cs
+++ b/Engine/Source/Programs/Enterprise/Datasmith/DatasmithFacadeCSharp/Public/DatasmithFacadeCSharpPINVOKE.cs
@@ -614,9 +614,6 @@ class DatasmithFacadeCSharpPINVOKE {
   [global::System.Runtime.InteropServices.DllImport("DatasmithFacadeCSharp", EntryPoint="CSharp_FDatasmithFacadeUEPbrMaterial_GetNormal")]
   public static extern global::System.IntPtr FDatasmithFacadeUEPbrMaterial_GetNormal(global::System.Runtime.InteropServices.HandleRef jarg1);
 
-  [global::System.Runtime.InteropServices.DllImport("DatasmithFacadeCSharp", EntryPoint="CSharp_FDatasmithFacadeUEPbrMaterial_GetWorldDisplacement")]
-  public static extern global::System.IntPtr FDatasmithFacadeUEPbrMaterial_GetWorldDisplacement(global::System.Runtime.InteropServices.HandleRef jarg1);
-
   [global::System.Runtime.InteropServices.DllImport("DatasmithFacadeCSharp", EntryPoint="CSharp_FDatasmithFacadeUEPbrMaterial_GetRefraction")]
   public static extern global::System.IntPtr FDatasmithFacadeUEPbrMaterial_GetRefraction(global::System.Runtime.InteropServices.HandleRef jarg1);
 
diff --git a/Engine/Source/Programs/Enterprise/Datasmith/DatasmithFacadeCSharp/Public/FDatasmithFacadeUEPbrMaterial.cs b/Engine/Source/Programs/Enterprise/Datasmith/DatasmithFacadeCSharp/Public/FDatasmithFacadeUEPbrMaterial.cs
--- a/Engine/Source/Programs/Enterprise/Datasmith/DatasmithFacadeCSharp/Public/FDatasmithFacadeUEPbrMaterial.cs
+++ b/Engine/Source/Programs/Enterprise/Datasmith/DatasmithFacadeCSharp/Public/FDatasmithFacadeUEPbrMaterial.cs
@@ -73,11 +73,6 @@ public class FDatasmithFacadeUEPbrMaterial : FDatasmithFacadeBaseMaterial {
     return ret;
   }
 
-  public FDatasmithFacadeExpressionInput GetWorldDisplacement() {
-    FDatasmithFacadeExpressionInput ret = new FDatasmithFacadeExpressionInput(DatasmithFacadeCSharpPINVOKE.FDatasmithFacadeUEPbrMaterial_GetWorldDisplacement(swigCPtr), true);
-    return ret;
-  }
-
   public FDatasmithFacadeExpressionInput GetRefraction() {
     FDatasmithFacadeExpressionInput ret = new FDatasmithFacadeExpressionInput(DatasmithFacadeCSharpPINVOKE.FDatasmithFacadeUEPbrMaterial_GetRefraction(swigCPtr), true);
     return ret;
diff --git a/Engine/Source/Programs/Enterprise/Datasmith/DatasmithMaxExporter/Private/DatasmithMaxArnold.cpp b/Engine/Source/Programs/Enterprise/Datasmith/DatasmithMaxExporter/Private/DatasmithMaxArnold.cpp
--- a/Engine/Source/Programs/Enterprise/Datasmith/DatasmithMaxExporter/Private/DatasmithMaxArnold.cpp
+++ b/Engine/Source/Programs/Enterprise/Datasmith/DatasmithMaxExporter/Private/DatasmithMaxArnold.cpp
@@ -87,8 +87,8 @@ void FDatasmithMaxMatWriter::ExportPhysicalMaterialCoat(TSharedRef< IDatasmithSc
 		// Loop through all the defined parameters therein
 		for (int i = 0; i < ParamBlockDesc->count; i++)
 		{
-			const ParamDef& ParamDefinition = ParamBlockDesc->paramdefs[i];			
-			
+			const ParamDef& ParamDefinition = ParamBlockDesc->paramdefs[i];
+
 			if (FCString::Stricmp(ParamDefinition.int_name, TEXT("coat_map_on")) == 0)
 			{
 				if (ParamBlock2->GetInt(ParamDefinition.ID, GetCOREInterface()->GetTime()) == 0)
@@ -188,7 +188,7 @@ void FDatasmithMaxMatWriter::ExportPhysicalMaterial(TSharedRef< IDatasmithScene
 	bool bEmittanceMapOn = true, bEmittanceColorMapOn = true;
 	bool bRoughnessInverted = true;
 	bool bThinWalled = true;
-	bool bBumpmapOn = true, bDisplaceMapOn = true, bCutoutMapOn = true;
+	bool bBumpmapOn = true, bCutoutMapOn = true;
 
 	float DiffuseWeight = 0, Roughness = 0, Metalness = 0, Ior = 0;
 	float Transparency = 0;
@@ -202,7 +202,7 @@ void FDatasmithMaxMatWriter::ExportPhysicalMaterial(TSharedRef< IDatasmithScene
 	Texmap* ReflectivityMap = NULL, *RoughnessMap = NULL, *MetalnessMap = NULL, *ReflectionColorMap = NULL;
 	Texmap* TransparencyMap = NULL, *TransparencyColorMap = NULL;
 	Texmap* EmittanceMap = NULL, *EmittanceColorMap = NULL;
-	Texmap* BumpMap = NULL, *DisplaceMap = NULL, *CutoutMap = NULL;
+	Texmap* BumpMap = NULL, *CutoutMap = NULL;
 
 	int MaterialMode = 0; //0 means simple, just metalness, 1 means specular + metalness
 
@@ -323,13 +323,6 @@ void FDatasmithMaxMatWriter::ExportPhysicalMaterial(TSharedRef< IDatasmithScene
 					bBumpmapOn = false;
 				}
 			}
-			else if (FCString::Stricmp(ParamDefinition.int_name, TEXT("displacement_map_on")) == 0)
-			{
-				if (ParamBlock2->GetInt(ParamDefinition.ID, GetCOREInterface()->GetTime()) == 0)
-				{
-					bDisplaceMapOn = false;
-				}
-			}
 			else if (FCString::Stricmp(ParamDefinition.int_name, TEXT("cutout_map_on")) == 0)
 			{
 				if (ParamBlock2->GetInt(ParamDefinition.ID, GetCOREInterface()->GetTime()) == 0)
@@ -338,7 +331,6 @@ void FDatasmithMaxMatWriter::ExportPhysicalMaterial(TSharedRef< IDatasmithScene
 				}
 			}
 
-
 			// float values
 			else if (FCString::Stricmp(ParamDefinition.int_name, TEXT("base_weight")) == 0)
 			{
@@ -448,10 +440,6 @@ void FDatasmithMaxMatWriter::ExportPhysicalMaterial(TSharedRef< IDatasmithScene
 			{
 				BumpMap = ParamBlock2->GetTexmap(ParamDefinition.ID, GetCOREInterface()->GetTime());
 			}
-			else if (FCString::Stricmp(ParamDefinition.int_name, TEXT("displacement_map")) == 0)
-			{
-				DisplaceMap = ParamBlock2->GetTexmap(ParamDefinition.ID, GetCOREInterface()->GetTime());
-			}
 			else if (FCString::Stricmp(ParamDefinition.int_name, TEXT("cutout_map")) == 0)
 			{
 				CutoutMap = ParamBlock2->GetTexmap(ParamDefinition.ID, GetCOREInterface()->GetTime());
@@ -460,12 +448,12 @@ void FDatasmithMaxMatWriter::ExportPhysicalMaterial(TSharedRef< IDatasmithScene
 
 		ParamBlock2->ReleaseDesc();
 	}
-	
+
 	ExportPhysicalMaterialProperty(DatasmithScene, DiffuseColorMap, bDiffuseColorMapOn, DiffuseWeightMap, bDiffuseWeightMapOn, DiffuseColor, DiffuseWeight, MaterialShader->GetDiffuseComp(), DATASMITH_DIFFUSETEXNAME, DATASMITH_DIFFUSECOLNAME, false, false);
 	ExportPhysicalMaterialProperty(DatasmithScene, TransparencyColorMap, bTransparencyColorMapOn, TransparencyMap, bTransparencyMapOn, TransparencyColor, Transparency, MaterialShader->GetTransComp(), DATASMITH_TRANSPTEXNAME, DATASMITH_TRANSPCOLNAME, false, false);
 	ExportPhysicalMaterialProperty(DatasmithScene, EmittanceColorMap, bEmittanceColorMapOn, EmittanceMap, bEmittanceMapOn, EmittanceColor, EmittanceMultiplier, MaterialShader->GetEmitComp(), DATASMITH_EMITTEXNAME, DATASMITH_EMITCOLNAME, false, false);
 	ExportPhysicalMaterialProperty(DatasmithScene, CutoutMap, bCutoutMapOn, NULL, NULL, BMM_Color_fl(0.0, 0.0, 0.0, 0.0), 1.0, MaterialShader->GetMaskComp(), DATASMITH_CLIPTEXNAME, DATASMITH_CLIPTEXNAME, false, true);
-	
+
 	if (MetalnessMap != NULL && bMetalnessMapOn)
 	{
 		ExportPhysicalMaterialProperty(DatasmithScene, MetalnessMap, bMetalnessMapOn, NULL, false, BMM_Color_fl(0.0, 0.0, 0.0, 0.0), Metalness, MaterialShader->GetMetalComp(), DATASMITH_METALTEXNAME, DATASMITH_METALTEXNAME, false, true);
@@ -518,13 +506,6 @@ void FDatasmithMaxMatWriter::ExportPhysicalMaterial(TSharedRef< IDatasmithScene
 		}
 	}
 
-	if (DisplaceMap != NULL && bDisplaceMapOn == true)
-	{
-		ExportPhysicalMaterialProperty(DatasmithScene, DisplaceMap, bDisplaceMapOn, NULL, NULL, BMM_Color_fl(0.0, 0.0, 0.0, 0.0), 1.0, MaterialShader->GetDisplaceComp(), DATASMITH_DISPLACETEXNAME, DATASMITH_DISPLACETEXNAME, false, true);
-		MaterialShader->SetDisplace(10.0f);
-		MaterialShader->SetDisplaceSubDivision(4);
-	}
-
 	MaterialShader->SetIOR(Ior);
 	if (bThinWalled)
 	{
diff --git a/Engine/Source/Programs/Enterprise/Datasmith/DatasmithMaxExporter/Private/DatasmithMaxCorona.cpp b/Engine/Source/Programs/Enterprise/Datasmith/DatasmithMaxExporter/Private/DatasmithMaxCorona.cpp
--- a/Engine/Source/Programs/Enterprise/Datasmith/DatasmithMaxExporter/Private/DatasmithMaxCorona.cpp
+++ b/Engine/Source/Programs/Enterprise/Datasmith/DatasmithMaxExporter/Private/DatasmithMaxCorona.cpp
@@ -129,7 +129,6 @@ void FDatasmithMaxMatWriter::ExportCoronaMaterial(TSharedRef< IDatasmithScene >
 	bool bGlossyTexEnable = true;
 	bool bBumpTexEnable = true;
 	bool bOpacityTexEnable = true;
-	bool bDisplaceTexEnable = true;
 
 	float DiffuseTexAmount = 0.f;
 	float ReflectanceTexAmount = 0.f;
@@ -214,13 +213,6 @@ void FDatasmithMaxMatWriter::ExportCoronaMaterial(TSharedRef< IDatasmithScene >
 			{
 				BumpAmount = ParamBlock2->GetFloat(ParamDefinition.ID, GetCOREInterface()->GetTime());
 			}
-			else if (FCString::Stricmp(ParamDefinition.int_name, TEXT("texmapDisplace")) == 0)
-			{
-				if (ParamBlock2->GetTexmap(ParamDefinition.ID, GetCOREInterface()->GetTime()) == NULL)
-				{
-					bDisplaceTexEnable = false;
-				}
-			}
 			else if (FCString::Stricmp(ParamDefinition.int_name, TEXT("texmapOnDiffuse")) == 0)
 			{
 				if (ParamBlock2->GetInt(ParamDefinition.ID, GetCOREInterface()->GetTime()) == 0)
@@ -272,13 +264,6 @@ void FDatasmithMaxMatWriter::ExportCoronaMaterial(TSharedRef< IDatasmithScene >
 					bBumpTexEnable = false;
 				}
 			}
-			else if (FCString::Stricmp(ParamDefinition.int_name, TEXT("texmapOnDisplacement")) == 0)
-			{
-				if (ParamBlock2->GetInt(ParamDefinition.ID, GetCOREInterface()->GetTime()) == 0)
-				{
-					bDisplaceTexEnable = false;
-				}
-			}
 			else if (FCString::Stricmp(ParamDefinition.int_name, TEXT("texmapOnOpacity")) == 0)
 			{
 				if (ParamBlock2->GetInt(ParamDefinition.ID, GetCOREInterface()->GetTime()) == 0)
@@ -528,16 +513,6 @@ void FDatasmithMaxMatWriter::ExportCoronaMaterial(TSharedRef< IDatasmithScene >
 					}
 				}
 			}
-			else if (FCString::Stricmp(ParamDefinition.int_name, TEXT("texmapDisplace")) == 0 && bDisplaceTexEnable)
-			{
-				Texmap* LocalTex = ParamBlock2->GetTexmap(ParamDefinition.ID, GetCOREInterface()->GetTime());
-				if (LocalTex)
-				{
-					DumpTexture(DatasmithScene, MaterialShader->GetDisplaceComp(), LocalTex, DATASMITH_DISPLACETEXNAME, DATASMITH_DISPLACETEXNAME, false, true);
-					MaterialShader->SetDisplace(10.f);
-					MaterialShader->SetDisplaceSubDivision(4.0);
-				}
-			}
 		}
 		ParamBlock2->ReleaseDesc();
 	}
diff --git a/Engine/Source/Programs/Enterprise/Datasmith/DatasmithMaxExporter/Private/DatasmithMaxExporter.cpp b/Engine/Source/Programs/Enterprise/Datasmith/DatasmithMaxExporter/Private/DatasmithMaxExporter.cpp
--- a/Engine/Source/Programs/Enterprise/Datasmith/DatasmithMaxExporter/Private/DatasmithMaxExporter.cpp
+++ b/Engine/Source/Programs/Enterprise/Datasmith/DatasmithMaxExporter/Private/DatasmithMaxExporter.cpp
@@ -225,7 +225,6 @@ namespace DatasmithMaxExporterUtils
 			RemapUVChannelsForCompositeTexture(ShaderElement->GetBumpComp(), UVChannels);
 			RemapUVChannelsForCompositeTexture(ShaderElement->GetTransComp(), UVChannels);
 			RemapUVChannelsForCompositeTexture(ShaderElement->GetMaskComp(), UVChannels);
-			RemapUVChannelsForCompositeTexture(ShaderElement->GetDisplaceComp(), UVChannels);
 			RemapUVChannelsForCompositeTexture(ShaderElement->GetMetalComp(), UVChannels);
 			RemapUVChannelsForCompositeTexture(ShaderElement->GetEmitComp(), UVChannels);
 			RemapUVChannelsForCompositeTexture(ShaderElement->GetWeightComp(), UVChannels);
diff --git a/Engine/Source/Programs/Enterprise/Datasmith/DatasmithMaxExporter/Private/DatasmithMaxScanline.cpp b/Engine/Source/Programs/Enterprise/Datasmith/DatasmithMaxExporter/Private/DatasmithMaxScanline.cpp
--- a/Engine/Source/Programs/Enterprise/Datasmith/DatasmithMaxExporter/Private/DatasmithMaxScanline.cpp
+++ b/Engine/Source/Programs/Enterprise/Datasmith/DatasmithMaxExporter/Private/DatasmithMaxScanline.cpp
@@ -544,7 +544,6 @@ enum class EScanlineMaterialMaps
 	Bump,
 	Reflection,
 	Refraction,
-	Displacement
 };
 
 void FDatasmithMaxMatWriter::ExportStandardMaterial(TSharedRef< IDatasmithScene > DatasmithScene, TSharedPtr< IDatasmithMaterialElement >& MaterialElement, Mtl* Material)
@@ -558,14 +557,12 @@ void FDatasmithMaxMatWriter::ExportStandardMaterial(TSharedRef< IDatasmithScene
 	bool bMaskTexEnable = true;
 	bool bGlossyTexEnable = true;
 	bool bBumpTexEnable = true;
-	bool bDisplaceTexEnable = true;
 
 	float DiffuseTexAmount = 0.f;
 	float ReflectanceTexAmount = 0.f;
 	float MaskTexAmount = 0.f;
 	float GlossyTexAmount = 0.f;
 	float BumpTexAmount = 0.f;
-	float DisplaceTexAmount = 0.f;
 
 	bool bUseSelfIllumColor = true;
 
@@ -610,11 +607,6 @@ void FDatasmithMaxMatWriter::ExportStandardMaterial(TSharedRef< IDatasmithScene
 				{
 					bBumpTexEnable = false;
 				}
-
-				if (ParamBlock2->GetTexmap(ParamDefinition.ID, GetCOREInterface()->GetTime(), (int)EScanlineMaterialMaps::Displacement) == nullptr)
-				{
-					bDisplaceTexEnable = false;
-				}
 			}
 
 			if (FCString::Stricmp(ParamDefinition.int_name, TEXT("mapEnables")) == 0)
@@ -643,11 +635,6 @@ void FDatasmithMaxMatWriter::ExportStandardMaterial(TSharedRef< IDatasmithScene
 				{
 					bBumpTexEnable = false;
 				}
-
-				if (ParamBlock2->GetInt(ParamDefinition.ID, GetCOREInterface()->GetTime(), (int)EScanlineMaterialMaps::Displacement) == 0)
-				{
-					bDisplaceTexEnable = false;
-				}
 			}
 
 			if (FCString::Stricmp(ParamDefinition.int_name, TEXT("mapAmounts")) == 0)
@@ -657,7 +644,6 @@ void FDatasmithMaxMatWriter::ExportStandardMaterial(TSharedRef< IDatasmithScene
 				GlossyTexAmount = ParamBlock2->GetFloat(ParamDefinition.ID, GetCOREInterface()->GetTime(), (int)EScanlineMaterialMaps::Glossiness);
 				MaskTexAmount = ParamBlock2->GetFloat(ParamDefinition.ID, GetCOREInterface()->GetTime(), (int)EScanlineMaterialMaps::Opacity);
 				BumpTexAmount = ParamBlock2->GetFloat(ParamDefinition.ID, GetCOREInterface()->GetTime(), (int)EScanlineMaterialMaps::Bump);
-				DisplaceTexAmount = ParamBlock2->GetFloat(ParamDefinition.ID, GetCOREInterface()->GetTime(), (int)EScanlineMaterialMaps::Displacement);
 			}
 
 			if (FCString::Stricmp(ParamDefinition.int_name, TEXT("SpecularLevel")) == 0)
@@ -701,10 +687,6 @@ void FDatasmithMaxMatWriter::ExportStandardMaterial(TSharedRef< IDatasmithScene
 				{
 					MaterialShader->SetBumpAmount( ParamBlock2->GetFloat(ParamDefinition.ID, GetCOREInterface()->GetTime(), 8) );
 				}
-				if (bDisplaceTexEnable == true)
-				{
-					MaterialShader->SetDisplace( 0.1f * ParamBlock2->GetFloat(ParamDefinition.ID, GetCOREInterface()->GetTime(), 11) );
-				}
 			}
 			else if (FCString::Stricmp(ParamDefinition.int_name, TEXT("ior")) == 0)
 			{
@@ -799,12 +781,6 @@ void FDatasmithMaxMatWriter::ExportStandardMaterial(TSharedRef< IDatasmithScene
 						DumpTexture(DatasmithScene, MaterialShader->GetBumpComp(), LocalTex, DATASMITH_BUMPTEXNAME, DATASMITH_BUMPTEXNAME, false, true);
 					}
 				}
-				if (bDisplaceTexEnable == true && DisplaceTexAmount > 0)
-				{
-					Texmap* LocalTex = ParamBlock2->GetTexmap(ParamDefinition.ID, GetCOREInterface()->GetTime(), (int)EScanlineMaterialMaps::Displacement);
-					DumpTexture(DatasmithScene, MaterialShader->GetDisplaceComp(), LocalTex, DATASMITH_DISPLACETEXNAME, DATASMITH_DISPLACETEXNAME, false, true);
-					MaterialShader->SetDisplaceSubDivision(4.0);
-				}
 			}
 			else if (FCString::Stricmp(ParamDefinition.int_name, TEXT("selfillumMap")) == 0 && bUseSelfIllumColor == true && SelfIllumTex != nullptr)
 			{
diff --git a/Engine/Source/Programs/Enterprise/Datasmith/DatasmithMaxExporter/Private/DatasmithMaxThea.cpp b/Engine/Source/Programs/Enterprise/Datasmith/DatasmithMaxExporter/Private/DatasmithMaxThea.cpp
--- a/Engine/Source/Programs/Enterprise/Datasmith/DatasmithMaxExporter/Private/DatasmithMaxThea.cpp
+++ b/Engine/Source/Programs/Enterprise/Datasmith/DatasmithMaxExporter/Private/DatasmithMaxThea.cpp
@@ -407,11 +407,6 @@ void FDatasmithMaxMatWriter::ExportTheaMaterial(TSharedRef< IDatasmithScene > Da
 	bool bClipByAmount = false;
 	float ClipAmount = 1.0;
 
-	bool bDisplaceEnabled = true;
-	Texmap* DisplaceTex = NULL;
-	int DisplaceSubdivs = 4;
-	float DisplaceAmount = 10;
-
 	bool bEmitEnabled = false;
 	int EmitMode = 0;
 	BMM_Color_fl EmitColor = 0;
@@ -468,10 +463,6 @@ void FDatasmithMaxMatWriter::ExportTheaMaterial(TSharedRef< IDatasmithScene > Da
 				StackedMtlWeight[2] = float(ParamBlock2->GetInt(ParamDefinition.ID, GetCOREInterface()->GetTime(), 6)) / 100.0f;
 				StackedMtlWeight[3] = float(ParamBlock2->GetInt(ParamDefinition.ID, GetCOREInterface()->GetTime(), 7)) / 100.0f;
 			}
-			else if (FCString::Stricmp(ParamDefinition.int_name, TEXT("disp_enabled")) == 0)
-			{
-				bDisplaceEnabled = ParamBlock2->GetInt(ParamDefinition.ID, GetCOREInterface()->GetTime()) != 0;
-			}
 			else if (FCString::Stricmp(ParamDefinition.int_name, TEXT("clip_enabled")) == 0)
 			{
 				bClipEnabled = ParamBlock2->GetInt(ParamDefinition.ID, GetCOREInterface()->GetTime()) != 0;
@@ -480,10 +471,6 @@ void FDatasmithMaxMatWriter::ExportTheaMaterial(TSharedRef< IDatasmithScene > Da
 			{
 				ClipTex = ParamBlock2->GetTexmap(ParamDefinition.ID, GetCOREInterface()->GetTime());
 			}
-			else if (FCString::Stricmp(ParamDefinition.int_name, TEXT("disp_tex")) == 0)
-			{
-				DisplaceTex = ParamBlock2->GetTexmap(ParamDefinition.ID, GetCOREInterface()->GetTime());
-			}
 			else if (FCString::Stricmp(ParamDefinition.int_name, TEXT("soft")) == 0)
 			{
 				bSoftClip = ParamBlock2->GetInt(ParamDefinition.ID, GetCOREInterface()->GetTime()) != 0;
@@ -496,14 +483,6 @@ void FDatasmithMaxMatWriter::ExportTheaMaterial(TSharedRef< IDatasmithScene > Da
 			{
 				ClipAmount = ParamBlock2->GetFloat(ParamDefinition.ID, GetCOREInterface()->GetTime());
 			}
-			else if (FCString::Stricmp(ParamDefinition.int_name, TEXT("disp_height")) == 0)
-			{
-				DisplaceAmount = ParamBlock2->GetFloat(ParamDefinition.ID, GetCOREInterface()->GetTime()) * 2.0f;
-			}
-			else if (FCString::Stricmp(ParamDefinition.int_name, TEXT("disp_subdiv")) == 0)
-			{
-				DisplaceSubdivs = ParamBlock2->GetInt(ParamDefinition.ID, GetCOREInterface()->GetTime());
-			}
 			else if (FCString::Stricmp(ParamDefinition.int_name, TEXT("emit_enabled")) == 0)
 			{
 				bEmitEnabled = ParamBlock2->GetInt(ParamDefinition.ID, GetCOREInterface()->GetTime()) != 0;
@@ -564,16 +543,6 @@ void FDatasmithMaxMatWriter::ExportTheaMaterial(TSharedRef< IDatasmithScene > Da
 		}
 	}
 
-	if (bDisplaceEnabled)
-	{
-		if (DisplaceTex)
-		{
-			DumpTexture(DatasmithScene, MaterialShader->GetDisplaceComp(), DisplaceTex, DATASMITH_CLIPTEXNAME, DATASMITH_CLIPTEXNAME, false, true);
-			MaterialShader->SetDisplace( DisplaceAmount );
-			MaterialShader->SetDisplaceSubDivision( DisplaceSubdivs );
-		}
-	}
-
 	if (bEmitEnabled)
 	{
 		if (EmitMode == 0)
@@ -703,10 +672,6 @@ void FDatasmithMaxMatWriter::ExportTheaMaterial(TSharedRef< IDatasmithScene > Da
 
 	for (int i = 1; i < AllShaders.Num(); i++)
 	{
-		AllShaders[i]->SetDisplace( AllShaders[0]->GetDisplace() );
-		AllShaders[i]->SetDisplaceComp( AllShaders[0]->GetDisplaceComp() );
-		AllShaders[i]->SetDisplaceSubDivision( AllShaders[0]->GetDisplaceSubDivision() );
-
 		AllShaders[i]->SetEmitComp( AllShaders[0]->GetEmitComp() );
 		AllShaders[i]->SetEmitPower( AllShaders[0]->GetEmitPower() );
 		AllShaders[i]->SetEmitTemperature( AllShaders[0]->GetEmitTemperature() );
diff --git a/Engine/Source/Programs/Enterprise/Datasmith/DatasmithMaxExporter/Private/DatasmithMaxVRay.cpp b/Engine/Source/Programs/Enterprise/Datasmith/DatasmithMaxExporter/Private/DatasmithMaxVRay.cpp
--- a/Engine/Source/Programs/Enterprise/Datasmith/DatasmithMaxExporter/Private/DatasmithMaxVRay.cpp
+++ b/Engine/Source/Programs/Enterprise/Datasmith/DatasmithMaxExporter/Private/DatasmithMaxVRay.cpp
@@ -147,7 +147,7 @@ void FDatasmithMaxMatWriter::GetVrayHdri(TSharedRef< IDatasmithScene > Datasmith
 		ParamBlock2->ReleaseDesc();
 	}
 
-	
+
 
 	if (Path.IsEmpty())
 	{
@@ -192,10 +192,8 @@ void FDatasmithMaxMatWriter::ExportVRayMaterial(TSharedRef< IDatasmithScene > Da
 	bool bOpacityTexEnable = true;
 	bool bLockIor = true;
 	bool bRefleFresnel = true;
-	bool bDisplaceTexEnable = true;
 
 	float BumpAmount = 0.f;
-	float DisplaceAmount = 5.f;
 
 	BMM_Color_fl ColorDiffuse;
 	BMM_Color_fl ColorReflection;
@@ -270,17 +268,6 @@ void FDatasmithMaxMatWriter::ExportVRayMaterial(TSharedRef< IDatasmithScene > Da
 			{
 				BumpAmount = ParamBlock2->GetFloat(ParamDefinition.ID, GetCOREInterface()->GetTime());
 			}
-			else if (FCString::Stricmp(ParamDefinition.int_name, TEXT("texmap_displacement")) == 0)
-			{
-				if (ParamBlock2->GetTexmap(ParamDefinition.ID, GetCOREInterface()->GetTime()) == NULL)
-				{
-					bDisplaceTexEnable = false;
-				}
-			}
-			else if (FCString::Stricmp(ParamDefinition.int_name, TEXT("texmap_displacement_multiplier")) == 0)
-			{
-				DisplaceAmount = ParamBlock2->GetFloat(ParamDefinition.ID, GetCOREInterface()->GetTime()) * 0.1f;
-			}
 			else if (FCString::Stricmp(ParamDefinition.int_name, TEXT("texmap_diffuse_on")) == 0)
 			{
 				if (ParamBlock2->GetInt(ParamDefinition.ID, GetCOREInterface()->GetTime()) == 0)
@@ -374,13 +361,6 @@ void FDatasmithMaxMatWriter::ExportVRayMaterial(TSharedRef< IDatasmithScene > Da
 					bBumpTexEnable = false;
 				}
 			}
-			else if (FCString::Stricmp(ParamDefinition.int_name, TEXT("texmap_displacement_on")) == 0)
-			{
-				if (ParamBlock2->GetInt(ParamDefinition.ID, GetCOREInterface()->GetTime()) == 0)
-				{
-					bDisplaceTexEnable = false;
-				}
-			}
 			else if (FCString::Stricmp(ParamDefinition.int_name, TEXT("texmap_opacity_on")) == 0)
 			{
 				if (ParamBlock2->GetInt(ParamDefinition.ID, GetCOREInterface()->GetTime()) == 0)
@@ -579,16 +559,6 @@ void FDatasmithMaxMatWriter::ExportVRayMaterial(TSharedRef< IDatasmithScene > Da
 					}
 				}
 			}
-			else if (FCString::Stricmp(ParamDefinition.int_name, TEXT("texmap_displacement")) == 0 && bDisplaceTexEnable == true)
-			{
-				Texmap* LocalTex = ParamBlock2->GetTexmap(ParamDefinition.ID, GetCOREInterface()->GetTime());
-				if (LocalTex)
-				{
-					DumpTexture(DatasmithScene, MaterialShader->GetDisplaceComp(), LocalTex, DATASMITH_DISPLACETEXNAME, DATASMITH_DISPLACETEXNAME, false, true);
-					MaterialShader->SetDisplace( DisplaceAmount );
-					MaterialShader->SetDisplaceSubDivision(4.0);
-				}
-			}
 		}
 		ParamBlock2->ReleaseDesc();
 	}
diff --git a/Engine/Source/Programs/Enterprise/Datasmith/DatasmithMaxExporter/Private/MaxMaterialsToUEPbr/DatasmithMaxCoronaMaterialsToUEPbr.cpp b/Engine/Source/Programs/Enterprise/Datasmith/DatasmithMaxExporter/Private/MaxMaterialsToUEPbr/DatasmithMaxCoronaMaterialsToUEPbr.cpp
--- a/Engine/Source/Programs/Enterprise/Datasmith/DatasmithMaxExporter/Private/MaxMaterialsToUEPbr/DatasmithMaxCoronaMaterialsToUEPbr.cpp
+++ b/Engine/Source/Programs/Enterprise/Datasmith/DatasmithMaxExporter/Private/MaxMaterialsToUEPbr/DatasmithMaxCoronaMaterialsToUEPbr.cpp
@@ -49,9 +49,6 @@ namespace DatasmithMaxCoronaMaterialsToUEPbrImpl
 
 		// Bump
 		DatasmithMaxTexmapParser::FMapParameter BumpMap;
-
-		// Displacement
-		DatasmithMaxTexmapParser::FMapParameter DisplacementMap;
 	};
 
 	FMaxCoronaMaterial ParseCoronaMaterialProperties( Mtl& Material )
@@ -66,7 +63,7 @@ namespace DatasmithMaxCoronaMaterialsToUEPbrImpl
 		{
 			IParamBlock2* ParamBlock2 = Material.GetParamBlockByID((short)j);
 			ParamBlockDesc2* ParamBlockDesc = ParamBlock2->GetDesc();
-			
+
 			for (int i = 0; i < ParamBlockDesc->count; i++)
 			{
 				const ParamDef& ParamDefinition = ParamBlockDesc->paramdefs[i];
@@ -208,16 +205,6 @@ namespace DatasmithMaxCoronaMaterialsToUEPbrImpl
 				{
 					CoronaMaterialProperties.BumpMap.bEnabled = ( ParamBlock2->GetInt( ParamDefinition.ID, CurrentTime ) != 0 );
 				}
-
-				// Displacement
-				else if (FCString::Stricmp(ParamDefinition.int_name, TEXT("texmapDisplace")) == 0)
-				{
-					CoronaMaterialProperties.DisplacementMap.Map = ParamBlock2->GetTexmap( ParamDefinition.ID, CurrentTime );
-				}
-				else if (FCString::Stricmp(ParamDefinition.int_name, TEXT("texmapOnDisplacement")) == 0)
-				{
-					CoronaMaterialProperties.DisplacementMap.bEnabled = ( ParamBlock2->GetInt( ParamDefinition.ID, CurrentTime ) != 0 );
-				}
 			}
 			ParamBlock2->ReleaseDesc();
 		}
@@ -234,7 +221,7 @@ namespace DatasmithMaxCoronaMaterialsToUEPbrImpl
 
 			DatasmithMaxTexmapParser::FMapParameter Mask;
 		};
-		
+
 		Mtl* BaseMaterial = nullptr;
 		static const int32 MaximumNumberOfCoat = 10;
 		FCoronaCoatMaterialProperties CoatedMaterials[MaximumNumberOfCoat];
@@ -417,23 +404,6 @@ void FDatasmithMaxCoronaMaterialsToUEPbr::Convert( TSharedRef< IDatasmithScene >
 
 		ConvertState.bCanBake = true;
 	}
-	
-	// Displacement
-	{
-		ConvertState.DefaultTextureMode = EDatasmithTextureMode::Displace;
-
-		IDatasmithMaterialExpression* DisplacementExpression = FDatasmithMaxTexmapToUEPbrUtils::MapOrValue( this, CoronaMaterialProperties.DisplacementMap, TEXT("Displacement Map"), TOptional< FLinearColor >(), TOptional< float >() );
-
-		if ( DisplacementExpression )
-		{
-			DisplacementExpression->ConnectExpression( PbrMaterialElement->GetWorldDisplacement() );
-		}
-
-		if ( DisplacementExpression )
-		{
-			DisplacementExpression->SetName( TEXT("Displacement Map") );
-		}
-	}
 
 	ConvertState.DefaultTextureMode = EDatasmithTextureMode::Specular; // At this point, all maps are considered specular maps
 
@@ -504,7 +474,7 @@ void FDatasmithMaxCoronaMaterialsToUEPbr::Convert( TSharedRef< IDatasmithScene >
 		DiffuseLerpExpression->ConnectExpression( PbrMaterialElement->GetBaseColor() );
 
 		IDatasmithMaterialExpression* ReflectionIOR = nullptr;
-		
+
 		{
 			TGuardValue< bool > SetIsMonoChannel( ConvertState.bIsMonoChannel, true );
 			ReflectionIOR = FDatasmithMaxTexmapToUEPbrUtils::MapOrValue( this, CoronaMaterialProperties.ReflectionIORMap, TEXT("Fresnel IOR"), TOptional< FLinearColor >(), CoronaMaterialProperties.ReflectionIOR );
@@ -576,7 +546,7 @@ void FDatasmithMaxCoronaMaterialsToUEPbr::Convert( TSharedRef< IDatasmithScene >
 	{
 		MetallicExpression->ConnectExpression( PbrMaterialElement->GetMetallic() );
 	}
-	
+
 	// UE Specular
 	if ( MetallicExpression )
 	{
@@ -725,12 +695,12 @@ void FDatasmithMaxCoronaBlendMaterialToUEPbr::Convert( TSharedRef<IDatasmithScen
 
 			IDatasmithMaterialExpression* MaskExpression = FDatasmithMaxTexmapToUEPbrUtils::MapOrValue(this, CoatedMaterial.Mask, TEXT("MixAmount"),
 				FLinearColor::White, TOptional< float >());
-			
+
 			IDatasmithMaterialExpressionGeneric* AlphaExpression = PbrMaterialElement->AddMaterialExpression<IDatasmithMaterialExpressionGeneric>();
 			AlphaExpression->SetExpressionName(TEXT("Multiply"));
 
 			//AlphaExpression is nullptr only when there is no mask and the mask weight is ~100% so we add scalar 0 instead.
-			if (!MaskExpression) 
+			if (!MaskExpression)
 			{
 				IDatasmithMaterialExpressionScalar* WeightExpression = PbrMaterialElement->AddMaterialExpression< IDatasmithMaterialExpressionScalar >();
 				WeightExpression->GetScalar() = 0.f;
diff --git a/Engine/Source/Programs/Enterprise/Datasmith/DatasmithMaxExporter/Private/MaxMaterialsToUEPbr/DatasmithMaxScanlineMaterialsToUEPbr.cpp b/Engine/Source/Programs/Enterprise/Datasmith/DatasmithMaxExporter/Private/MaxMaterialsToUEPbr/DatasmithMaxScanlineMaterialsToUEPbr.cpp
--- a/Engine/Source/Programs/Enterprise/Datasmith/DatasmithMaxExporter/Private/MaxMaterialsToUEPbr/DatasmithMaxScanlineMaterialsToUEPbr.cpp
+++ b/Engine/Source/Programs/Enterprise/Datasmith/DatasmithMaxExporter/Private/MaxMaterialsToUEPbr/DatasmithMaxScanlineMaterialsToUEPbr.cpp
@@ -30,7 +30,6 @@ namespace DatasmithMaxScanlineMaterialsToUEPbrImpl
 		Bump,
 		Reflection,
 		Refraction,
-		Displacement
 	};
 
 	struct FMaxScanlineMaterial
@@ -60,9 +59,6 @@ namespace DatasmithMaxScanlineMaterialsToUEPbrImpl
 		// Bump
 		DatasmithMaxTexmapParser::FMapParameter BumpMap;
 
-		// Displacement
-		DatasmithMaxTexmapParser::FMapParameter DisplacementMap;
-
 		// Self-illumination
 		bool bUseSelfIllumColor = false;
 		FLinearColor SelfIllumColor;
@@ -88,7 +84,7 @@ namespace DatasmithMaxScanlineMaterialsToUEPbrImpl
 		{
 			IParamBlock2* ParamBlock2 = Material.GetParamBlockByID((short)j);
 			ParamBlockDesc2* ParamBlockDesc = ParamBlock2->GetDesc();
-			
+
 			for (int i = 0; i < ParamBlockDesc->count; i++)
 			{
 				const ParamDef& ParamDefinition = ParamBlockDesc->paramdefs[i];
@@ -102,7 +98,6 @@ namespace DatasmithMaxScanlineMaterialsToUEPbrImpl
 					ScanlineMaterialProperties.GlossinessMap.Map = ParamBlock2->GetTexmap( ParamDefinition.ID, GetCOREInterface()->GetTime(), (int)EScanlineMaterialMaps::Glossiness );
 					ScanlineMaterialProperties.OpacityMap.Map = ParamBlock2->GetTexmap( ParamDefinition.ID, GetCOREInterface()->GetTime(), (int)EScanlineMaterialMaps::Opacity );
 					ScanlineMaterialProperties.BumpMap.Map = ParamBlock2->GetTexmap( ParamDefinition.ID, GetCOREInterface()->GetTime(), (int)EScanlineMaterialMaps::Bump );
-					ScanlineMaterialProperties.DisplacementMap.Map = ParamBlock2->GetTexmap( ParamDefinition.ID, GetCOREInterface()->GetTime(), (int)EScanlineMaterialMaps::Displacement );
 					ScanlineMaterialProperties.SelfIllumMap.Map = ParamBlock2->GetTexmap( ParamDefinition.ID, GetCOREInterface()->GetTime(), (int)EScanlineMaterialMaps::SelfIllumination );
 				}
 				else if ( FCString::Stricmp( ParamDefinition.int_name, TEXT("mapEnables") ) == 0 )
@@ -113,7 +108,6 @@ namespace DatasmithMaxScanlineMaterialsToUEPbrImpl
 					ScanlineMaterialProperties.GlossinessMap.bEnabled = ( ParamBlock2->GetInt( ParamDefinition.ID, GetCOREInterface()->GetTime(), (int)EScanlineMaterialMaps::Glossiness ) != 0 );
 					ScanlineMaterialProperties.OpacityMap.bEnabled = ( ParamBlock2->GetInt( ParamDefinition.ID, GetCOREInterface()->GetTime(), (int)EScanlineMaterialMaps::Opacity ) != 0 );
 					ScanlineMaterialProperties.BumpMap.bEnabled = ( ParamBlock2->GetInt( ParamDefinition.ID, GetCOREInterface()->GetTime(), (int)EScanlineMaterialMaps::Bump ) != 0 );
-					ScanlineMaterialProperties.DisplacementMap.bEnabled = ( ParamBlock2->GetInt( ParamDefinition.ID, GetCOREInterface()->GetTime(), (int)EScanlineMaterialMaps::Displacement ) != 0 );
 					ScanlineMaterialProperties.SelfIllumMap.bEnabled = ( ParamBlock2->GetInt( ParamDefinition.ID, GetCOREInterface()->GetTime(), (int)EScanlineMaterialMaps::SelfIllumination ) != 0 );
 				}
 				else if ( FCString::Stricmp( ParamDefinition.int_name, TEXT("mapAmounts") ) == 0 )
@@ -124,7 +118,6 @@ namespace DatasmithMaxScanlineMaterialsToUEPbrImpl
 					ScanlineMaterialProperties.GlossinessMap.Weight = ParamBlock2->GetFloat( ParamDefinition.ID, GetCOREInterface()->GetTime(), (int)EScanlineMaterialMaps::Glossiness );
 					ScanlineMaterialProperties.OpacityMap.Weight = ParamBlock2->GetFloat( ParamDefinition.ID, GetCOREInterface()->GetTime(), (int)EScanlineMaterialMaps::Opacity );
 					ScanlineMaterialProperties.BumpMap.Weight = ParamBlock2->GetFloat( ParamDefinition.ID, GetCOREInterface()->GetTime(), (int)EScanlineMaterialMaps::Bump );
-					ScanlineMaterialProperties.DisplacementMap.Weight = ParamBlock2->GetFloat( ParamDefinition.ID, GetCOREInterface()->GetTime(), (int)EScanlineMaterialMaps::Displacement );
 					ScanlineMaterialProperties.SelfIllumMap.Weight = ParamBlock2->GetFloat( ParamDefinition.ID, GetCOREInterface()->GetTime(), (int)EScanlineMaterialMaps::SelfIllumination );
 				}
 				else if ( FCString::Stricmp( ParamDefinition.int_name, TEXT("Opacity") ) == 0 )
@@ -201,7 +194,7 @@ void FDatasmithMaxScanlineMaterialsToUEPbr::Convert( TSharedRef< IDatasmithScene
 	if ( SpecularColorExpression )
 	{
 		SpecularColorExpression->SetName( TEXT("Specular") );
-		
+
 		IDatasmithMaterialExpressionScalar* SpecularLevelExpression = PbrMaterialElement->AddMaterialExpression< IDatasmithMaterialExpressionScalar >();
 		SpecularLevelExpression->SetName( TEXT("Specular Level") );
 		SpecularLevelExpression->GetScalar() = ScanlineMaterialProperties.SpecularLevel;
@@ -254,18 +247,6 @@ void FDatasmithMaxScanlineMaterialsToUEPbr::Convert( TSharedRef< IDatasmithScene
 
 		ConvertState.bCanBake = true;
 	}
-	
-	// Displacement
-	{
-		ConvertState.DefaultTextureMode = EDatasmithTextureMode::Displace;
-
-		IDatasmithMaterialExpression* DisplacementExpression = FDatasmithMaxTexmapToUEPbrUtils::MapOrValue( this, ScanlineMaterialProperties.DisplacementMap, TEXT("Displacement Map"), TOptional< FLinearColor >(), TOptional< float >() );
-
-		if ( DisplacementExpression )
-		{
-			DisplacementExpression->ConnectExpression( PbrMaterialElement->GetWorldDisplacement() );
-		}
-	}
 
 	// ConvertFromDiffSpec
 	{
@@ -409,7 +390,7 @@ void FDatasmithMaxBlendMaterialsToUEPbr::Convert( TSharedRef< IDatasmithScene >
 		}
 
 		//AlphaExpression is nullptr only when there is no mask and the mask weight is ~100% so we add scalar 0 instead.
-		if ( !AlphaExpression ) 
+		if ( !AlphaExpression )
 		{
 			IDatasmithMaterialExpressionScalar* WeightExpression = PbrMaterialElement->AddMaterialExpression< IDatasmithMaterialExpressionScalar >();
 			WeightExpression->GetScalar() = 0.f;
diff --git a/Engine/Source/Programs/Enterprise/Datasmith/DatasmithMaxExporter/Private/MaxMaterialsToUEPbr/DatasmithMaxVrayMaterialsToUEPbr.cpp b/Engine/Source/Programs/Enterprise/Datasmith/DatasmithMaxExporter/Private/MaxMaterialsToUEPbr/DatasmithMaxVrayMaterialsToUEPbr.cpp
--- a/Engine/Source/Programs/Enterprise/Datasmith/DatasmithMaxExporter/Private/MaxMaterialsToUEPbr/DatasmithMaxVrayMaterialsToUEPbr.cpp
+++ b/Engine/Source/Programs/Enterprise/Datasmith/DatasmithMaxExporter/Private/MaxMaterialsToUEPbr/DatasmithMaxVrayMaterialsToUEPbr.cpp
@@ -84,7 +84,7 @@ namespace DatasmithMaxVRayMaterialsToUEPbrImpl
 		{
 			IParamBlock2* ParamBlock2 = Material.GetParamBlockByID((short)j);
 			ParamBlockDesc2* ParamBlockDesc = ParamBlock2->GetDesc();
-			
+
 			for (int i = 0; i < ParamBlockDesc->count; i++)
 			{
 				const ParamDef& ParamDefinition = ParamBlockDesc->paramdefs[i];
@@ -302,7 +302,7 @@ namespace DatasmithMaxVRayMaterialsToUEPbrImpl
 		{
 			IParamBlock2* ParamBlock2 = Material.GetParamBlockByID((short)j);
 			ParamBlockDesc2* ParamBlockDesc = ParamBlock2->GetDesc();
-			
+
 			for (int i = 0; i < ParamBlockDesc->count; i++)
 			{
 				const ParamDef& ParamDefinition = ParamBlockDesc->paramdefs[i];
@@ -342,7 +342,7 @@ namespace DatasmithMaxVRayMaterialsToUEPbrImpl
 		{
 			IParamBlock2* ParamBlock2 = Material.GetParamBlockByID((short)j);
 			ParamBlockDesc2* ParamBlockDesc = ParamBlock2->GetDesc();
-			
+
 			for (int i = 0; i < ParamBlockDesc->count; i++)
 			{
 				const ParamDef& ParamDefinition = ParamBlockDesc->paramdefs[i];
@@ -367,7 +367,7 @@ namespace DatasmithMaxVRayMaterialsToUEPbrImpl
 			DatasmithMaxTexmapParser::FMapParameter MaterialBlendParameter;
 			FLinearColor MixColor;
 		};
-		
+
 		Mtl* BaseMaterial = nullptr;
 		static const int32 MaximumNumberOfCoat = 10;
 		FVRayCoatMaterialProperties CoatedMaterials[MaximumNumberOfCoat];
@@ -487,7 +487,7 @@ void FDatasmithMaxVRayMaterialsToUEPbr::Convert( TSharedRef< IDatasmithScene > D
 	}
 
 	IDatasmithMaterialExpressionGeneric* ReflectionIntensityExpression = PbrMaterialElement->AddMaterialExpression< IDatasmithMaterialExpressionGeneric >();
-	
+
 	if ( VRayMaterialProperties.bReflectionFresnel )
 	{
 		ReflectionIntensityExpression->SetExpressionName( TEXT("Desaturation") );
@@ -559,23 +559,6 @@ void FDatasmithMaxVRayMaterialsToUEPbr::Convert( TSharedRef< IDatasmithScene > D
 
 		ConvertState.bCanBake = true;
 	}
-	
-	// Displacement
-	{
-		ConvertState.DefaultTextureMode = EDatasmithTextureMode::Displace;
-
-		IDatasmithMaterialExpression* DisplacementExpression = FDatasmithMaxTexmapToUEPbrUtils::MapOrValue( this, VRayMaterialProperties.DisplacementMap, TEXT("Displacement Map"), TOptional< FLinearColor >(), TOptional< float >() );
-
-		if ( DisplacementExpression )
-		{
-			DisplacementExpression->ConnectExpression( PbrMaterialElement->GetWorldDisplacement() );
-		}
-
-		if ( DisplacementExpression )
-		{
-			DisplacementExpression->SetName( TEXT("Displacement Map") );
-		}
-	}
 
 	ConvertState.DefaultTextureMode = EDatasmithTextureMode::Specular; // At this point, all maps are considered specular maps
 
@@ -650,7 +633,7 @@ void FDatasmithMaxVRayMaterialsToUEPbr::Convert( TSharedRef< IDatasmithScene > D
 		DiffuseLerpExpression->ConnectExpression( PbrMaterialElement->GetBaseColor() );
 
 		IDatasmithMaterialExpression* ReflectionIOR = nullptr;
-		
+
 		{
 			TGuardValue< bool > SetIsMonoChannel( ConvertState.bIsMonoChannel, true );
 
@@ -757,7 +740,7 @@ void FDatasmithMaxVRayMaterialsToUEPbr::Convert( TSharedRef< IDatasmithScene > D
 	}
 	else
 	{
-		
+
 		MetallicExpression = ReflectionIntensityExpression;
 	}
 
@@ -765,7 +748,7 @@ void FDatasmithMaxVRayMaterialsToUEPbr::Convert( TSharedRef< IDatasmithScene > D
 	{
 		MetallicExpression->ConnectExpression( PbrMaterialElement->GetMetallic() );
 	}
-	
+
 	// UE Specular
 	if ( MetallicExpression )
 	{
@@ -899,7 +882,7 @@ void FDatasmithMaxVRay2SidedMaterialsToUEPbr::Convert( TSharedRef<IDatasmithScen
 			if ( MaterialElement )
 			{
 				MaterialElement->SetName( Material->GetName().data() ); // Name it with the main material not the front material
-				
+
 				if ( MaterialElement->IsA( EDatasmithElementType::UEPbrMaterial ) )
 				{
 					TSharedRef< IDatasmithUEPbrMaterialElement > UEPbrMaterialElement = StaticCastSharedRef< IDatasmithUEPbrMaterialElement >( MaterialElement.ToSharedRef() );
@@ -1028,9 +1011,9 @@ void FDatasmithMaxVRayBlendMaterialToUEPbr::Convert( TSharedRef<IDatasmithScene>
 
 			IDatasmithMaterialExpression* AlphaExpression = FDatasmithMaxTexmapToUEPbrUtils::MapOrValue(this, CoatedMaterial.MaterialBlendParameter, TEXT("MixAmount"),
 				CoatedMaterial.MixColor, TOptional< float >());
-				
+
 			//AlphaExpression is nullptr only when there is no mask and the mask weight is ~100% so we add scalar 0 instead.
-			if(!AlphaExpression) 
+			if(!AlphaExpression)
 			{
 				IDatasmithMaterialExpressionScalar* WeightExpression = PbrMaterialElement->AddMaterialExpression< IDatasmithMaterialExpressionScalar >();
 				WeightExpression->SetName(TEXT("MixAmount"));
diff --git a/Engine/Source/Runtime/Datasmith/DatasmithCore/Private/DatasmithMaterialElementsImpl.h b/Engine/Source/Runtime/Datasmith/DatasmithCore/Private/DatasmithMaterialElementsImpl.h
--- a/Engine/Source/Runtime/Datasmith/DatasmithCore/Private/DatasmithMaterialElementsImpl.h
+++ b/Engine/Source/Runtime/Datasmith/DatasmithCore/Private/DatasmithMaterialElementsImpl.h
@@ -672,7 +672,6 @@ public:
 	virtual IDatasmithExpressionInput& GetEmissiveColor() override { return *EmissiveColor.Edit(); }
 	virtual IDatasmithExpressionInput& GetOpacity() override { return *Opacity.Edit(); }
 	virtual IDatasmithExpressionInput& GetNormal() override { return *Normal.Edit(); }
-	virtual IDatasmithExpressionInput& GetWorldDisplacement() override { return *WorldDisplacement.Edit(); }
 	virtual IDatasmithExpressionInput& GetRefraction() override { return *Refraction.Edit(); }
 	virtual IDatasmithExpressionInput& GetAmbientOcclusion() override { return *AmbientOcclusion.Edit(); }
 	virtual IDatasmithExpressionInput& GetMaterialAttributes() override { return *MaterialAttributes.Edit(); }
diff --git a/Engine/Source/Runtime/Datasmith/DatasmithCore/Private/DatasmithSceneElementsImpl.cpp b/Engine/Source/Runtime/Datasmith/DatasmithCore/Private/DatasmithSceneElementsImpl.cpp
--- a/Engine/Source/Runtime/Datasmith/DatasmithCore/Private/DatasmithSceneElementsImpl.cpp
+++ b/Engine/Source/Runtime/Datasmith/DatasmithCore/Private/DatasmithSceneElementsImpl.cpp
@@ -658,7 +658,6 @@ FDatasmithShaderElementImpl::FDatasmithShaderElementImpl(const TCHAR* InName)
 	GetBumpComp()->SetBaseNames(DATASMITH_BUMPTEXNAME, TEXT("unsupported"), DATASMITH_BUMPVALUENAME, DATASMITH_BUMPCOMPNAME);
 	GetTransComp()->SetBaseNames(DATASMITH_TRANSPTEXNAME, DATASMITH_TRANSPCOLNAME, TEXT("unsupported"), DATASMITH_TRANSPCOMPNAME);
 	GetMaskComp()->SetBaseNames(DATASMITH_CLIPTEXNAME, TEXT("unsupported"), TEXT("unsupported"), DATASMITH_CLIPCOMPNAME);
-	GetDisplaceComp()->SetBaseNames(DATASMITH_DISPLACETEXNAME, TEXT("unsupported"), TEXT("unsupported"), DATASMITH_DISPLACECOMPNAME);
 	GetMetalComp()->SetBaseNames(DATASMITH_METALTEXNAME, TEXT("unsupported"), DATASMITH_METALVALUENAME, DATASMITH_METALCOMPNAME);
 	GetEmitComp()->SetBaseNames(DATASMITH_EMITTEXNAME, DATASMITH_EMITCOLNAME, TEXT("unsupported"), DATASMITH_EMITCOMPNAME);
 	GetWeightComp()->SetBaseNames(DATASMITH_WEIGHTTEXNAME, DATASMITH_WEIGHTCOLNAME, DATASMITH_WEIGHTVALUENAME, DATASMITH_WEIGHTCOMPNAME);
diff --git a/Engine/Source/Runtime/Datasmith/DatasmithCore/Private/DatasmithSceneElementsImpl.h b/Engine/Source/Runtime/Datasmith/DatasmithCore/Private/DatasmithSceneElementsImpl.h
--- a/Engine/Source/Runtime/Datasmith/DatasmithCore/Private/DatasmithSceneElementsImpl.h
+++ b/Engine/Source/Runtime/Datasmith/DatasmithCore/Private/DatasmithSceneElementsImpl.h
@@ -109,13 +109,13 @@ public:
 
 	virtual FVector GetTranslation() const override { return Translation.Get(Store); }
 	virtual void SetTranslation(float InX, float InY, float InZ, bool bKeepChildrenRelative) override { SetTranslation( FVector( InX, InY, InZ ), bKeepChildrenRelative ); }
-	virtual void SetTranslation(const FVector& Value, bool bKeepChildrenRelative) override 
-	{ 
+	virtual void SetTranslation(const FVector& Value, bool bKeepChildrenRelative) override
+	{
 		if (bKeepChildrenRelative)
 		{
 			ConvertChildsToRelative();
 		}
-		SetInternalTranslation(Value); 
+		SetInternalTranslation(Value);
 		if (bKeepChildrenRelative)
 		{
 			ConvertChildsToWorld();
@@ -130,7 +130,7 @@ public:
 		{
 			ConvertChildsToRelative();
 		}
-		SetInternalScale(Value); 
+		SetInternalScale(Value);
 		if (bKeepChildrenRelative)
 		{
 			ConvertChildsToWorld();
@@ -139,13 +139,13 @@ public:
 
 	virtual FQuat GetRotation() const override { return Rotation; }
 	virtual void SetRotation(float InX, float InY, float InZ, float InW, bool bKeepChildrenRelative) override { SetRotation( FQuat( InX, InY, InZ, InW ), bKeepChildrenRelative ); }
-	virtual void SetRotation(const FQuat& Value, bool bKeepChildrenRelative) override 
-	{ 
+	virtual void SetRotation(const FQuat& Value, bool bKeepChildrenRelative) override
+	{
 		if (bKeepChildrenRelative)
 		{
 			ConvertChildsToRelative();
 		}
-		SetInternalRotation(Value); 
+		SetInternalRotation(Value);
 		if (bKeepChildrenRelative)
 		{
 			ConvertChildsToWorld();
@@ -1177,22 +1177,6 @@ public:
 	virtual const TSharedPtr<IDatasmithCompositeTexture>& GetMaskComp() const override { return MaskComp; }
 	virtual void SetMaskComp(const TSharedPtr<IDatasmithCompositeTexture>& InValue) override { MaskComp = InValue; }
 
-	virtual const TCHAR* GetDisplaceTexture() const override { return *DisplaceTexture; }
-	virtual void SetDisplaceTexture(const TCHAR* InValue) override { DisplaceTexture = InValue; }
-
-	virtual FDatasmithTextureSampler GetDisplaceTextureSampler() const override { return DisplaceSampler; }
-	virtual void SetDisplaceTextureSampler(FDatasmithTextureSampler InValue) override { DisplaceSampler = InValue; }
-
-	virtual double GetDisplace() const override { return Displace; }
-	virtual void SetDisplace(double InValue) override { Displace = InValue; }
-
-	virtual double GetDisplaceSubDivision() const override { return DisplaceSubDivision; }
-	virtual void SetDisplaceSubDivision(double InValue) override { DisplaceSubDivision = InValue; }
-
-	virtual TSharedPtr<IDatasmithCompositeTexture>& GetDisplaceComp() override { return DisplaceComp; }
-	virtual const TSharedPtr<IDatasmithCompositeTexture>& GetDisplaceComp() const override { return DisplaceComp; }
-	virtual void SetDisplaceComp(const TSharedPtr<IDatasmithCompositeTexture>& InValue) override { DisplaceComp = InValue; }
-
 	virtual double GetMetal() const override { return Metal; }
 	virtual void SetMetal(double InValue) override { Metal = InValue; }
 
diff --git a/Engine/Source/Runtime/Datasmith/DatasmithCore/Private/DatasmithSceneXmlReader.cpp b/Engine/Source/Runtime/Datasmith/DatasmithCore/Private/DatasmithSceneXmlReader.cpp
--- a/Engine/Source/Runtime/Datasmith/DatasmithCore/Private/DatasmithSceneXmlReader.cpp
+++ b/Engine/Source/Runtime/Datasmith/DatasmithCore/Private/DatasmithSceneXmlReader.cpp
@@ -1413,203 +1413,185 @@ void FDatasmithSceneXmlReader::ParseMaterial(FXmlNode* InNode, TSharedPtr< IData
 				FString Texture;
 				FDatasmithTextureSampler TextureSampler;
 				FLinearColor Color;
-
-				if (ShaderNodes[j]->GetTag() == DATASMITH_DIFFUSETEXNAME)
+				const FString& NodeTag = ShaderNodes[j]->GetTag();
+				if (NodeTag == DATASMITH_DIFFUSETEXNAME)
 				{
 					ParseTexture(ShaderNodes[j], Texture, TextureSampler);
 					ShaderElement->SetDiffuseTexture(*Texture);
 					ShaderElement->SetDiffTextureSampler(TextureSampler);
 				}
-				else if (ShaderNodes[j]->GetTag() == DATASMITH_DIFFUSECOLNAME)
+				else if (NodeTag == DATASMITH_DIFFUSECOLNAME)
 				{
 					ParseColor(ShaderNodes[j], Color);
 					ShaderElement->SetDiffuseColor(Color);
 				}
-				else if (ShaderNodes[j]->GetTag() == DATASMITH_DIFFUSECOMPNAME)
+				else if (NodeTag == DATASMITH_DIFFUSECOMPNAME)
 				{
 					ParseComp(ShaderNodes[j], ShaderElement->GetDiffuseComp());
 				}
-				else if (ShaderNodes[j]->GetTag() == DATASMITH_REFLETEXNAME)
+				else if (NodeTag == DATASMITH_REFLETEXNAME)
 				{
 					ParseTexture(ShaderNodes[j], Texture, TextureSampler);
 					ShaderElement->SetReflectanceTexture(*Texture);
 					ShaderElement->SetRefleTextureSampler(TextureSampler);
 				}
-				else if (ShaderNodes[j]->GetTag() == DATASMITH_REFLECOLNAME)
+				else if (NodeTag == DATASMITH_REFLECOLNAME)
 				{
 					ParseColor(ShaderNodes[j], Color);
 					ShaderElement->SetReflectanceColor(Color);
 				}
-				else if (ShaderNodes[j]->GetTag() == DATASMITH_REFLECOMPNAME)
+				else if (NodeTag == DATASMITH_REFLECOMPNAME)
 				{
 					ParseComp(ShaderNodes[j], ShaderElement->GetRefleComp());
 				}
-				else if (ShaderNodes[j]->GetTag() == DATASMITH_ROUGHNESSTEXNAME)
+				else if (NodeTag == DATASMITH_ROUGHNESSTEXNAME)
 				{
 					ParseTexture(ShaderNodes[j], Texture, TextureSampler);
 					ShaderElement->SetRoughnessTexture(*Texture);
 					ShaderElement->SetRoughTextureSampler(TextureSampler);
 				}
-				else if (ShaderNodes[j]->GetTag() == DATASMITH_ROUGHNESSVALUENAME)
+				else if (NodeTag == DATASMITH_ROUGHNESSVALUENAME)
 				{
 					ShaderElement->SetRoughness(fmax(0.02, ValueFromString<double>(ShaderNodes[j]->GetAttribute(TEXT("value")))));
 				}
-				else if (ShaderNodes[j]->GetTag() == DATASMITH_ROUGHNESSCOMPNAME)
+				else if (NodeTag == DATASMITH_ROUGHNESSCOMPNAME)
 				{
 					ParseComp(ShaderNodes[j], ShaderElement->GetRoughnessComp());
 				}
-				else if (ShaderNodes[j]->GetTag() == DATASMITH_BUMPVALUENAME)
+				else if (NodeTag == DATASMITH_BUMPVALUENAME)
 				{
 					ShaderElement->SetBumpAmount(ValueFromString<double>(ShaderNodes[j]->GetAttribute(TEXT("value"))));
 				}
-				else if (ShaderNodes[j]->GetTag() == DATASMITH_BUMPTEXNAME)
+				else if (NodeTag == DATASMITH_BUMPTEXNAME)
 				{
 					ParseTexture(ShaderNodes[j], Texture, TextureSampler);
 					ShaderElement->SetBumpTexture(*Texture);
 					ShaderElement->SetBumpTextureSampler(TextureSampler);
 				}
-				else if (ShaderNodes[j]->GetTag() == DATASMITH_NORMALTEXNAME)
+				else if (NodeTag == DATASMITH_NORMALTEXNAME)
 				{
 					ParseTexture(ShaderNodes[j], Texture, TextureSampler);
 					ShaderElement->SetNormalTexture(*Texture);
 					ShaderElement->SetNormalTextureSampler(TextureSampler);
 				}
-				else if (ShaderNodes[j]->GetTag() == DATASMITH_NORMALCOMPNAME)
+				else if (NodeTag == DATASMITH_NORMALCOMPNAME)
 				{
 					ParseComp(ShaderNodes[j], ShaderElement->GetNormalComp());
 				}
-				else if (ShaderNodes[j]->GetTag() == DATASMITH_TRANSPTEXNAME)
+				else if (NodeTag == DATASMITH_TRANSPTEXNAME)
 				{
 					ParseTexture(ShaderNodes[j], Texture, TextureSampler);
 					ShaderElement->SetTransparencyTexture(*Texture);
 					ShaderElement->SetTransTextureSampler(TextureSampler);
 				}
-				else if (ShaderNodes[j]->GetTag() == DATASMITH_TRANSPCOMPNAME)
+				else if (NodeTag == DATASMITH_TRANSPCOMPNAME)
 				{
 					ParseComp(ShaderNodes[j], ShaderElement->GetTransComp());
 				}
-				else if (ShaderNodes[j]->GetTag() == DATASMITH_CLIPTEXNAME)
+				else if (NodeTag == DATASMITH_CLIPTEXNAME)
 				{
 					ParseTexture(ShaderNodes[j], Texture, TextureSampler);
 					ShaderElement->SetMaskTexture(*Texture);
 					ShaderElement->SetMaskTextureSampler(TextureSampler);
 				}
-				else if (ShaderNodes[j]->GetTag() == DATASMITH_CLIPCOMPNAME)
+				else if (NodeTag == DATASMITH_CLIPCOMPNAME)
 				{
 					ParseComp(ShaderNodes[j], ShaderElement->GetMaskComp());
 				}
-				else if (ShaderNodes[j]->GetTag() == DATASMITH_TRANSPCOLNAME)
+				else if (NodeTag == DATASMITH_TRANSPCOLNAME)
 				{
 					ParseColor(ShaderNodes[j], Color);
 					ShaderElement->SetTransparencyColor(Color);
 				}
-				else if (ShaderNodes[j]->GetTag() == DATASMITH_IORVALUENAME)
+				else if (NodeTag == DATASMITH_IORVALUENAME)
 				{
 					ShaderElement->SetIOR(ValueFromString<double>(ShaderNodes[j]->GetAttribute(TEXT("value"))));
 				}
-				else if (ShaderNodes[j]->GetTag() == DATASMITH_IORKVALUENAME)
+				else if (NodeTag == DATASMITH_IORKVALUENAME)
 				{
 					ShaderElement->SetIORk(ValueFromString<double>(ShaderNodes[j]->GetAttribute(TEXT("value"))));
 				}
-				else if (ShaderNodes[j]->GetTag() == DATASMITH_REFRAIORVALUENAME)
+				else if (NodeTag == DATASMITH_REFRAIORVALUENAME)
 				{
 					ShaderElement->SetIORRefra(ValueFromString<double>(ShaderNodes[j]->GetAttribute(TEXT("value"))));
 				}
-				else if (ShaderNodes[j]->GetTag() == DATASMITH_TWOSIDEDVALUENAME)
+				else if (NodeTag == DATASMITH_TWOSIDEDVALUENAME)
 				{
 					ShaderElement->SetTwoSided(ValueFromString<bool>(ShaderNodes[j]->GetAttribute(TEXT("enabled"))));
 				}
-				else if (ShaderNodes[j]->GetTag() == DATASMITH_DISPLACETEXNAME)
-				{
-					ParseTexture(ShaderNodes[j], Texture, TextureSampler);
-					ShaderElement->SetDisplaceTexture(*Texture);
-					ShaderElement->SetDisplaceTextureSampler(TextureSampler);
-				}
-				else if (ShaderNodes[j]->GetTag() == DATASMITH_DISPLACEVALNAME)
-				{
-					ShaderElement->SetDisplace(ValueFromString<double>(ShaderNodes[j]->GetAttribute(TEXT("value"))));
-				}
-				else if (ShaderNodes[j]->GetTag() == DATASMITH_DISPLACECOMPNAME)
-				{
-					ParseComp(ShaderNodes[j], ShaderElement->GetDisplaceComp());
-				}
-				else if (ShaderNodes[j]->GetTag() == DATASMITH_DISPLACESUBNAME)
-				{
-					ShaderElement->SetDisplaceSubDivision(ValueFromString<double>(ShaderNodes[j]->GetAttribute(TEXT("value"))));
-				}
-				else if (ShaderNodes[j]->GetTag() == DATASMITH_METALTEXNAME)
+				else if (NodeTag == DATASMITH_METALTEXNAME)
 				{
 					ParseTexture(ShaderNodes[j], Texture, TextureSampler);
 					ShaderElement->SetMetalTexture(*Texture);
 					ShaderElement->SetMetalTextureSampler(TextureSampler);
 				}
-				else if (ShaderNodes[j]->GetTag() == DATASMITH_METALVALUENAME)
+				else if (NodeTag == DATASMITH_METALVALUENAME)
 				{
 					ShaderElement->SetMetal(ValueFromString<double>(ShaderNodes[j]->GetAttribute(TEXT("value"))));
 				}
-				else if (ShaderNodes[j]->GetTag() == DATASMITH_METALCOMPNAME)
+				else if (NodeTag == DATASMITH_METALCOMPNAME)
 				{
 					ParseComp(ShaderNodes[j], ShaderElement->GetMetalComp());
 				}
-				else if (ShaderNodes[j]->GetTag() == DATASMITH_EMITTEXNAME)
+				else if (NodeTag == DATASMITH_EMITTEXNAME)
 				{
 					ParseTexture(ShaderNodes[j], Texture, TextureSampler);
 					ShaderElement->SetEmitTexture(*Texture);
 					ShaderElement->SetEmitTextureSampler(TextureSampler);
 				}
-				else if (ShaderNodes[j]->GetTag() == DATASMITH_EMITVALUENAME)
+				else if (NodeTag == DATASMITH_EMITVALUENAME)
 				{
 					ShaderElement->SetEmitPower(ValueFromString<double>(ShaderNodes[j]->GetAttribute(TEXT("value"))));
 				}
-				else if (ShaderNodes[j]->GetTag() == DATASMITH_EMITCOMPNAME)
+				else if (NodeTag == DATASMITH_EMITCOMPNAME)
 				{
 					ParseComp(ShaderNodes[j], ShaderElement->GetEmitComp());
 				}
-				else if (ShaderNodes[j]->GetTag() == DATASMITH_EMITTEMPNAME)
+				else if (NodeTag == DATASMITH_EMITTEMPNAME)
 				{
 					ShaderElement->SetEmitTemperature(ValueFromString<double>(ShaderNodes[j]->GetAttribute(TEXT("value"))));
 				}
-				else if (ShaderNodes[j]->GetTag() == DATASMITH_EMITCOLNAME)
+				else if (NodeTag == DATASMITH_EMITCOLNAME)
 				{
 					ParseColor(ShaderNodes[j], Color);
 					ShaderElement->SetEmitColor(Color);
 				}
-				else if (ShaderNodes[j]->GetTag() == DATASMITH_EMITONLYVALUENAME)
+				else if (NodeTag == DATASMITH_EMITONLYVALUENAME)
 				{
 					ShaderElement->SetLightOnly(ValueFromString<bool>(ShaderNodes[j]->GetAttribute(TEXT("enabled"))));
 				}
-				else if (ShaderNodes[j]->GetTag() == DATASMITH_WEIGHTTEXNAME)
+				else if (NodeTag == DATASMITH_WEIGHTTEXNAME)
 				{
 					ParseTexture(ShaderNodes[j], Texture, TextureSampler);
 					ShaderElement->SetWeightTexture(*Texture);
 					ShaderElement->SetWeightTextureSampler(TextureSampler);
 				}
-				else if (ShaderNodes[j]->GetTag() == DATASMITH_WEIGHTCOLNAME)
+				else if (NodeTag == DATASMITH_WEIGHTCOLNAME)
 				{
 					ParseColor(ShaderNodes[j], Color);
 					ShaderElement->SetWeightColor(Color);
 				}
-				else if (ShaderNodes[j]->GetTag() == DATASMITH_WEIGHTCOMPNAME)
+				else if (NodeTag == DATASMITH_WEIGHTCOMPNAME)
 				{
 					ParseComp(ShaderNodes[j], ShaderElement->GetWeightComp());
 				}
-				else if (ShaderNodes[j]->GetTag() == DATASMITH_WEIGHTVALUENAME)
+				else if (NodeTag == DATASMITH_WEIGHTVALUENAME)
 				{
 					ShaderElement->SetWeightValue(ValueFromString<double>(ShaderNodes[j]->GetAttribute(TEXT("value"))));
 				}
-				else if (ShaderNodes[j]->GetTag() == DATASMITH_BLENDMODE)
+				else if (NodeTag == DATASMITH_BLENDMODE)
 				{
 					ShaderElement->SetBlendMode((EDatasmithBlendMode)ValueFromString<int32>(ShaderNodes[j]->GetAttribute(TEXT("value"))));
 				}
-				else if (ShaderNodes[j]->GetTag() == DATASMITH_STACKLAYER)
+				else if (NodeTag == DATASMITH_STACKLAYER)
 				{
 					ShaderElement->SetIsStackedLayer(ValueFromString<bool>(ShaderNodes[j]->GetAttribute(TEXT("enabled"))));
 				}
-				else if (ShaderNodes[j]->GetTag() == DATASMITH_DYNAMICEMISSIVE)
+				else if (NodeTag == DATASMITH_DYNAMICEMISSIVE)
 				{
 					ShaderElement->SetUseEmissiveForDynamicAreaLighting(ValueFromString<bool>(ShaderNodes[j]->GetAttribute(TEXT("enabled"))));
 				}
-				else if (ShaderNodes[j]->GetTag() == DATASMITH_SHADERUSAGE)
+				else if (NodeTag == DATASMITH_SHADERUSAGE)
 				{
 					ShaderElement->SetShaderUsage((EDatasmithShaderUsage)ValueFromString<int32>(ShaderNodes[j]->GetAttribute(TEXT("value"))));
 				}
@@ -1930,7 +1912,6 @@ void FDatasmithSceneXmlReader::ParseUEPbrMaterial(FXmlNode* InNode, TSharedPtr<
 	TryConnectMaterialInput(OutElement->GetEmissiveColor());
 	TryConnectMaterialInput(OutElement->GetOpacity());
 	TryConnectMaterialInput(OutElement->GetNormal());
-	TryConnectMaterialInput(OutElement->GetWorldDisplacement());
 	TryConnectMaterialInput(OutElement->GetRefraction());
 	TryConnectMaterialInput(OutElement->GetAmbientOcclusion());
 	TryConnectMaterialInput(OutElement->GetMaterialAttributes());
diff --git a/Engine/Source/Runtime/Datasmith/DatasmithCore/Private/DatasmithSceneXmlWriter.cpp b/Engine/Source/Runtime/Datasmith/DatasmithCore/Private/DatasmithSceneXmlWriter.cpp
--- a/Engine/Source/Runtime/Datasmith/DatasmithCore/Private/DatasmithSceneXmlWriter.cpp
+++ b/Engine/Source/Runtime/Datasmith/DatasmithCore/Private/DatasmithSceneXmlWriter.cpp
@@ -1191,16 +1191,6 @@ void FDatasmithSceneXmlWriterImpl::WriteShaderElement(const TSharedPtr< IDatasmi
 		WriteBool(Archive, Indent + 1, DATASMITH_TWOSIDEDVALUENAME, ShaderElement->GetTwoSided());
 	}
 
-	if (ShaderElement->GetDisplace() > 0)
-	{
-		WriteValue(Archive, Indent + 1, DATASMITH_DISPLACEVALNAME, (float)ShaderElement->GetDisplace());
-	}
-
-	if (ShaderElement->GetDisplaceSubDivision() > 0)
-	{
-		WriteValue(Archive, Indent + 1, DATASMITH_DISPLACESUBNAME, (float)ShaderElement->GetDisplaceSubDivision());
-	}
-
 	if (ShaderElement->GetMetal() > 0)
 	{
 		WriteValue(Archive, Indent + 1, DATASMITH_METALVALUENAME, (float)ShaderElement->GetMetal());
@@ -1223,7 +1213,6 @@ void FDatasmithSceneXmlWriterImpl::WriteShaderElement(const TSharedPtr< IDatasmi
 	WriteCompTex(ShaderElement->GetBumpComp(), Archive, Indent + 1);
 	WriteCompTex(ShaderElement->GetTransComp(), Archive, Indent + 1);
 	WriteCompTex(ShaderElement->GetMaskComp(), Archive, Indent + 1);
-	WriteCompTex(ShaderElement->GetDisplaceComp(), Archive, Indent + 1);
 	WriteCompTex(ShaderElement->GetMetalComp(), Archive, Indent + 1);
 	WriteCompTex(ShaderElement->GetEmitComp(), Archive, Indent + 1);
 	WriteCompTex(ShaderElement->GetWeightComp(), Archive, Indent + 1);
@@ -1428,7 +1417,6 @@ void FDatasmithSceneXmlWriterImpl::WriteUEPbrMaterialElement( const TSharedRef<
 	WriteUEPbrMaterialExpressionInput( MaterialElement, MaterialElement->GetEmissiveColor(), Archive, Indent + 1 );
 	WriteUEPbrMaterialExpressionInput( MaterialElement, MaterialElement->GetOpacity(), Archive, Indent + 1 );
 	WriteUEPbrMaterialExpressionInput( MaterialElement, MaterialElement->GetNormal(), Archive, Indent + 1 );
-	WriteUEPbrMaterialExpressionInput( MaterialElement, MaterialElement->GetWorldDisplacement(), Archive, Indent + 1 );
 	WriteUEPbrMaterialExpressionInput( MaterialElement, MaterialElement->GetRefraction(), Archive, Indent + 1 );
 	WriteUEPbrMaterialExpressionInput( MaterialElement, MaterialElement->GetAmbientOcclusion(), Archive, Indent + 1 );
 	WriteUEPbrMaterialExpressionInput( MaterialElement, MaterialElement->GetMaterialAttributes(), Archive, Indent + 1 );
diff --git a/Engine/Source/Runtime/Datasmith/DatasmithCore/Private/DatasmithUtils.cpp b/Engine/Source/Runtime/Datasmith/DatasmithCore/Private/DatasmithUtils.cpp
--- a/Engine/Source/Runtime/Datasmith/DatasmithCore/Private/DatasmithUtils.cpp
+++ b/Engine/Source/Runtime/Datasmith/DatasmithCore/Private/DatasmithUtils.cpp
@@ -1036,7 +1036,6 @@ namespace DatasmithSceneUtilsImpl
 			ParseExpressionElement(MaterialElement->GetRefraction().GetExpression());
 			ParseExpressionElement(MaterialElement->GetAmbientOcclusion().GetExpression());
 			ParseExpressionElement(MaterialElement->GetOpacity().GetExpression());
-			ParseExpressionElement(MaterialElement->GetWorldDisplacement().GetExpression());
 
 			if ( MaterialElement->GetUseMaterialAttributes() )
 			{
@@ -1091,7 +1090,6 @@ namespace DatasmithSceneUtilsImpl
 					ScanCompositeTexture( Shader->GetBumpComp().Get() );
 					ScanCompositeTexture( Shader->GetTransComp().Get() );
 					ScanCompositeTexture( Shader->GetMaskComp().Get() );
-					ScanCompositeTexture( Shader->GetDisplaceComp().Get() );
 					ScanCompositeTexture( Shader->GetMetalComp().Get() );
 					ScanCompositeTexture( Shader->GetEmitComp().Get() );
 					ScanCompositeTexture( Shader->GetWeightComp().Get() );
diff --git a/Engine/Source/Runtime/Datasmith/DatasmithCore/Public/DatasmithDefinitions.h b/Engine/Source/Runtime/Datasmith/DatasmithCore/Public/DatasmithDefinitions.h
--- a/Engine/Source/Runtime/Datasmith/DatasmithCore/Public/DatasmithDefinitions.h
+++ b/Engine/Source/Runtime/Datasmith/DatasmithCore/Public/DatasmithDefinitions.h
@@ -161,7 +161,7 @@ enum class EDatasmithTextureMode : uint8
 	Specular,
 	Normal,
 	NormalGreenInv,
-	Displace,
+	UNUSED_Displace,
 	Other,
 	Bump,
 	Ies
@@ -598,11 +598,6 @@ static const TCHAR* KeyValuePropertyTypeStrings[] = { TEXT("String"), TEXT("Colo
 #define DATASMITH_IORKVALUENAME					TEXT("IORk")
 #define DATASMITH_REFRAIORVALUENAME				TEXT("IORRefraction")
 
-#define DATASMITH_DISPLACETEXNAME				TEXT("Displace")
-#define DATASMITH_DISPLACEVALNAME				TEXT("Displaceval")
-#define DATASMITH_DISPLACECOMPNAME				TEXT("Displacecomp")
-#define DATASMITH_DISPLACESUBNAME				TEXT("DisplaceSubdivisionval")
-
 #define DATASMITH_METALTEXNAME					TEXT("Metal")
 #define DATASMITH_METALVALUENAME				TEXT("Metalval")
 #define DATASMITH_METALCOMPNAME					TEXT("Metalcomp")
diff --git a/Engine/Source/Runtime/Datasmith/DatasmithCore/Public/DatasmithMaterialElements.h b/Engine/Source/Runtime/Datasmith/DatasmithCore/Public/DatasmithMaterialElements.h
--- a/Engine/Source/Runtime/Datasmith/DatasmithCore/Public/DatasmithMaterialElements.h
+++ b/Engine/Source/Runtime/Datasmith/DatasmithCore/Public/DatasmithMaterialElements.h
@@ -219,7 +219,6 @@ public:
 	virtual IDatasmithExpressionInput& GetEmissiveColor() = 0;
 	virtual IDatasmithExpressionInput& GetOpacity() = 0;
 	virtual IDatasmithExpressionInput& GetNormal() = 0;
-	virtual IDatasmithExpressionInput& GetWorldDisplacement() = 0;
 	virtual IDatasmithExpressionInput& GetRefraction() = 0;
 	virtual IDatasmithExpressionInput& GetAmbientOcclusion() = 0;
 	virtual IDatasmithExpressionInput& GetMaterialAttributes() = 0;
diff --git a/Engine/Source/Runtime/Datasmith/DatasmithCore/Public/IDatasmithSceneElements.h b/Engine/Source/Runtime/Datasmith/DatasmithCore/Public/IDatasmithSceneElements.h
--- a/Engine/Source/Runtime/Datasmith/DatasmithCore/Public/IDatasmithSceneElements.h
+++ b/Engine/Source/Runtime/Datasmith/DatasmithCore/Public/IDatasmithSceneElements.h
@@ -841,7 +841,7 @@ class DATASMITHCORE_API IDatasmithShaderElement : public IDatasmithElement
 public:
 	/**
 	 * Realistic fresnel creates a pretty more complex node tree based on the actual fresnel equation.
-	 * If this param is not enabled an aproximation will be used.
+	 * If this param is not enabled an approximation will be used.
 	 *
 	 * It has no effect if bDisableReflectionFresnel is set to true.
 	*/
@@ -1053,103 +1053,70 @@ public:
 	/** Set the opacity mask compound map */
 	virtual void SetMaskComp(const TSharedPtr<IDatasmithCompositeTexture>& Value) = 0;
 
-	/** Get the displacement filename */
-	virtual const TCHAR* GetDisplaceTexture() const = 0;
-
-	/** Set the displacement filename */
-	virtual void SetDisplaceTexture(const TCHAR* Value) = 0;
-
-	/** Get the displacement UV coordinates */
-	virtual FDatasmithTextureSampler GetDisplaceTextureSampler() const = 0;
-
-	/** Set the displacement UV coordinates */
-	virtual void SetDisplaceTextureSampler(FDatasmithTextureSampler Value) = 0;
-
-	/** Get the displacement value in centimeters */
-	virtual double GetDisplace() const = 0;
-
-	/** Set the displacement value in centimeters */
-	virtual void SetDisplace(double Value) = 0;
-
-	/** Get the displacement subdivision multiplier */
-	virtual double GetDisplaceSubDivision() const = 0;
-
-	/** Set the displacement subdivision multiplier */
-	virtual void SetDisplaceSubDivision(double Value) = 0;
-
-	/** Get the displacement compound map */
-	virtual TSharedPtr<IDatasmithCompositeTexture>& GetDisplaceComp() = 0;
-
-	/** Get the displacement compound map */
-	virtual const TSharedPtr<IDatasmithCompositeTexture>& GetDisplaceComp() const = 0;
-
-	/** Set the displacement compound map */
-	virtual void SetDisplaceComp(const TSharedPtr<IDatasmithCompositeTexture>& Value) = 0;
-
 	/** Get the metalness value */
 	virtual double GetMetal() const = 0;
 
 	/** Set the metalness value */
 	virtual void SetMetal(double Value) = 0;
 
-	/** Get the diffuse filename */
+	/** Get the metalness filename */
 	virtual const TCHAR* GetMetalTexture() const = 0;
 
-	/** Set the diffuse filename */
+	/** Set the metalness filename */
 	virtual void SetMetalTexture(const TCHAR* Value) = 0;
 
-	/** Get the diffuse UV coordinates */
+	/** Get the metalness UV coordinates */
 	virtual FDatasmithTextureSampler GetMetalTextureSampler() const = 0;
 
-	/** Set the diffuse UV coordinates */
+	/** Set the metalness UV coordinates */
 	virtual void SetMetalTextureSampler(FDatasmithTextureSampler Value) = 0;
 
-	/** Get the diffuse compound map */
+	/** Get the metalness compound map */
 	virtual TSharedPtr<IDatasmithCompositeTexture>& GetMetalComp() = 0;
 
-	/** Get the diffuse compound map */
+	/** Get the metalness compound map */
 	virtual const TSharedPtr<IDatasmithCompositeTexture>& GetMetalComp() const = 0;
 
-	/** Set the diffuse compound map */
+	/** Set the metalness compound map */
 	virtual void SetMetalComp(const TSharedPtr<IDatasmithCompositeTexture>& Value) = 0;
 
-	/** Get the emmitance color in linear space */
+	/** Get the emittance color in linear space */
 	virtual FLinearColor GetEmitColor() const = 0;
 
-	/** Set the emmitance color in linear space */
+	/** Set the emittance color in linear space */
 	virtual void SetEmitColor(FLinearColor Value) = 0;
 
-	/** Get the emmitance filename */
+	/** Get the emittance filename */
 	virtual const TCHAR* GetEmitTexture() const = 0;
 
-	/** Set the emmitance filename */
+	/** Set the emittance filename */
 	virtual void SetEmitTexture(const TCHAR* Value) = 0;
 
-	/** Get the emmitance UV coordinates */
+	/** Get the emittance UV coordinates */
 	virtual FDatasmithTextureSampler GetEmitTextureSampler() const = 0;
 
-	/** Set the emmitance UV coordinates */
+	/** Set the emittance UV coordinates */
 	virtual void SetEmitTextureSampler(FDatasmithTextureSampler Value) = 0;
 
-	/** Get the emmitance temperature color */
+	/** Get the emittance temperature color */
 	virtual double GetEmitTemperature() const = 0;
 
-	/** Set the emmitance temperature color */
+	/** Set the emittance temperature color */
 	virtual void SetEmitTemperature(double Value) = 0;
 
-	/** Get the emmitance power in lumens */
+	/** Get the emittance power in lumens */
 	virtual double GetEmitPower() const = 0;
 
-	/** Set the emmitance power in lumens */
+	/** Set the emittance power in lumens */
 	virtual void SetEmitPower(double Value) = 0;
 
-	/** Get the emmitance compound map */
+	/** Get the emittance compound map */
 	virtual TSharedPtr<IDatasmithCompositeTexture>& GetEmitComp() = 0;
-	/** Get the emmitance compound map */
 
+	/** Get the emittance compound map */
 	virtual const TSharedPtr<IDatasmithCompositeTexture>& GetEmitComp() const = 0;
 
-	/** Set the emmitance compound map */
+	/** Set the emittance compound map */
 	virtual void SetEmitComp(const TSharedPtr<IDatasmithCompositeTexture>& Value) = 0;
 
 
