CADKernel fix build of untrimmed face

diff --git a/Engine/Source/Runtime/Datasmith/CADKernel/Private/CADKernel/Topo/TopologicalFace.cpp b/Engine/Source/Runtime/Datasmith/CADKernel/Private/CADKernel/Topo/TopologicalFace.cpp
--- a/Engine/Source/Runtime/Datasmith/CADKernel/Private/CADKernel/Topo/TopologicalFace.cpp
+++ b/Engine/Source/Runtime/Datasmith/CADKernel/Private/CADKernel/Topo/TopologicalFace.cpp
@@ -57,7 +57,7 @@ void FTopologicalFace::ApplyNaturalLoops()
 		TSharedRef<FCurve> Curve2D = FEntity::MakeShared<FSegmentCurve>(StartPoint, EndPoint, 2);
 		TSharedRef<FRestrictionCurve> Curve3D = FEntity::MakeShared<FRestrictionCurve>(CarrierSurface.ToSharedRef(), Curve2D);
 		TSharedPtr<FTopologicalEdge> Edge = FTopologicalEdge::Make(Curve3D);
-		if (Edge.IsValid())
+		if (!Edge.IsValid())
 		{
 			return;
 		}
diff --git a/Engine/Source/Runtime/Datasmith/CADKernel/Public/CADKernel/Topo/TopologicalEdge.h b/Engine/Source/Runtime/Datasmith/CADKernel/Public/CADKernel/Topo/TopologicalEdge.h
--- a/Engine/Source/Runtime/Datasmith/CADKernel/Public/CADKernel/Topo/TopologicalEdge.h
+++ b/Engine/Source/Runtime/Datasmith/CADKernel/Public/CADKernel/Topo/TopologicalEdge.h
@@ -181,7 +181,7 @@ namespace CADKernel
 		 */
 		bool CheckIfDegenerated();
 
-		void Link(const TSharedRef<FTopologicalEdge>& OtherEdge, double JoiningTolerance);
+		void Link(const TSharedRef<FTopologicalEdge>& OtherEdge, double SquareJoiningTolerance);
 
 		TSharedRef<const FTopologicalEdge> GetLinkActiveEdge() const
 		{
