Datasmith: basic pass on static analysis fixes

diff --git a/Engine/Plugins/Enterprise/DatasmithCADImporter/Source/DatasmithPLMXMLTranslator/Private/DatasmithPlmXmlImporter.cpp b/Engine/Plugins/Enterprise/DatasmithCADImporter/Source/DatasmithPLMXMLTranslator/Private/DatasmithPlmXmlImporter.cpp
--- a/Engine/Plugins/Enterprise/DatasmithCADImporter/Source/DatasmithPLMXMLTranslator/Private/DatasmithPlmXmlImporter.cpp
+++ b/Engine/Plugins/Enterprise/DatasmithCADImporter/Source/DatasmithPLMXMLTranslator/Private/DatasmithPlmXmlImporter.cpp
@@ -183,7 +183,7 @@ namespace PlmXml
 
 	struct FProductInstance
 	{
-		const FXmlNode* Node;
+		const FXmlNode* Node = nullptr;
 
 		operator bool() const
 		{
@@ -193,7 +193,7 @@ namespace PlmXml
 
 	struct FProductRevisionView
 	{
-		const FXmlNode* Node;
+		const FXmlNode* Node = nullptr;
 		
 		operator bool() const
 		{
@@ -203,7 +203,7 @@ namespace PlmXml
 
 	struct FRepresentation
 	{
-		const FXmlNode* Node;
+		const FXmlNode* Node = nullptr;
 
 		operator bool() const
 		{
diff --git a/Engine/Plugins/Enterprise/DatasmithFBXImporter/Source/DatasmithDeltaGenTranslator/Private/DatasmithDeltaGenAnimationInterpolator.cpp b/Engine/Plugins/Enterprise/DatasmithFBXImporter/Source/DatasmithDeltaGenTranslator/Private/DatasmithDeltaGenAnimationInterpolator.cpp
--- a/Engine/Plugins/Enterprise/DatasmithFBXImporter/Source/DatasmithDeltaGenTranslator/Private/DatasmithDeltaGenAnimationInterpolator.cpp
+++ b/Engine/Plugins/Enterprise/DatasmithFBXImporter/Source/DatasmithDeltaGenTranslator/Private/DatasmithDeltaGenAnimationInterpolator.cpp
@@ -339,7 +339,7 @@ namespace DeltaGen
 
 		// Solve cubic polynomial with Gerolano Cardano's formula
 		double Coefs[4] = { CoefConst, CoefT1, CoefT2, CoefT3 };
-		double Solutions[3];
+		double Solutions[3] = {-1., -1., -1.};
 		int32 NumSolutions = SolveCubic(Coefs, Solutions);
 
 		// The target solution, if it exists, is the only real one within [0, 1]
diff --git a/Engine/Source/Runtime/Datasmith/CADKernel/Public/CADKernel/Geo/Surfaces/Surface.h b/Engine/Source/Runtime/Datasmith/CADKernel/Public/CADKernel/Geo/Surfaces/Surface.h
--- a/Engine/Source/Runtime/Datasmith/CADKernel/Public/CADKernel/Geo/Surfaces/Surface.h
+++ b/Engine/Source/Runtime/Datasmith/CADKernel/Public/CADKernel/Geo/Surfaces/Surface.h
@@ -81,7 +81,7 @@ namespace CADKernel
 
 		virtual void Serialize(FCADKernelArchive& Ar) override
 		{
-			// Surface's type is serialize because it is used to instantiate the correct entity on deserialization (@see Deserialize(FCADKernelArchive& Archive)) 
+			// Surface's type is serialize because it is used to instantiate the correct entity on deserialization (@see Deserialize(FCADKernelArchive& Archive))
 			if (Ar.IsSaving())
 			{
 				ESurface SurfaceType = GetSurfaceType();
@@ -144,17 +144,17 @@ namespace CADKernel
 		}
 
 		/**
-		 * Return the minimum tolerance in the parametric space of the surface along the specified axis 
+		 * Return the minimum tolerance in the parametric space of the surface along the specified axis
 		 * With Tolerance3D = FSysteme.GeometricalTolerance
 		 * @see FBoundary::ComputeMinimalTolerance
 		 */
-		const double& GetIsoTolerance(EIso Iso) const
+		double GetIsoTolerance(EIso Iso) const
 		{
 			ensureCADKernel(MinToleranceIso.IsValid());
 			return ((FSurfacicTolerance) MinToleranceIso)[Iso];
 		}
 
-		const double Get3DTolerance()
+		double Get3DTolerance() const
 		{
 			return Tolerance3D;
 		}
diff --git a/Engine/Source/Runtime/Datasmith/CADKernel/Public/CADKernel/Topo/TopologicalFace.h b/Engine/Source/Runtime/Datasmith/CADKernel/Public/CADKernel/Topo/TopologicalFace.h
--- a/Engine/Source/Runtime/Datasmith/CADKernel/Public/CADKernel/Topo/TopologicalFace.h
+++ b/Engine/Source/Runtime/Datasmith/CADKernel/Public/CADKernel/Topo/TopologicalFace.h
@@ -32,7 +32,7 @@ namespace CADKernel
 	class FThinZoneFinder;
 	class FBezierSurface;
 	class FSegmentCurve;
-	
+
 	class CADKERNEL_API FTopologicalFace : public FTopologicalEntity, public FMetadataDictionary
 	{
 		friend class FEntity;
@@ -57,7 +57,7 @@ namespace CADKernel
 		 * Temporary discretization of the surface used to compute the mesh of the edge
 		 */
 		FCoordinateGrid CrossingCoordinates;
-		
+
 		/**
 		 * Min delta U at the crossing u coordinate to respect meshing criteria
 		 */
@@ -69,8 +69,8 @@ namespace CADKernel
 		FCoordinateGrid CrossingPointDeltaMaxs;
 
 		/**
-		 * Build a non-trimmed trimmed surface 
-		 * This constructor has to be completed with one of the three "AddBoundaries" methods to be finalized. 
+		 * Build a non-trimmed trimmed surface
+		 * This constructor has to be completed with one of the three "AddBoundaries" methods to be finalized.
 		 */
 		FTopologicalFace(const TSharedPtr<FSurface>& InCarrierSurface)
 			: FTopologicalEntity()
@@ -128,7 +128,7 @@ namespace CADKernel
 			return CarrierSurface->GetIsoTolerances();
 		}
 
-		const double& GetIsoTolerance(EIso Iso) const
+		double GetIsoTolerance(EIso Iso) const
 		{
 			return CarrierSurface->GetIsoTolerance(Iso);
 		}
@@ -147,7 +147,7 @@ namespace CADKernel
 			return 1;
 		}
 
-		virtual void GetFaces(TArray<TSharedPtr<FTopologicalFace>>& OutFaces) override		
+		virtual void GetFaces(TArray<TSharedPtr<FTopologicalFace>>& OutFaces) override
 		{
 			if (!HasMarker1())
 			{
@@ -196,11 +196,11 @@ namespace CADKernel
 
 		/**
 		 * Trimmed the face with its natural limit curves (Iso UMin,  ...). This function is called to trim untrimmed topological face.
-		 * This function should not be called if the topological face already has a loop. 
+		 * This function should not be called if the topological face already has a loop.
 		 */
 		void ApplyNaturalLoops();
 
-		int32 LoopCount() const 
+		int32 LoopCount() const
 		{
 			return Loops.Num();
 		}
@@ -211,8 +211,8 @@ namespace CADKernel
 		}
 
 		/**
-		 * Get a sampling of each loop of the face 
-		 * @param OutLoopSamplings an array of 2d points 
+		 * Get a sampling of each loop of the face
+		 * @param OutLoopSamplings an array of 2d points
 		 * @param OutAABBs an array of 2d axis aligned bounding box of each boundary
 		 */
 		const void Get2DLoopSampling(TArray<TArray<FPoint2D>>& OutLoopSamplings) const;
@@ -282,7 +282,7 @@ namespace CADKernel
 			ensureCADKernel(Mesh.IsValid());
 			return Mesh.ToSharedRef();
 		}
-		
+
 		void InitDeltaUs();
 
 		void ChooseFinalDeltaUs();
@@ -430,7 +430,7 @@ namespace CADKernel
 		}
 
 		void ComputeSurfaceSideProperties();
-		
+
 		/**
 		 * Defines if the surface is either EQuadType::QUAD, either EQuadType::TRIANGULAR or EQuadType::OTHER
 		 */
@@ -461,12 +461,12 @@ namespace CADKernel
 			return NumOfMeshedSide;
 		}
 
-		void AddMeshedLength(double Length) 
+		void AddMeshedLength(double Length)
 		{
 			LengthOfMeshedSide += Length;
 		}
 
-		double MeshedSideRatio() const 
+		double MeshedSideRatio() const
 		{
 			return LengthOfMeshedSide / LoopLength;
 		}
@@ -481,7 +481,7 @@ namespace CADKernel
 			return StartSideIndices;
 		}
 
-		int32 GetSideIndex(TSharedPtr<FTopologicalEdge> Edge) const 
+		int32 GetSideIndex(TSharedPtr<FTopologicalEdge> Edge) const
 		{
 			int32 EdgeIndex = Loops[0]->GetEdgeIndex(Edge);
 			if (EdgeIndex < 0)
diff --git a/Engine/Source/Runtime/Datasmith/DatasmithCore/Private/DatasmithSceneElementsImpl.cpp b/Engine/Source/Runtime/Datasmith/DatasmithCore/Private/DatasmithSceneElementsImpl.cpp
--- a/Engine/Source/Runtime/Datasmith/DatasmithCore/Private/DatasmithSceneElementsImpl.cpp
+++ b/Engine/Source/Runtime/Datasmith/DatasmithCore/Private/DatasmithSceneElementsImpl.cpp
@@ -651,16 +651,16 @@ FDatasmithShaderElementImpl::FDatasmithShaderElementImpl(const TCHAR* InName)
 	, ShaderUsage(EDatasmithShaderUsage::Surface)
 	, bUseEmissiveForDynamicAreaLighting(false)
 {
-	GetDiffuseComp()->SetBaseNames(DATASMITH_DIFFUSETEXNAME, DATASMITH_DIFFUSECOLNAME, TEXT("unsupported"), DATASMITH_DIFFUSECOMPNAME);
-	GetRefleComp()->SetBaseNames(DATASMITH_REFLETEXNAME, DATASMITH_REFLECOLNAME, TEXT("unsupported"), DATASMITH_REFLECOMPNAME);
-	GetRoughnessComp()->SetBaseNames(DATASMITH_ROUGHNESSTEXNAME, TEXT("unsupported"), DATASMITH_ROUGHNESSVALUENAME, DATASMITH_ROUGHNESSCOMPNAME);
-	GetNormalComp()->SetBaseNames(DATASMITH_NORMALTEXNAME, TEXT("unsupported"), DATASMITH_BUMPVALUENAME, DATASMITH_NORMALCOMPNAME);
-	GetBumpComp()->SetBaseNames(DATASMITH_BUMPTEXNAME, TEXT("unsupported"), DATASMITH_BUMPVALUENAME, DATASMITH_BUMPCOMPNAME);
-	GetTransComp()->SetBaseNames(DATASMITH_TRANSPTEXNAME, DATASMITH_TRANSPCOLNAME, TEXT("unsupported"), DATASMITH_TRANSPCOMPNAME);
-	GetMaskComp()->SetBaseNames(DATASMITH_CLIPTEXNAME, TEXT("unsupported"), TEXT("unsupported"), DATASMITH_CLIPCOMPNAME);
-	GetMetalComp()->SetBaseNames(DATASMITH_METALTEXNAME, TEXT("unsupported"), DATASMITH_METALVALUENAME, DATASMITH_METALCOMPNAME);
-	GetEmitComp()->SetBaseNames(DATASMITH_EMITTEXNAME, DATASMITH_EMITCOLNAME, TEXT("unsupported"), DATASMITH_EMITCOMPNAME);
-	GetWeightComp()->SetBaseNames(DATASMITH_WEIGHTTEXNAME, DATASMITH_WEIGHTCOLNAME, DATASMITH_WEIGHTVALUENAME, DATASMITH_WEIGHTCOMPNAME);
+	FDatasmithShaderElementImpl::GetDiffuseComp()->SetBaseNames(DATASMITH_DIFFUSETEXNAME, DATASMITH_DIFFUSECOLNAME, TEXT("unsupported"), DATASMITH_DIFFUSECOMPNAME);
+	FDatasmithShaderElementImpl::GetRefleComp()->SetBaseNames(DATASMITH_REFLETEXNAME, DATASMITH_REFLECOLNAME, TEXT("unsupported"), DATASMITH_REFLECOMPNAME);
+	FDatasmithShaderElementImpl::GetRoughnessComp()->SetBaseNames(DATASMITH_ROUGHNESSTEXNAME, TEXT("unsupported"), DATASMITH_ROUGHNESSVALUENAME, DATASMITH_ROUGHNESSCOMPNAME);
+	FDatasmithShaderElementImpl::GetNormalComp()->SetBaseNames(DATASMITH_NORMALTEXNAME, TEXT("unsupported"), DATASMITH_BUMPVALUENAME, DATASMITH_NORMALCOMPNAME);
+	FDatasmithShaderElementImpl::GetBumpComp()->SetBaseNames(DATASMITH_BUMPTEXNAME, TEXT("unsupported"), DATASMITH_BUMPVALUENAME, DATASMITH_BUMPCOMPNAME);
+	FDatasmithShaderElementImpl::GetTransComp()->SetBaseNames(DATASMITH_TRANSPTEXNAME, DATASMITH_TRANSPCOLNAME, TEXT("unsupported"), DATASMITH_TRANSPCOMPNAME);
+	FDatasmithShaderElementImpl::GetMaskComp()->SetBaseNames(DATASMITH_CLIPTEXNAME, TEXT("unsupported"), TEXT("unsupported"), DATASMITH_CLIPCOMPNAME);
+	FDatasmithShaderElementImpl::GetMetalComp()->SetBaseNames(DATASMITH_METALTEXNAME, TEXT("unsupported"), DATASMITH_METALVALUENAME, DATASMITH_METALCOMPNAME);
+	FDatasmithShaderElementImpl::GetEmitComp()->SetBaseNames(DATASMITH_EMITTEXNAME, DATASMITH_EMITCOLNAME, TEXT("unsupported"), DATASMITH_EMITCOMPNAME);
+	FDatasmithShaderElementImpl::GetWeightComp()->SetBaseNames(DATASMITH_WEIGHTTEXNAME, DATASMITH_WEIGHTCOLNAME, DATASMITH_WEIGHTVALUENAME, DATASMITH_WEIGHTCOMPNAME);
 }
 
 FDatasmithCompositeSurface::FDatasmithCompositeSurface(const TSharedPtr<IDatasmithCompositeTexture>& SubComp)
@@ -1044,7 +1044,7 @@ FDatasmithSceneImpl::FDatasmithSceneImpl(const TCHAR * InName)
 	Store.RegisterParameter(UserOS,             "UserOS"             );
 	Store.RegisterParameter(ExportDuration,     "ExportDuration"     );
 	Store.RegisterParameter(bUseSky,            "bUseSky"            );
-	Reset();
+	FDatasmithSceneImpl::Reset();
 }
 
 void FDatasmithSceneImpl::Reset()
diff --git a/Engine/Source/Runtime/Datasmith/DatasmithCore/Private/DatasmithSceneElementsImpl.h b/Engine/Source/Runtime/Datasmith/DatasmithCore/Private/DatasmithSceneElementsImpl.h
--- a/Engine/Source/Runtime/Datasmith/DatasmithCore/Private/DatasmithSceneElementsImpl.h
+++ b/Engine/Source/Runtime/Datasmith/DatasmithCore/Private/DatasmithSceneElementsImpl.h
@@ -61,7 +61,7 @@ inline FDatasmithElementImpl< InterfaceType >::FDatasmithElementImpl(const TCHAR
 	: Type(InType)
 	, Subtype(InSubType)
 {
-	SetName(InName);
+	FDatasmithElementImpl< InterfaceType >::SetName(InName);
 	Store.RegisterParameter(Type,    "Type");
 	Store.RegisterParameter(Subtype, "Subtype");
 	Store.RegisterParameter(Name,    "Name");
@@ -1549,8 +1549,8 @@ public:
 
 	virtual void Reset() override;
 
-	virtual const TCHAR* GetHost() const;
-	virtual void SetHost(const TCHAR* InHostname);
+	virtual const TCHAR* GetHost() const override;
+	virtual void SetHost(const TCHAR* InHostname) override;
 
 	virtual const TCHAR* GetExporterVersion() const override { return *(FString&)ExporterVersion; }
 	virtual void SetExporterVersion(const TCHAR* InVersion) override { ExporterVersion = InVersion; }
diff --git a/Engine/Source/Runtime/Datasmith/DatasmithCore/Private/DatasmithSceneXmlReader.cpp b/Engine/Source/Runtime/Datasmith/DatasmithCore/Private/DatasmithSceneXmlReader.cpp
--- a/Engine/Source/Runtime/Datasmith/DatasmithCore/Private/DatasmithSceneXmlReader.cpp
+++ b/Engine/Source/Runtime/Datasmith/DatasmithCore/Private/DatasmithSceneXmlReader.cpp
@@ -608,9 +608,11 @@ void FDatasmithSceneXmlReader::ParsePostProcessVolume(FXmlNode* InNode, const TS
 
 void FDatasmithSceneXmlReader::ParseColor(FXmlNode* InNode, FLinearColor& OutColor) const
 {
-	OutColor.R = ValueFromString<float>(InNode->GetAttribute(TEXT("R")));
-	OutColor.G = ValueFromString<float>(InNode->GetAttribute(TEXT("G")));
-	OutColor.B = ValueFromString<float>(InNode->GetAttribute(TEXT("B")));
+	OutColor = FLinearColor(
+		ValueFromString<float>(InNode->GetAttribute(TEXT("R"))),
+		ValueFromString<float>(InNode->GetAttribute(TEXT("G"))),
+		ValueFromString<float>(InNode->GetAttribute(TEXT("B")))
+	);
 }
 
 void FDatasmithSceneXmlReader::ParseComp(FXmlNode* InNode, TSharedPtr< IDatasmithCompositeTexture >& OutCompTexture, bool bInIsNormal) const
