Datasmith: Animation frames data (type FDatasmithTransformFrameInfo) made compatible with LargeWorldCoordinates

diff --git a/Engine/Source/Runtime/Datasmith/DatasmithCore/Private/DatasmithAnimationSerializer.cpp b/Engine/Source/Runtime/Datasmith/DatasmithCore/Private/DatasmithAnimationSerializer.cpp
--- a/Engine/Source/Runtime/Datasmith/DatasmithCore/Private/DatasmithAnimationSerializer.cpp
+++ b/Engine/Source/Runtime/Datasmith/DatasmithCore/Private/DatasmithAnimationSerializer.cpp
@@ -46,22 +46,22 @@ namespace DatasmithAnimationJsonSerializerImpl
 
 		// Serialize the frame transform only if the values have changed since the previous frame or represent a change with the following frame
 		if (EnumHasAnyFlags(EnabledComponents, ETransformChannelComponents::X) &&
-			(!FMath::IsNearlyEqual(FrameInfo.X, PreviousFrame.X, KINDA_SMALL_NUMBER) ||
-			(NextFrame.IsValid() && !FMath::IsNearlyEqual(FrameInfo.X, NextFrame.X, KINDA_SMALL_NUMBER))))
+			(!FMath::IsNearlyEqual(FrameInfo.X, PreviousFrame.X, DOUBLE_KINDA_SMALL_NUMBER) ||
+			(NextFrame.IsValid() && !FMath::IsNearlyEqual(FrameInfo.X, NextFrame.X, DOUBLE_KINDA_SMALL_NUMBER))))
 		{
 			Frame->SetNumberField(DSANIM_TRANSFORM_X, FrameInfo.X);
 			PreviousFrame.X = FrameInfo.X;
 		}
 		if (EnumHasAnyFlags(EnabledComponents, ETransformChannelComponents::Y) &&
-			(!FMath::IsNearlyEqual(FrameInfo.Y, PreviousFrame.Y, KINDA_SMALL_NUMBER) ||
-			(NextFrame.IsValid() && !FMath::IsNearlyEqual(FrameInfo.Y, NextFrame.Y, KINDA_SMALL_NUMBER))))
+			(!FMath::IsNearlyEqual(FrameInfo.Y, PreviousFrame.Y, DOUBLE_KINDA_SMALL_NUMBER) ||
+			(NextFrame.IsValid() && !FMath::IsNearlyEqual(FrameInfo.Y, NextFrame.Y, DOUBLE_KINDA_SMALL_NUMBER))))
 		{
 			Frame->SetNumberField(DSANIM_TRANSFORM_Y, FrameInfo.Y);
 			PreviousFrame.Y = FrameInfo.Y;
 		}
 		if (EnumHasAnyFlags(EnabledComponents, ETransformChannelComponents::Z) &&
-			(!FMath::IsNearlyEqual(FrameInfo.Z, PreviousFrame.Z, KINDA_SMALL_NUMBER) ||
-			(NextFrame.IsValid() && !FMath::IsNearlyEqual(FrameInfo.Z, NextFrame.Z, KINDA_SMALL_NUMBER))))
+			(!FMath::IsNearlyEqual(FrameInfo.Z, PreviousFrame.Z, DOUBLE_KINDA_SMALL_NUMBER) ||
+			(NextFrame.IsValid() && !FMath::IsNearlyEqual(FrameInfo.Z, NextFrame.Z, DOUBLE_KINDA_SMALL_NUMBER))))
 		{
 			Frame->SetNumberField(DSANIM_TRANSFORM_Z, FrameInfo.Z);
 			PreviousFrame.Z = FrameInfo.Z;
diff --git a/Engine/Source/Runtime/Datasmith/DatasmithCore/Private/DatasmithTypes.cpp b/Engine/Source/Runtime/Datasmith/DatasmithCore/Private/DatasmithTypes.cpp
--- a/Engine/Source/Runtime/Datasmith/DatasmithCore/Private/DatasmithTypes.cpp
+++ b/Engine/Source/Runtime/Datasmith/DatasmithCore/Private/DatasmithTypes.cpp
@@ -4,18 +4,18 @@
 #include "DatasmithDefinitions.h"
 #include "Math/UnrealMathUtility.h"
 
-FDatasmithTransformFrameInfo FDatasmithTransformFrameInfo::InvalidFrameInfo(MIN_int32, 0.f, 0.f, 0.f);
+FDatasmithTransformFrameInfo FDatasmithTransformFrameInfo::InvalidFrameInfo(MIN_int32, 0., 0., 0.);
 
 bool FDatasmithTransformFrameInfo::IsValid() const
-{ 
+{
 	return FrameNumber != MIN_int32;
 }
 
 bool FDatasmithTransformFrameInfo::operator==(const FDatasmithTransformFrameInfo& Other) const
 {
-	return FMath::IsNearlyEqual(X, Other.X, KINDA_SMALL_NUMBER) &&
-		   FMath::IsNearlyEqual(Y, Other.Y, KINDA_SMALL_NUMBER) &&
-		   FMath::IsNearlyEqual(Z, Other.Z, KINDA_SMALL_NUMBER);
+	return FMath::IsNearlyEqual(X, Other.X, DOUBLE_KINDA_SMALL_NUMBER) &&
+	       FMath::IsNearlyEqual(Y, Other.Y, DOUBLE_KINDA_SMALL_NUMBER) &&
+	       FMath::IsNearlyEqual(Z, Other.Z, DOUBLE_KINDA_SMALL_NUMBER);
 }
 
 FDatasmithVisibilityFrameInfo FDatasmithVisibilityFrameInfo::InvalidFrameInfo(MIN_int32, false);
diff --git a/Engine/Source/Runtime/Datasmith/DatasmithCore/Public/DatasmithTypes.h b/Engine/Source/Runtime/Datasmith/DatasmithCore/Public/DatasmithTypes.h
--- a/Engine/Source/Runtime/Datasmith/DatasmithCore/Public/DatasmithTypes.h
+++ b/Engine/Source/Runtime/Datasmith/DatasmithCore/Public/DatasmithTypes.h
@@ -82,7 +82,7 @@ public:
 	int MirrorY;
 };
 
-/** 
+/**
  * FDatasmithTransformFrameInfo holds the data for the transform values of a frame
  * The transform values must be relative to the parent
  * The rotation is represented as Euler angles in degrees
@@ -105,7 +105,7 @@ struct DATASMITHCORE_API FDatasmithTransformFrameInfo
 		Z = EulerAngles.Z;
 	}
 
-	FDatasmithTransformFrameInfo(int32 InFrameNumber, float InX, float InY, float InZ)
+	FDatasmithTransformFrameInfo(int32 InFrameNumber, double InX, double InY, double InZ)
 		: FrameNumber(InFrameNumber)
 		, X(InX)
 		, Y(InY)
@@ -119,9 +119,9 @@ struct DATASMITHCORE_API FDatasmithTransformFrameInfo
 
 	// For translation and scale, X, Y, Z are the transform values of the respective components
 	// For rotation, they represent rotation around the X, Y, Z axis (roll, pitch, yaw respectively) in degrees
-	float X;
-	float Y;
-	float Z;
+	double X;
+	double Y;
+	double Z;
 
 	static FDatasmithTransformFrameInfo InvalidFrameInfo;
 };
