CADKernel: small improvement in topology functions + some cleaning

diff --git a/Engine/Source/Runtime/Datasmith/CADKernel/Private/CADKernel/Geo/Sampler/SamplerAbstract.h b/Engine/Source/Runtime/Datasmith/CADKernel/Private/CADKernel/Geo/Sampler/SamplerAbstract.h
--- a/Engine/Source/Runtime/Datasmith/CADKernel/Private/CADKernel/Geo/Sampler/SamplerAbstract.h
+++ b/Engine/Source/Runtime/Datasmith/CADKernel/Private/CADKernel/Geo/Sampler/SamplerAbstract.h
@@ -151,7 +151,8 @@ namespace CADKernel
 		void RunSampling()
 		{
 
-			while (int32 CandidatePointsCount = CandidatePoints.Size())
+			// < 100000 : check to avoid that the process loops endlessly
+			while (int32 CandidatePointsCount = CandidatePoints.Size() && Sampling.Coordinates.Num() < 100000)
 			{
 				NextCoordinates.Empty();
 
diff --git a/Engine/Source/Runtime/Datasmith/CADKernel/Private/CADKernel/Mesh/Structure/Mesh.cpp b/Engine/Source/Runtime/Datasmith/CADKernel/Private/CADKernel/Mesh/Structure/Mesh.cpp
--- a/Engine/Source/Runtime/Datasmith/CADKernel/Private/CADKernel/Mesh/Structure/Mesh.cpp
+++ b/Engine/Source/Runtime/Datasmith/CADKernel/Private/CADKernel/Mesh/Structure/Mesh.cpp
@@ -20,8 +20,8 @@ int32 FMesh::RegisterCoordinates()
 FInfoEntity& FMesh::GetInfo(FInfoEntity& Info) const
 {
 	return FEntityGeom::GetInfo(Info)
-		.Add(TEXT("geometric Entity"), GetGeometricEntity())
-		.Add(TEXT("mesh model"), GetMeshModel())
+		.Add(TEXT("Geometric Entity"), (FEntity&) GetGeometricEntity())
+		.Add(TEXT("Mesh model"), (FEntity&) GetMeshModel())
 		.Add(TEXT("Node Num"), (int32) NodeCoordinates.Num());
 }
 #endif
diff --git a/Engine/Source/Runtime/Datasmith/CADKernel/Private/CADKernel/Topo/Joiner.cpp b/Engine/Source/Runtime/Datasmith/CADKernel/Private/CADKernel/Topo/Joiner.cpp
--- a/Engine/Source/Runtime/Datasmith/CADKernel/Private/CADKernel/Topo/Joiner.cpp
+++ b/Engine/Source/Runtime/Datasmith/CADKernel/Private/CADKernel/Topo/Joiner.cpp
@@ -19,7 +19,27 @@
 
 using namespace CADKernel;
 
-FJoiner::FJoiner(TSharedRef<FSession> InSession, const TArray<TSharedPtr<FTopologicalFace>>& InFaces, double InTolerance)
+FJoiner::FJoiner(FSession& InSession, const TSharedRef<FModel>& InModel, double InTolerance)
+	: Session(InSession)
+	, JoiningTolerance(InTolerance* UE_SQRT_2)
+	, JoiningToleranceSquare(FMath::Square(JoiningTolerance))
+{
+	int32 ShellCount = 0;
+	for (const TSharedPtr<FBody>& Body : InModel->GetBodies())
+	{
+		ShellCount += Body->GetShells().Num();
+	}
+	Shells.Reserve(ShellCount);
+
+	for (const TSharedPtr<FBody>& Body : InModel->GetBodies())
+	{
+		Shells.Append(Body->GetShells());
+	}
+
+	InitFaces();
+}
+
+FJoiner::FJoiner(FSession& InSession, const TArray<TSharedPtr<FTopologicalFace>>& InFaces, double InTolerance)
 	: Session(InSession)
 	, Faces(InFaces)
 	, JoiningTolerance(InTolerance * UE_SQRT_2)
@@ -27,6 +47,45 @@ FJoiner::FJoiner(TSharedRef<FSession> InSession, const TArray<TSharedPtr<FTopolo
 {
 }
 
+FJoiner::FJoiner(FSession& InSession, const TArray<TSharedPtr<FShell>>& InShells, double InTolerance)
+	: Session(InSession)
+	, Shells(InShells)
+	, JoiningTolerance(InTolerance)
+	, JoiningToleranceSquare(FMath::Square(JoiningTolerance))
+{
+	InitFaces();
+}
+
+void FJoiner::InitFaces()
+{
+	int32 FaceCount = 0;
+	for (const TSharedPtr<FShell>& Shell : Shells)
+	{
+		Shell->CompleteMetadata();
+		FaceCount += Shell->FaceCount();
+	}
+	Faces.Reserve(FaceCount);
+
+	for (const TSharedPtr<FShell>& Shell : Shells)
+	{
+		Shell->SpreadBodyOrientation();
+		for (const FOrientedFace& Face : Shell->GetFaces())
+		{
+			Faces.Add(Face.Entity);
+		}
+		for (const FOrientedFace& Face : Shell->GetFaces())
+		{
+			Face.Entity->CompleteMetadata();
+		}
+	}
+
+	for (const TSharedPtr<FTopologicalFace>& Face : Faces)
+	{
+		Face->ResetMarker2();
+	}
+}
+
+
 void FJoiner::RemoveFacesFromShell()
 {
 	// remove faces from their shells
@@ -39,6 +98,7 @@ void FJoiner::RemoveFacesFromShell()
 			Face->ResetHost();
 		}	
 	}
+
 	for (TWeakPtr<FShell> WeakShell : ShellSet)
 	{
 		TSharedPtr<FShell> Shell = WeakShell.Pin();
@@ -64,40 +124,6 @@ void FJoiner::RemoveFacesFromShell()
 	}
 }
 
-FJoiner::FJoiner(TSharedRef<FSession> InSession, const TArray<TSharedPtr<FShell>>& InShells, double InTolerance)
-	: Session(InSession)
-	, Shells(InShells)
-	, JoiningTolerance(InTolerance)
-	, JoiningToleranceSquare(FMath::Square(JoiningTolerance))
-{
-	int32 FaceCount = 0;
-	for (const TSharedPtr<FShell>& Shell : InShells)
-	{
-		Shell->CompleteMetadata();
-		FaceCount += Shell->FaceCount();
-	}
-	Faces.Reserve(FaceCount);
-
-	for (const TSharedPtr<FShell>& Shell : InShells)
-	{
-		Shell->SpreadBodyOrientation();
-		for (const FOrientedFace& Face : Shell->GetFaces())
-		{
-			Faces.Add(Face.Entity);
-		}
-		for (const FOrientedFace& Face : Shell->GetFaces())
-		{
-			Face.Entity->CompleteMetadata();
-		}
-	}
-
-	for (const TSharedPtr<FTopologicalFace>& Face : Faces)
-	{
-		Face->ResetMarker2();
-	}
-}
-
-
 void FJoiner::EmptyShells()
 {
 	for (const TSharedPtr<FShell>& Shell : Shells)
@@ -1000,7 +1026,7 @@ void FJoiner::SplitIntoConnectedShell()
 		else
 		{
 			TSharedRef<FBody> Body = FEntity::MakeShared<FBody>();
-			Session->GetModel()->Add(Body);
+			Session.GetModel()->Add(Body);
 			TSharedRef<FShell> Shell = FEntity::MakeShared<FShell>();
 			Body->AddShell(Shell);
 			Body->SetName(FaceSubset.MainName);
@@ -1012,7 +1038,7 @@ void FJoiner::SplitIntoConnectedShell()
 		}
 	}
 
-	Session->GetModel()->RemoveEmptyBodies();
+	Session.GetModel()->RemoveEmptyBodies();
 }
 
 
diff --git a/Engine/Source/Runtime/Datasmith/CADKernel/Private/CADKernel/Topo/Joiner.h b/Engine/Source/Runtime/Datasmith/CADKernel/Private/CADKernel/Topo/Joiner.h
deleted file mode 100644
diff --git a/Engine/Source/Runtime/Datasmith/CADKernel/Public/CADKernel/Core/Session.h b/Engine/Source/Runtime/Datasmith/CADKernel/Public/CADKernel/Core/Session.h
--- a/Engine/Source/Runtime/Datasmith/CADKernel/Public/CADKernel/Core/Session.h
+++ b/Engine/Source/Runtime/Datasmith/CADKernel/Public/CADKernel/Core/Session.h
@@ -103,7 +103,7 @@ namespace CADKernel
 			return Database.SpawnEntityIdent(SelectedEntities, bForceSpawning);
 		}
 
-		uint32 SpawnEntityIdent(TSharedPtr<FEntity>& SelectedEntity, bool bForceSpawning)
+		uint32 SpawnEntityIdent(TSharedPtr<FEntity> SelectedEntity, bool bForceSpawning)
 		{
 			return Database.SpawnEntityIdent(SelectedEntity, bForceSpawning);
 		}
diff --git a/Engine/Source/Runtime/Datasmith/CADKernel/Public/CADKernel/Topo/Joiner.h b/Engine/Source/Runtime/Datasmith/CADKernel/Public/CADKernel/Topo/Joiner.h
new file mode 100644
--- /dev/null
+++ b/Engine/Source/Runtime/Datasmith/CADKernel/Public/CADKernel/Topo/Joiner.h
@@ -0,0 +1,124 @@
+// Copyright Epic Games, Inc. All Rights Reserved.
+#pragma once
+
+#include "CADKernel/Core/Types.h"
+
+namespace CADKernel
+{
+	class FModel;
+	class FSession;
+	class FShell;
+	class FTopologicalEdge;
+	class FTopologicalFace;
+	class FTopologicalVertex;
+
+	class CADKERNEL_API FJoiner
+	{
+		friend class FThinZone;
+		friend class FThinZoneFinder;
+
+	protected:
+
+		FSession& Session;
+
+		TArray<TSharedPtr<FShell>> Shells;
+		TArray<TSharedPtr<FTopologicalFace>> Faces;
+
+		double JoiningTolerance;
+		double JoiningToleranceSquare;
+
+	public:
+
+		FJoiner(FSession& InSession, const TSharedRef<FModel>& Shells, double Tolerance);
+		FJoiner(FSession& InSession, const TArray<TSharedPtr<FShell>>& Shells, double Tolerance);
+		FJoiner(FSession& InSession, const TArray<TSharedPtr<FTopologicalFace>>& Surfaces, double Tolerance);
+
+		void JoinFaces();
+		//void JoinFaces(bool bProcessOnlyBorderEdges, bool bProcessOnlyNonManifoldEdges);
+
+		/**
+		 * Check topology of each body
+		 */
+		void CheckTopology();
+
+		//void MergeInto(TSharedPtr<FBody> Body, TArray<TSharedPtr<FTopologicalEntity>>& InEntities);
+		//void SortByShell(TSharedPtr<FBody> Body, TArray<TSharedPtr<FBody>>& OutNewBody);
+		//void Join(TArray<TSharedPtr<FBody>> Bodies, double Tolerance);
+
+		/**
+		 * Split into connected shell and put each shell into the appropriate body
+		 */
+		void SplitIntoConnectedShell();
+
+	private:
+
+		/**
+		 * Call by constructor.
+		 * For each shell, add their faces into Faces array, complete the metadata and set the states for the joining process
+		 */
+		void InitFaces();
+
+		void EmptyShells();
+
+		void RemoveFacesFromShell();
+
+		/**
+		 * Return an array of active vertices.
+		 */
+		void GetVertices(TArray<TSharedPtr<FTopologicalVertex>>& Vertices);
+
+		/**
+		 * Return an array of active border vertices.
+		 */
+		void GetBorderVertices(TArray<TSharedPtr<FTopologicalVertex>>& BorderVertices);
+
+		/**
+		 * Merge Border Vertices with other vertices.
+		 * @param Vertices: the initial array of active vertices, this array is updated at the end of the process
+		 */
+		void MergeCoincidentVertices(TArray<TSharedPtr<FTopologicalVertex>>& VerticesToMerge);
+
+		/**
+		 * Merge Border Vertices with other vertices.
+		 * @param Vertices: the initial array of active vertices, this array is updated at the end of the process
+		 */
+		void MergeBorderVerticesWithCoincidentOtherVertices(TArray<TSharedPtr<FTopologicalVertex>>& Vertices);
+
+		//void FixCollapsedEdges();
+		void CheckSelfConnectedEdge();
+		void RemoveIsolatedEdges(); // Usefull ???
+
+		/**
+		 * First step, trivial edge merge i.e. couple of edges with same extremity vertex
+		 */
+		void MergeCoincidentEdges(TArray<TSharedPtr<FTopologicalVertex>>& Vertices);
+
+		/**
+		 * Second step, parallel edges but with different length. the longest must be split
+		 */
+		void StitchParallelEdges(TArray<TSharedPtr<FTopologicalVertex>>& Vertices);
+
+		/**
+		 * Split the edge to split except if one side is too small (JoiningTolerance). In this case (too small), the EdgeToSplit is not split but linked to EdgeToLink. 
+		 * @return the created vertex or TSharedPtr<FTopologicalVertex>()
+		 */
+		TSharedPtr<FTopologicalVertex> SplitAndLink(TSharedRef<FTopologicalVertex>& StartVertex, TSharedPtr<FTopologicalEdge>& EdgeToLink, TSharedPtr<FTopologicalEdge>& EdgeToSplit);
+
+		/**
+		 * For each loop of each surface, check if successive edges are unconnected and if their common vertices are connected only to them. 
+		 * These edges are merged into one edge. 
+		 * E.g. :
+		 * Face A has 3 successive unconnected edges. If these 3 edges are merged to give only one edge, the new edge could be linked to its parallel edge of Face B 
+		 * 
+		 *              \                          Face A                                   |    Face C
+		 *               \                                                                  |
+		 *    Face E     CV ------------------- UV ------------------ UV ----------------- CV --------------------
+		 *               CV -------------------------------------------------------------- CV --------------------
+		 *               /                         Face B                                   |    Face D
+		 *              /                                                                   |
+		 */
+		void MergeUnconnectedAdjacentEdges();
+
+	};
+
+} // namespace CADKernel
diff --git a/Engine/Source/Runtime/Datasmith/CADKernel/Public/CADKernel/UI/Display.h b/Engine/Source/Runtime/Datasmith/CADKernel/Public/CADKernel/UI/Display.h
--- a/Engine/Source/Runtime/Datasmith/CADKernel/Public/CADKernel/UI/Display.h
+++ b/Engine/Source/Runtime/Datasmith/CADKernel/Public/CADKernel/UI/Display.h
@@ -37,7 +37,11 @@ namespace CADKernel
 
 	struct FLinearBoundary;
 
+#ifdef CADKERNEL_DEV
 	void Wait(bool bMakeWait = true);
+#else
+	inline void Wait(bool bMakeWait = true) {};
+#endif
 
 	void Open3DDebugSession(FString name, const TArray<FIdent>& idList = TArray<FIdent>());
 	void Close3DDebugSession();
@@ -73,7 +77,7 @@ namespace CADKernel
 		}
 	};
 
-	void CADKERNEL_API FlushVisu();
+	CADKERNEL_API void FlushVisu();
 
 	template<typename TPoint>
 	void DrawPoint(const TPoint& InPoint, EVisuProperty Property = EVisuProperty::BluePoint)
