Datasmith: Support XML serialization of Actors visibility

diff --git a/Engine/Source/Runtime/Datasmith/DatasmithCore/Private/DatasmithSceneXmlReader.cpp b/Engine/Source/Runtime/Datasmith/DatasmithCore/Private/DatasmithSceneXmlReader.cpp
--- a/Engine/Source/Runtime/Datasmith/DatasmithCore/Private/DatasmithSceneXmlReader.cpp
+++ b/Engine/Source/Runtime/Datasmith/DatasmithCore/Private/DatasmithSceneXmlReader.cpp
@@ -755,6 +755,9 @@ void FDatasmithSceneXmlReader::ParseActor(FXmlNode* InNode, TSharedPtr<IDatasmit
 
 	InOutElement->SetIsAComponent( ValueFromString<bool>( InNode->GetAttribute(TEXT("component"))) );
 
+	FString VisibleAtribute = InNode->GetAttribute(TEXT("visible"));
+	InOutElement->SetVisibility(VisibleAtribute.IsEmpty() ? true : ValueFromString<bool>(VisibleAtribute));
+
 	for (FXmlNode* ChildNode : InNode->GetChildrenNodes())
 	{
 		if (ChildNode->GetTag() == TEXT("transform"))
@@ -767,9 +770,6 @@ void FDatasmithSceneXmlReader::ParseActor(FXmlNode* InNode, TSharedPtr<IDatasmit
 		}
 		else if (ChildNode->GetTag() == TEXT("children"))
 		{
-			// @todo: If attribute does not exist, default value should be true not false.
-			InOutElement->SetVisibility(ValueFromString<bool>(ChildNode->GetAttribute(TEXT("visible"))));
-
 			// Recursively parse the children, can be of any supported actor type
 			for (FXmlNode* ChildActorNode : ChildNode->GetChildrenNodes())
 			{
diff --git a/Engine/Source/Runtime/Datasmith/DatasmithCore/Private/DatasmithSceneXmlWriter.cpp b/Engine/Source/Runtime/Datasmith/DatasmithCore/Private/DatasmithSceneXmlWriter.cpp
--- a/Engine/Source/Runtime/Datasmith/DatasmithCore/Private/DatasmithSceneXmlWriter.cpp
+++ b/Engine/Source/Runtime/Datasmith/DatasmithCore/Private/DatasmithSceneXmlWriter.cpp
@@ -120,6 +120,11 @@ FString FDatasmithSceneXmlWriterImpl::GetLabelAndLayer(const TSharedPtr<IDatasmi
 		LabelAndLayer += TEXT(" layer=\"") + SanitizeXMLText(ActorElement->GetLayer()) + TEXT("\"");
 	}
 
+	if (!ActorElement->GetVisibility())
+	{
+		LabelAndLayer += TEXT(" visible=\"0\"");
+	}
+
 	return LabelAndLayer;
 }
 
@@ -616,6 +621,8 @@ void FDatasmithSceneXmlWriterImpl::WriteActorChildren(const TSharedPtr< IDatasmi
 	FString XmlString;
 	if (ActorElement->GetChildrenCount() > 0)
 	{
+		// Compatibility note: in 4.26 and previous version, this 'visible' attribute was mandatory (because it's value defaulted to false).
+		// In 4.27+, the 'visible' attribute moved on the actor tag itself.
 		if (!ActorElement->GetVisibility())
 		{
 			XmlString = TEXT("<children visible=\"false\">");
