datasmith: fix broken sanitation introduced in CL 18236389

diff --git a/Engine/Source/Runtime/Datasmith/DatasmithCore/Private/DatasmithSceneElementsImpl.h b/Engine/Source/Runtime/Datasmith/DatasmithCore/Private/DatasmithSceneElementsImpl.h
--- a/Engine/Source/Runtime/Datasmith/DatasmithCore/Private/DatasmithSceneElementsImpl.h
+++ b/Engine/Source/Runtime/Datasmith/DatasmithCore/Private/DatasmithSceneElementsImpl.h
@@ -28,7 +28,7 @@ public:
 	virtual bool IsA( EDatasmithElementType InType ) const override { return EnumHasAnyFlags( GetElementType(), InType); }
 
 	virtual const TCHAR* GetName() const override { return *Name.Get(); }
-	virtual void SetName(const TCHAR* InName) override { Name = InName; }
+	virtual void SetName(const TCHAR* InName) override { Name = FDatasmithUtils::SanitizeObjectName(InName); }
 
 	virtual const TCHAR* GetLabel() const override { const FString& Tmp = Label.Get(); return Tmp.IsEmpty() ? GetName() : *Tmp; }
 	virtual void SetLabel(const TCHAR* InLabel) override { Label = FDatasmithUtils::SanitizeObjectName(InLabel); }
