CADKernel update: Improve the fix of loop intersections
- Fix TryToConnectTwoSubLoopsWithShortestSegment:
    When adding a segment A-B, NextA-B and A-NextB were also added without checking that NextA or NextB were in the cell.
- Fix ConnectCellSubLoopsByNeighborhood:
   [Connect inner close loops] and [Connect inner with Outer loop] step:
       - For each node pair -> TryToConnectTwoSubLoopsWithTheMostIsoSegment defined a set of candidate segment
       - Cell.SelectSegmentInCandidateSegments selects in Cell.CandidateSegments segments and adds them in Cell.FinalSegments (add also in Cell.IntersectionTool)
fyi jeanluc.corenthin, flavien.picon

diff --git a/Engine/Source/Runtime/Datasmith/CADKernel/Private/CADKernel/Core/Version.h b/Engine/Source/Runtime/Datasmith/CADKernel/Private/CADKernel/Core/Version.h
--- a/Engine/Source/Runtime/Datasmith/CADKernel/Private/CADKernel/Core/Version.h
+++ b/Engine/Source/Runtime/Datasmith/CADKernel/Private/CADKernel/Core/Version.h
@@ -2,8 +2,8 @@
 #pragma once
 
 //** ** Version du Toolkit** **//
-#define TOOLKIT_VERSION 0, 2, 18230756
-#define TOOLKIT_VERSION_ASCII "0.2.18230756"
+#define TOOLKIT_VERSION 0, 2, 18231529
+#define TOOLKIT_VERSION_ASCII "0.2.18231529"
 
 #define RELEASE_DATE	 __DATE__
 
diff --git a/Engine/Source/Runtime/Datasmith/CADKernel/Private/CADKernel/Mesh/Meshers/IsoTriangulator.cpp b/Engine/Source/Runtime/Datasmith/CADKernel/Private/CADKernel/Mesh/Meshers/IsoTriangulator.cpp
--- a/Engine/Source/Runtime/Datasmith/CADKernel/Private/CADKernel/Mesh/Meshers/IsoTriangulator.cpp
+++ b/Engine/Source/Runtime/Datasmith/CADKernel/Private/CADKernel/Mesh/Meshers/IsoTriangulator.cpp
@@ -49,11 +49,7 @@ bool FIsoTriangulator::Triangulate()
 {
 	EGridSpace DisplaySpace = EGridSpace::UniformScaled;
 #ifdef DEBUG_ISOTRIANGULATOR
-	if (bDisplay)
-	{
-		Open3DDebugSession(FString::Printf(TEXT("Triangulate %d"), Grid.GetFace()->GetId()));
-		//Wait(Grid.GetFace()->GetId() == FaceToDebug);
-	}
+	F3DDebugSession _(bDisplay, FString::Printf(TEXT("Triangulate %d"), Grid.GetFace()->GetId()));
 #endif
 
 	FTimePoint StartTime = FChrono::Now();
@@ -64,37 +60,37 @@ bool FIsoTriangulator::Triangulate()
 
 	BuildNodes();
 
-#ifdef CADKERNEL_DEV
+#ifdef DEBUG_ISOTRIANGULATOR
 	DisplayIsoNodes(DisplaySpace);
 #endif
 
 	FillMeshNodes();
 	BuildLoopSegments();
 
-#ifdef CADKERNEL_DEV
+#ifdef DEBUG_ISOTRIANGULATOR
 	Display(DisplaySpace, TEXT("FIsoTrianguler::LoopSegments"), LoopSegments, false, false, EVisuProperty::OrangeCurve);
 #endif
 
 	CheckLoopInSelfIntersectingAndFixIt();
 
-#ifdef CADKERNEL_DEV
+#ifdef DEBUG_ISOTRIANGULATOR
 	Display(DisplaySpace, TEXT("FIsoTrianguler::LoopSegments"), LoopSegments, true, true, EVisuProperty::OrangeCurve);
 	//Wait();
 #endif
 
 
-#ifdef CADKERNEL_DEV
+#ifdef DEBUG_ISOTRIANGULATOR
 	Display(DisplaySpace, TEXT("FIsoTrianguler::LoopSegments"), LoopSegments, false, false, EVisuProperty::OrangeCurve);
 #endif
 
 	BuildThinZoneSegments();
-#ifdef CADKERNEL_DEV
+#ifdef DEBUG_ISOTRIANGULATOR
 	Display(DisplaySpace, TEXT("FIsoTrianguler::ThinZoneSegments"), ThinZoneSegments, false);
 	LoopSegmentsIntersectionTool.Display(TEXT("FIsoTrianguler::IntersectionTool Loop"));
 #endif
 
 	BuildInnerSegments();
-#ifdef CADKERNEL_DEV
+#ifdef DEBUG_ISOTRIANGULATOR
 	Display(DisplaySpace, TEXT("FIsoTrianguler::FinalInnerSegments"), FinalInnerSegments, false, false, EVisuProperty::BlueCurve);
 	InnerToOuterSegmentsIntersectionTool.Display(TEXT("FIsoTrianguler::IntersectionTool InnerToOutter"));
 	Chronos.TriangulateDuration1 = FChrono::Elapse(StartTime);
@@ -104,11 +100,11 @@ bool FIsoTriangulator::Triangulate()
 	// =============================================================================================================
 
 	BuildInnerSegmentsIntersectionTool();
-#ifdef CADKERNEL_DEV
+#ifdef DEBUG_ISOTRIANGULATOR
 	InnerSegmentsIntersectionTool.Display(TEXT("FIsoTrianguler::IntersectionTool Inner"));
 #endif
 
-#ifdef CADKERNEL_DEV
+#ifdef DEBUG_ISOTRIANGULATOR
 	Chronos.TriangulateDuration2 = FChrono::Elapse(StartTime);
 #endif
 
@@ -119,8 +115,11 @@ bool FIsoTriangulator::Triangulate()
 	// 	           - build Delaunay connection
 	// 	           - find the shortest segment to connect each connected loop by Delaunay
 	// =============================================================================================================
+#ifdef DEBUG_ISOTRIANGULATOR
+	Display(DisplaySpace, TEXT("FIsoTrianguler::Final To Loops Before"), FinalToLoops, false, false, EVisuProperty::YellowCurve);
+#endif
 	ConnectCellLoops();
-#ifdef CADKERNEL_DEV
+#ifdef DEBUG_ISOTRIANGULATOR
 	Display(DisplaySpace, TEXT("FIsoTrianguler::Final Iso ToLink Inner To Loops"), FinalToLoops, false, false, EVisuProperty::YellowCurve);
 #endif
 
@@ -203,13 +202,11 @@ bool FIsoTriangulator::Triangulate()
 	// Finalise the mesh by the tessellation of the inner grid
 	TriangulateInnerNodes();
 
-#ifdef CADKERNEL_DEV
+#ifdef DEBUG_ISOTRIANGULATOR
 	if (bDisplay)
 	{
-		Open3DDebugSession(TEXT("Mesh 3D"));
+		F3DDebugSession A(TEXT("Mesh 3D"));
 		DisplayMesh(*Mesh);
-		Close3DDebugSession();
-		Close3DDebugSession();
 	}
 	//Wait(bDisplay);
 #endif
@@ -906,8 +903,6 @@ void FIsoTriangulator::FindIsoSegmentToLinkLoopToLoop()
 
 void FIsoTriangulator::FindSegmentToLinkOuterLoopNodes(FCell& Cell)
 {
-	//F3DDebugSession _(TEXT("FindSegmentToLinkOuterLoopNodes"));
-
 	TArray<FLoopNode*>& OuterLoop = Cell.SubLoops[0];
 	int32 NodeCount = OuterLoop.Num();
 
@@ -1152,32 +1147,55 @@ void FIsoTriangulator::FindIsoSegmentToLinkInnerToLoop()
 // 	           - build Delaunay connection
 // 	           - find the shortest segment to connect each connected loop by Delaunay
 // =============================================================================================================
-//#define DEBUG_CELLS
+//#define DEBUG_CONNECT_CELL_LOOPS
 void FIsoTriangulator::ConnectCellLoops()
 {
 	TArray<FCell> Cells;
 	FindCellContainingBoundaryNodes(Cells);
-#ifdef DEBUG_CELLS
-	//DisplayCells(Cells);
-	//Wait();
+
+#ifdef DEBUG_CONNECT_CELL_LOOPS
 	F3DDebugSession _(bDisplay, ("ConnectCellLoops"));
 #endif
 
 	FinalToLoops.Reserve(LoopNodeCount + InnerNodeCount);
 	for (FCell& Cell : Cells)
 	{
+#ifdef DEBUG_CONNECT_CELL_LOOPS
+		F3DDebugSession _(bDisplay, FString::Printf(TEXT("Cell %d"), Cell.Id));
+		DisplayCell(Cell);
+#endif
+
 		ConnectCellSubLoopsByNeighborhood(Cell);
+
+#ifdef DEBUG_CONNECT_CELL_LOOPS
+		Display(EGridSpace::UniformScaled, TEXT("ConnectCellLoops::Step 0"), Cell.FinalSegments, false, false, EVisuProperty::YellowCurve);
+#endif
+
 		if (Cell.bHasOuterLoop)
 		{
 			FindIsoSegmentToLinkOuterLoopNodes(Cell);
+#ifdef DEBUG_CONNECT_CELL_LOOPS
+			Display(EGridSpace::UniformScaled, TEXT("ConnectCellLoops::Step 1"), Cell.FinalSegments, false, false, EVisuProperty::YellowCurve);
+#endif
+
 			if (Cell.CandidateSegments.Num() == 0)
 			{
 				FindSegmentToLinkOuterLoopNodes(Cell);
+#ifdef DEBUG_CONNECT_CELL_LOOPS
+				Display(EGridSpace::UniformScaled, TEXT("ConnectCellLoops::Step 2"), Cell.FinalSegments, false, false, EVisuProperty::YellowCurve);
+#endif
 			}
 		}
 		ConnectCellCornerToInnerLoop(Cell);
+#ifdef DEBUG_CONNECT_CELL_LOOPS
+		Display(EGridSpace::UniformScaled, TEXT("ConnectCellLoops::Step 3"), Cell.FinalSegments, false, false, EVisuProperty::YellowCurve);
+#endif
 
 		FinalToLoops.Append(Cell.FinalSegments);
+
+#ifdef DEBUG_CONNECT_CELL_LOOPS
+		Display(EGridSpace::UniformScaled, TEXT("ConnectCellLoops::Final Iso ToLink Inner To Loops"), FinalToLoops, false, false, EVisuProperty::YellowCurve);
+#endif
 	}
 }
 
@@ -2004,7 +2022,7 @@ bool FIsoTriangulator::TryToRemoveIntersectionBySwappingSegments(FIsoSegment* Se
 	Segment1->GetFirstNode().ConnectSegment(Segment1);
 	Segment0->GetSecondNode().ConnectSegment(Segment0);
 
-#ifdef CADKERNEL_DEV
+#ifdef DEBUG_BY_SWAPPING_SEGMENTS
 	if (bDisplay)
 	{
 		F3DDebugSession _(TEXT("New Segments"));
@@ -2035,7 +2053,7 @@ void FIsoTriangulator::CheckLoopInSelfIntersectingAndFixIt()
 
 		if (const FIsoSegment* IntersectingSegment = IntersectionTool.DoesIntersect(*Segment))
 		{
-#ifdef CADKERNEL_DEV
+#ifdef DEBUG_CHECK_LOOP_IN_SELF_INTERSECTING_AND_FIX_IT
 			//if (bDisplay)
 			//{
 			//	{
@@ -2099,7 +2117,7 @@ bool FIsoTriangulator::CanCycleBeFixedAndMeshed(const TArray<FIsoSegment*>& Cycl
 		{
 			FMessage::Printf(Debug, TEXT("Remove local self intesection.\n"), Grid.GetFace()->GetId());
 
-#ifdef CADKERNEL_DEV
+#ifdef DEBUG_CAN_CYCLE_BE_FIXED_AND_MESHED
 			if (bDisplay)
 			{
 				F3DDebugSession _(TEXT("Segment"));
@@ -2117,7 +2135,7 @@ bool FIsoTriangulator::CanCycleBeFixedAndMeshed(const TArray<FIsoSegment*>& Cycl
 				{
 					if (!TryToRemoveIntersectionByMovingTheClosedOusidePoint(IntersectingSegment, Segment))
 					{
-#ifdef CADKERNEL_DEV
+#ifdef DEBUG_CAN_CYCLE_BE_FIXED_AND_MESHED
 						if (bDisplay)
 						{
 							F3DDebugSession _(TEXT("Failed"));
@@ -2132,7 +2150,7 @@ bool FIsoTriangulator::CanCycleBeFixedAndMeshed(const TArray<FIsoSegment*>& Cycl
 				}
 			}
 
-#ifdef CADKERNEL_DEV
+#ifdef DEBUG_CAN_CYCLE_BE_FIXED_AND_MESHED
 			if (bDisplay)
 			{
 				F3DDebugSession _(TEXT("New Segment"));
@@ -3196,16 +3214,14 @@ void FIsoTriangulator::TriangulateInnerNodes()
 // =========================================================================================================================================================================================================
 // =========================================================================================================================================================================================================
 // =========================================================================================================================================================================================================
-
+//#define DEBUG_CONNECT_CELL_SUB_LOOPS_BY_NEIGHBORHOOD
 void FIsoTriangulator::ConnectCellSubLoopsByNeighborhood(FCell& Cell)
 {
 	FTimePoint StartTime = FChrono::Now();
-	//F3DDebugSession _(TEXT("ConnectCellLoopsByNeighborhood"));
 
 	int32 LoopCount = Cell.SubLoops.Num();
 
 	TArray<FPoint2D> LoopBarycenters;
-	// the four last points are the corners of Barycenter that defined the starting mesh for Bowyer & Watson algorithm
 	LoopBarycenters.Reserve(LoopCount + 4);
 
 	for (const TArray<FLoopNode*>& Nodes : Cell.SubLoops)
@@ -3274,7 +3290,7 @@ void FIsoTriangulator::ConnectCellSubLoopsByNeighborhood(FCell& Cell)
 	// Connect inner close loops 
 	// ==========================================================================================
 	{
-#ifdef DEBUG_DELAUNAY
+#ifdef DEBUG_CONNECT_CELL_SUB_LOOPS_BY_NEIGHBORHOOD
 		//F3DDebugSession _(TEXT("Build Segments Connect inner close loops "));
 #endif
 
@@ -3286,27 +3302,30 @@ void FIsoTriangulator::ConnectCellSubLoopsByNeighborhood(FCell& Cell)
 			const TArray<FLoopNode*>& SubLoopA = Cell.SubLoops[IndexLoopA];
 			const TArray<FLoopNode*>& SubLoopB = Cell.SubLoops[IndexLoopB];
 
+#ifdef DEBUG_CONNECT_CELL_SUB_LOOPS_BY_NEIGHBORHOOD
+			//F3DDebugSession _(TEXT("Segment"));
+#endif
 			TryToConnectTwoSubLoopsWithShortestSegment(Cell, SubLoopA, SubLoopB);
 		}
+
+		Cell.SelectSegmentInCandidateSegments(IsoSegmentFactory);
+#ifdef DEBUG_CONNECT_CELL_SUB_LOOPS_BY_NEIGHBORHOOD
+		Display(EGridSpace::UniformScaled, TEXT("Build Segments Connect inner close loops"), Cell.FinalSegments, false, false, EVisuProperty::YellowCurve);
+#endif
 	}
 
+
 	// With Outer loop
 	// ==========================================================================================
 	if (Cell.bHasOuterLoop && Cell.SubLoops.Num() > 1)
 	{
-		Cell.IntersectionTool.AddSegments(Cell.CandidateSegments);
-		Cell.IntersectionTool.Sort();
-
-#ifdef DEBUG_DELAUNAY
-		F3DDebugSession _(TEXT("Build Segments  With Outer loop"));
+#ifdef DEBUG_CONNECT_CELL_SUB_LOOPS_BY_NEIGHBORHOOD
+		//F3DDebugSession _(TEXT("Build Segments to Outer loop"));
 #endif
 		for (TArray<FLoopNode*>& SubLoopA : Cell.OuterLoopSubdivision)
 		{
 			for (int32 IndexBorderLoop : Cell.BorderLoopIndices)
 			{
-#ifdef DEBUG_DELAUNAY
-				F3DDebugSession _(TEXT("Segment"));
-#endif
 				int32 CandidateSegmentsCount = Cell.CandidateSegments.Num();
 				const TArray<FLoopNode*>& BorderLoop = Cell.SubLoops[IndexBorderLoop];
 
@@ -3320,9 +3339,13 @@ void FIsoTriangulator::ConnectCellSubLoopsByNeighborhood(FCell& Cell)
 				}
 			}
 		}
+
+		Cell.SelectSegmentInCandidateSegments(IsoSegmentFactory);
+#ifdef DEBUG_CONNECT_CELL_SUB_LOOPS_BY_NEIGHBORHOOD
+		Display(EGridSpace::UniformScaled, TEXT("Build Segments Connect inner to outer loop"), Cell.FinalSegments, false, false, EVisuProperty::YellowCurve);
+#endif
 	}
 
-	Cell.SelectSegmentInCandidateSegments(IsoSegmentFactory);
 
 #ifdef CADKERNEL_DEV
 	Chronos.FindSegmentToLinkLoopToLoopByDelaunayDuration += FChrono::Elapse(StartTime);
@@ -3362,28 +3385,28 @@ void FIsoTriangulator::TryToConnectTwoSubLoopsWithShortestSegment(FCell& Cell, c
 		FLoopNode* NodeB = LoopB[MinIndexB];
 		const FPoint2D& BCoordinates = NodeB->Get2DPoint(EGridSpace::UniformScaled, Grid);
 
-		//Display(EGridSpace::UniformScaled, *NodeA, *NodeB, 0, EVisuProperty::RedCurve);
+		//Display(EGridSpace::UniformScaled, *NodeA, *NodeB, 0, EVisuProperty::BlueCurve);
 		if (TryToCreateSegment(Cell, NodeA, ACoordinates, NodeB, BCoordinates, 0.1))
 		{
 			NodeA = &LoopA[MinIndexA]->GetNextNode();
-			NodeB = LoopB[MinIndexB];
-
-			//Display(EGridSpace::UniformScaled, *NodeA, *NodeB, 0, EVisuProperty::RedCurve);
-			TryToCreateSegment(Cell, NodeA, NodeA->Get2DPoint(EGridSpace::UniformScaled, Grid), NodeB, NodeB->Get2DPoint(EGridSpace::UniformScaled, Grid), 0.1);
+			if(Cell.Contains(NodeA))
+			{
+				NodeB = LoopB[MinIndexB];
+				TryToCreateSegment(Cell, NodeA, NodeA->Get2DPoint(EGridSpace::UniformScaled, Grid), NodeB, NodeB->Get2DPoint(EGridSpace::UniformScaled, Grid), 0.1);
+			}
 
-			NodeA = LoopA[MinIndexA];
 			NodeB = &LoopB[MinIndexB]->GetNextNode();
-
-			//Display(EGridSpace::UniformScaled, *NodeA, *NodeB, 0, EVisuProperty::RedCurve);
-			TryToCreateSegment(Cell, NodeA, NodeA->Get2DPoint(EGridSpace::UniformScaled, Grid), NodeB, NodeB->Get2DPoint(EGridSpace::UniformScaled, Grid), 0.1);
-			//Wait();
+			if (Cell.Contains(NodeB))
+			{
+				NodeA = LoopA[MinIndexA];
+				TryToCreateSegment(Cell, NodeA, NodeA->Get2DPoint(EGridSpace::UniformScaled, Grid), NodeB, NodeB->Get2DPoint(EGridSpace::UniformScaled, Grid), 0.1);
+			}
 		}
-
 	}
 
 };
 
-//#define DEBUG_FIND_BEST_TRIANGLE_TO_LINK_LOOPS
+//#define DEBUG_TRY_TO_CONNECT
 void FIsoTriangulator::TryToConnectTwoLoopsWithIsocelesTriangle(FCell& Cell, const TArray<FLoopNode*>& LoopA, const TArray<FLoopNode*>& LoopB)
 {
 
@@ -3508,7 +3531,7 @@ void FIsoTriangulator::TryToConnectTwoLoopsWithIsocelesTriangle(FCell& Cell, con
 		FIsoNode* Node = FindBestTriangle(Segment);
 		if (Node)
 		{
-#ifdef DEBUG_FIND_BEST_TRIANGLE_TO_LINK_LOOPS
+#ifdef DEBUG_TRY_TO_CONNECT
 			if (bDisplay)
 			{
 				F3DDebugSession _(FString::Printf(TEXT("Triangle")));
diff --git a/Engine/Source/Runtime/Datasmith/CADKernel/Private/CADKernel/Mesh/Meshers/IsoTriangulator/IntersectionSegmentTool.cpp b/Engine/Source/Runtime/Datasmith/CADKernel/Private/CADKernel/Mesh/Meshers/IsoTriangulator/IntersectionSegmentTool.cpp
--- a/Engine/Source/Runtime/Datasmith/CADKernel/Private/CADKernel/Mesh/Meshers/IsoTriangulator/IntersectionSegmentTool.cpp
+++ b/Engine/Source/Runtime/Datasmith/CADKernel/Private/CADKernel/Mesh/Meshers/IsoTriangulator/IntersectionSegmentTool.cpp
@@ -43,6 +43,7 @@ const FIsoSegment* FIntersectionSegmentTool::DoesIntersect(const FIsoNode& Start
 
 		if (IntersectSegments2D(Segment.Segment2D, StartEndSegment.Segment2D))
 		{
+			//DisplaySegment(Segment.Segment2D.Point0, Segment.Segment2D.Point1, 0, RedCurve);
 			return Segment.IsoSegment;
 		}
 	}
diff --git a/Engine/Source/Runtime/Datasmith/CADKernel/Private/CADKernel/Mesh/Meshers/IsoTriangulator/IsoCell.h b/Engine/Source/Runtime/Datasmith/CADKernel/Private/CADKernel/Mesh/Meshers/IsoTriangulator/IsoCell.h
--- a/Engine/Source/Runtime/Datasmith/CADKernel/Private/CADKernel/Mesh/Meshers/IsoTriangulator/IsoCell.h
+++ b/Engine/Source/Runtime/Datasmith/CADKernel/Private/CADKernel/Mesh/Meshers/IsoTriangulator/IsoCell.h
@@ -108,26 +108,48 @@ namespace CADKernel
 
 		void SelectSegmentInCandidateSegments(TFactory<FIsoSegment>& SegmentFactory)
 		{
+#ifdef DEBUG_SELECT_SEGMENT
+			F3DDebugSession _(TEXT("SelectSegmentInCandidateSegments "));
+			IntersectionTool.Display(TEXT("Cell.IntersectionTool at SelectSegmentInCandidateSegments start"));
+			//Wait();
+#endif
+
 			Algo::Sort(CandidateSegments, [&](const FIsoSegment* Segment1, const FIsoSegment* Segment2)
 				{
 					return Segment1->Get2DLengthSquare(EGridSpace::UniformScaled, Grid) < Segment2->Get2DLengthSquare(EGridSpace::UniformScaled, Grid);
 				});
 
+			//IntersectionTool.Empty(IntersectionTool.Count());
+
 			// Validate all candidate segments
 			for (FIsoSegment* Segment : CandidateSegments)
 			{
+#ifdef DEBUG_SELECT_SEGMENT
+				F3DDebugSession _(TEXT("Segment"));
+#endif
 				if (IntersectionTool.DoesIntersect(*Segment))
 				{
+#ifdef DEBUG_SELECT_SEGMENT
+					F3DDebugSession _(TEXT("SelectSegmentInCandidateSegments "));
+					DisplaySegment(Segment->GetFirstNode().GetPoint(EGridSpace::UniformScaled, Grid), Segment->GetSecondNode().GetPoint(EGridSpace::UniformScaled, Grid), 0, YellowCurve);
+#endif
 					SegmentFactory.DeleteEntity(Segment);
 					continue;
 				}
 
 				if (FIsoSegment::IsItAlreadyDefined(&Segment->GetFirstNode(), &Segment->GetSecondNode()))
 				{
+#ifdef DEBUG_SELECT_SEGMENT
+					DisplaySegment(Segment->GetFirstNode().GetPoint(EGridSpace::UniformScaled, Grid), Segment->GetSecondNode().GetPoint(EGridSpace::UniformScaled, Grid), 0, GreenCurve);
+#endif
 					SegmentFactory.DeleteEntity(Segment);
 					continue;
 				}
 
+#ifdef DEBUG_SELECT_SEGMENT
+				DisplaySegment(Segment->GetFirstNode().GetPoint(EGridSpace::UniformScaled, Grid), Segment->GetSecondNode().GetPoint(EGridSpace::UniformScaled, Grid), 0, BlueCurve);
+#endif
+
 				FinalSegments.Add(Segment);
 				IntersectionTool.AddSegment(*Segment);
 				Segment->SetSelected();
@@ -136,6 +158,14 @@ namespace CADKernel
 			CandidateSegments.Empty();
 		}
 
+		bool Contains(FLoopNode* NodeToFind)
+		{
+			int32 LoopIndex = NodeToFind->GetLoopIndex();
+			int32 Index = LoopIndexToIndex[LoopIndex];
+			return SubLoops[Index].Find(NodeToFind) != INDEX_NONE;
+		}
+
+
 	};
 
 } // namespace CADKernel
