Datasmith: Fix serialisation of the component flag on base actors

diff --git a/Engine/Source/Runtime/Datasmith/DatasmithCore/Private/DatasmithSceneXmlWriter.cpp b/Engine/Source/Runtime/Datasmith/DatasmithCore/Private/DatasmithSceneXmlWriter.cpp
--- a/Engine/Source/Runtime/Datasmith/DatasmithCore/Private/DatasmithSceneXmlWriter.cpp
+++ b/Engine/Source/Runtime/Datasmith/DatasmithCore/Private/DatasmithSceneXmlWriter.cpp
@@ -125,6 +125,11 @@ FString FDatasmithSceneXmlWriterImpl::GetLabelAndLayer(const TSharedPtr<IDatasmi
 		LabelAndLayer += TEXT(" visible=\"0\"");
 	}
 
+	if (ActorElement->IsAComponent())
+	{
+		LabelAndLayer += TEXT(" component=\"true\"");
+	}
+
 	return LabelAndLayer;
 }
 
@@ -666,13 +671,6 @@ void FDatasmithSceneXmlWriterImpl::WriteBeginOfMeshActorElement(const TSharedPtr
 
 	XmlString += GetLabelAndLayer(MeshActorElement);
 
-	if (MeshActorElement->IsAComponent())
-	{
-		FString ComponentText = MeshActorElement->IsAComponent() ? TEXT("true") : TEXT("false");
-
-		XmlString += TEXT(" component=\"") + ComponentText + TEXT("\"");
-	}
-
 	XmlString += FString(TEXT(">")) + LINE_TERMINATOR;
 
 	SerializeToArchive(Archive, XmlString);
