Fix Datasmith DirectLink crash when the scene doesn't have a name.

diff --git a/Engine/Source/Runtime/Datasmith/DatasmithCore/Private/DirectLink/DatasmithSceneReceiver.cpp b/Engine/Source/Runtime/Datasmith/DatasmithCore/Private/DirectLink/DatasmithSceneReceiver.cpp
--- a/Engine/Source/Runtime/Datasmith/DatasmithCore/Private/DirectLink/DatasmithSceneReceiver.cpp
+++ b/Engine/Source/Runtime/Datasmith/DatasmithCore/Private/DirectLink/DatasmithSceneReceiver.cpp
@@ -259,7 +259,7 @@ TUniquePtr<FDatasmithSceneReceiver::FSceneState> FDatasmithSceneReceiver::ParseS
 
 		if (EnumHasAllFlags(FoundUnsupportedTypes, EDatasmithElementType::Material))
 		{
-			UE_LOG(LogDatasmith, Warning, TEXT("Datasmith scene \"%s\" imported with DirectLink contains deprecated IDatasmithMaterialElement, they will be ignored."), *Scene.GetName());
+			UE_LOG(LogDatasmith, Warning, TEXT("Datasmith scene \"%s\" imported with DirectLink contains deprecated IDatasmithMaterialElement, they will be ignored."), Scene.GetName());
 			for (int MaterialIndex = Scene.GetMaterialsCount() - 1; 0 <= MaterialIndex; --MaterialIndex)
 			{
 				TSharedPtr<IDatasmithBaseMaterialElement> CurrentMaterial = Scene.GetMaterial(MaterialIndex);
