CADKernel update

diff --git a/Engine/Source/Runtime/Datasmith/CADKernel/Private/CADKernel/Core/Version.h b/Engine/Source/Runtime/Datasmith/CADKernel/Private/CADKernel/Core/Version.h
--- a/Engine/Source/Runtime/Datasmith/CADKernel/Private/CADKernel/Core/Version.h
+++ b/Engine/Source/Runtime/Datasmith/CADKernel/Private/CADKernel/Core/Version.h
@@ -2,8 +2,8 @@
 #pragma once
 
 //** ** Version du Toolkit** **//
-#define TOOLKIT_VERSION 0, 2, 18231529
-#define TOOLKIT_VERSION_ASCII "0.2.18231529"
+#define TOOLKIT_VERSION 0, 2, 18390356
+#define TOOLKIT_VERSION_ASCII "0.2.18390356"
 
 #define RELEASE_DATE	 __DATE__
 
diff --git a/Engine/Source/Runtime/Datasmith/CADKernel/Private/CADKernel/Math/SlopeUtils.h b/Engine/Source/Runtime/Datasmith/CADKernel/Private/CADKernel/Math/SlopeUtils.h
--- a/Engine/Source/Runtime/Datasmith/CADKernel/Private/CADKernel/Math/SlopeUtils.h
+++ b/Engine/Source/Runtime/Datasmith/CADKernel/Private/CADKernel/Math/SlopeUtils.h
@@ -167,6 +167,14 @@ namespace CADKernel
 		return FMath::Abs(ComputeOrientedSlope(StartPoint, EndPoint, ReferenceSlope));
 	}
 
+	/**
+	 * return a slop between [0, 4] i.e. an angle between [0, Pi]
+	 */
+	inline double ComputeUnorientedSlope(const FPoint2D& StartPoint, const FPoint2D& EndPoint1, const FPoint2D& EndPoint2)
+	{
+		return FMath::Abs(ComputeOrientedSlope(StartPoint, EndPoint1, EndPoint2));
+	}
+
 	/**
 	 *                         P1
 	 *          inside        /
diff --git a/Engine/Source/Runtime/Datasmith/CADKernel/Private/CADKernel/Mesh/Meshers/BowyerWatsonTriangulator.h b/Engine/Source/Runtime/Datasmith/CADKernel/Private/CADKernel/Mesh/Meshers/BowyerWatsonTriangulator.h
--- a/Engine/Source/Runtime/Datasmith/CADKernel/Private/CADKernel/Mesh/Meshers/BowyerWatsonTriangulator.h
+++ b/Engine/Source/Runtime/Datasmith/CADKernel/Private/CADKernel/Mesh/Meshers/BowyerWatsonTriangulator.h
@@ -9,317 +9,338 @@
 namespace CADKernel
 {
 
-	class FBowyerWatsonTriangulator
+class FBowyerWatsonTriangulator
+{
+protected:
+	struct FTriangle
 	{
-	protected:
-		struct FTriangle
-		{
-			int32 VertexIndices[3];
-			double SquareRadius;
-			FPoint Center;
+		int32 VertexIndices[3];
+		double SquareRadius;
+		FPoint Center;
 
-			FTriangle(const int32& Index0, const int32& Index1, const int32& Index2, const TArray<FPoint2D>& InVertices)
-			{
-				Set(Index0, Index1, Index2, InVertices);
-			}
+		FTriangle(const int32& Index0, const int32& Index1, const int32& Index2, const TArray<FPoint2D>& InVertices)
+		{
+			Set(Index0, Index1, Index2, InVertices);
+		}
 
-			void Set(const int32& Index0, const int32& Index1, const int32& Index2, const TArray<FPoint2D>& InVertices)
-			{
-				VertexIndices[0] = Index0;
-				VertexIndices[1] = Index1;
-				VertexIndices[2] = Index2;
-				FTriangle2D Triangle(InVertices[Index0], InVertices[Index1], InVertices[Index2]);
-				Center = Triangle.CircumCircleCenterWithSquareRadius(SquareRadius);
-			}
-		};
+		void Set(const int32& Index0, const int32& Index1, const int32& Index2, const TArray<FPoint2D>& InVertices)
+		{
+			VertexIndices[0] = Index0;
+			VertexIndices[1] = Index1;
+			VertexIndices[2] = Index2;
+			FTriangle2D Triangle(InVertices[Index0], InVertices[Index1], InVertices[Index2]);
+			Center = Triangle.CircumCircleCenterWithSquareRadius(SquareRadius);
+			// to avoid numerical error
+			SquareRadius += KINDA_SMALL_NUMBER;
+		}
+	};
 
 
-	public:
+public:
 
-		/**
-		 * @param Vertices the 2d point cloud to mesh
-		 * @param OutEdgeVertices, the edges of the mesh. An edge is defined by the indices of its vertices
-		 * So the ith edge is defined by the vertices EdgeVertexIndices[2 * ith] and EdgeVertexIndices[2 * ith + 1]
-		 */
-		FBowyerWatsonTriangulator(TArray<FPoint2D>& InVertices, TArray<int32>& OutEdgeVertices)
-			: VerticesCount(InVertices.Num())
-			, Vertices(InVertices)
-			, EdgeVertexIndices(OutEdgeVertices)
-		{
-			Init();
-		}
+	/**
+	 * @param Vertices the 2d point cloud to mesh
+	 * @param OutEdgeVertices, the edges of the mesh. An edge is defined by the indices of its vertices
+	 * So the ith edge is defined by the vertices EdgeVertexIndices[2 * ith] and EdgeVertexIndices[2 * ith + 1]
+	 */
+	FBowyerWatsonTriangulator(TArray<FPoint2D>& InVertices, TArray<int32>& OutEdgeVertices)
+		: VerticesCount(InVertices.Num())
+		, Vertices(InVertices)
+		, EdgeVertexIndices(OutEdgeVertices)
+	{
+		Init();
+	}
 
-		void Triangulate(bool bAddFirstVertex)
-		{
-			FTimePoint StartTime = FChrono::Now();
+	void Triangulate(bool bAddFirstVertex)
+	{
+		FTimePoint StartTime = FChrono::Now();
 #ifdef DEBUG_BOWYERWATSON
-			F3DDebugSession DelaunayDebugSession(TEXT("Delaunay Algo"));
+		F3DDebugSession DelaunayDebugSession(bDisplay, TEXT("Delaunay Algo"));
 #endif // DEBUG_DELAUNAY
 
-			// initialization of Bowyer & Watson algorithm with a bounding mesh of the vertex cloud 
-			// i.e. 2 triangles defined by the corners of the offsetted vertices bounding box
-			MakeBoundingMesh();
+		// initialization of Bowyer & Watson algorithm with a bounding mesh of the vertex cloud 
+		// i.e. 2 triangles defined by the corners of the offsetted vertices bounding box
+		MakeBoundingMesh();
 
 #ifdef DEBUG_BOWYERWATSON
-			F3DDebugSession _(TEXT("Vertices"));
-			for (int32 LoopIndex = bAddFirstVertex ? 1 : 0; LoopIndex < VerticesCount; ++LoopIndex)
+		if (bDisplay)
+		{
 			{
-				DisplayPoint(Vertices[LoopIndex], EVisuProperty::YellowPoint);
+				F3DDebugSession _(TEXT("Vertices"));
+				for (int32 LoopIndex = bAddFirstVertex ? 1 : 0; LoopIndex < VerticesCount; ++LoopIndex)
+				{
+					DisplayPoint(Vertices[LoopIndex], EVisuProperty::YellowPoint);
+				}
 			}
 			DisplayTriangles();
+			Wait();
+		}
 #endif
 
-			// insert each point in the mesh
-			for (int32 VertexIndex = bAddFirstVertex ? 1 : 0; VertexIndex < VerticesCount; ++VertexIndex)
-			{
-				TriangleIndices.Empty(VerticesCount);
+		// insert each point in the mesh
+		for (int32 VertexIndex = bAddFirstVertex ? 1 : 0; VertexIndex < VerticesCount; ++VertexIndex)
+		{
+			TriangleIndices.Empty(VerticesCount);
+
+			const FPoint2D& NewVertex = Vertices[VertexIndex];
 
-				const FPoint2D& NewVertex = Vertices[VertexIndex];
+#ifdef DEBUG_BOWYERWATSON
+			F3DDebugSession _(FString::Printf(TEXT("Step %d"), VertexIndex));
+			if(bDisplay)
+			{
+				F3DDebugSession _(TEXT("New Point"));
+				DisplayPoint(NewVertex, EVisuProperty::YellowPoint);
+			}
+#endif
 
-				// find all triangles whose circumcircles contain the new vertex
-				for (int32 TriangleIndex = 0; TriangleIndex < TriangleSet.Num(); TriangleIndex++)
+			// find all triangles whose circumcircles contain the new vertex
+			for (int32 TriangleIndex = 0; TriangleIndex < TriangleSet.Num(); TriangleIndex++)
+			{
+				const FPoint2D Center = TriangleSet[TriangleIndex].Center;
+				double SquareRadius = TriangleSet[TriangleIndex].SquareRadius;
+				double SquareDistanceToCenter = Center.SquareDistance(NewVertex);
+				if (SquareDistanceToCenter < SquareRadius)
 				{
-					const FPoint2D Center = TriangleSet[TriangleIndex].Center;
-					double SquareRadius = TriangleSet[TriangleIndex].SquareRadius;
-					double SquareDistanceToCenter = Center.SquareDistance(NewVertex);
-					if (SquareDistanceToCenter < SquareRadius)
-					{
-						TriangleIndices.Add(TriangleIndex);
-					}
+					TriangleIndices.Add(TriangleIndex);
 				}
+			}
 
-				// Find the boundary edges of the selected triangles:
-				// For all selected triangles, 
-				//    For each triangle edges
-				//       if the edge is not in EdgeVertexIndices: Add the edge i.e. add its vertex indices
-				//       else (the edge is in EdgeVertexIndices), remove the edge of EdgeVertexIndices
-				// As the triangles are oriented, the edge AB of a triangle is the edge BA of the adjacent triangle
-				EdgeVertexIndices.Empty(VerticesCount);
-				for (int32 TriangleIndex : TriangleIndices)
+			// Find the boundary edges of the selected triangles:
+			// For all selected triangles, 
+			//    For each triangle edges
+			//       if the edge is not in EdgeVertexIndices: Add the edge i.e. add its vertex indices
+			//       else (the edge is in EdgeVertexIndices), remove the edge of EdgeVertexIndices
+			// As the triangles are oriented, the edge AB of a triangle is the edge BA of the adjacent triangle
+			EdgeVertexIndices.Empty(VerticesCount);
+			for (int32 TriangleIndex : TriangleIndices)
+			{
+				int32 EndVertex = TriangleSet[TriangleIndex].VertexIndices[2];
+				for (int32 Index = 0; Index < 3; Index++)
 				{
-					int32 EndVertex = TriangleSet[TriangleIndex].VertexIndices[2];
-					for (int32 Index = 0; Index < 3; Index++)
+					int32 StartVertex = TriangleSet[TriangleIndex].VertexIndices[Index];
+					int32 Endex = 0;
+					// Does the edge exist
+					for (; Endex < EdgeVertexIndices.Num(); Endex += 2)
 					{
-						int32 StartVertex = TriangleSet[TriangleIndex].VertexIndices[Index];
-						int32 Endex = 0;
-						// Does the edge exist
-						for (; Endex < EdgeVertexIndices.Num(); Endex += 2)
+						if (EdgeVertexIndices[Endex] == EndVertex && EdgeVertexIndices[Endex + 1] == StartVertex)
 						{
-							if (EdgeVertexIndices[Endex] == EndVertex && EdgeVertexIndices[Endex + 1] == StartVertex)
-							{
-								EdgeVertexIndices[Endex] = -1;
-								EdgeVertexIndices[Endex + 1] = -1;
-								break;
-							}
-						}
-						if (Endex == EdgeVertexIndices.Num())
-						{   // No
-							EdgeVertexIndices.Add(StartVertex);
-							EdgeVertexIndices.Add(EndVertex);
+							EdgeVertexIndices[Endex] = -1;
+							EdgeVertexIndices[Endex + 1] = -1;
+							break;
 						}
-						EndVertex = StartVertex;
 					}
+					if (Endex == EdgeVertexIndices.Num())
+					{   // No
+						EdgeVertexIndices.Add(StartVertex);
+						EdgeVertexIndices.Add(EndVertex);
+					}
+					EndVertex = StartVertex;
 				}
+			}
 
-				// make the new triangles : Each new triangle is defined by an edge of the boundary and the new vertex 
-				int32 EdgeIndex = 0;
+			// make the new triangles : Each new triangle is defined by an edge of the boundary and the new vertex 
+			int32 EdgeIndex = 0;
 
-				// The deleted triangles are replaced by the new ones
-				for (int32 TriangleIndex : TriangleIndices)
+			// The deleted triangles are replaced by the new ones
+			for (int32 TriangleIndex : TriangleIndices)
+			{
+				while (EdgeVertexIndices[EdgeIndex] < 0)
 				{
-					while (EdgeVertexIndices[EdgeIndex] < 0)
-					{
-						EdgeIndex += 2;
-					}
-					TriangleSet[TriangleIndex].Set(EdgeVertexIndices[EdgeIndex + 1], EdgeVertexIndices[EdgeIndex], VertexIndex, Vertices);
 					EdgeIndex += 2;
 				}
-				// When all deleted triangles are replaced, the new triangles are added in the array
-				for (; EdgeIndex < EdgeVertexIndices.Num(); EdgeIndex += 2)
+				TriangleSet[TriangleIndex].Set(EdgeVertexIndices[EdgeIndex + 1], EdgeVertexIndices[EdgeIndex], VertexIndex, Vertices);
+				EdgeIndex += 2;
+			}
+			// When all deleted triangles are replaced, the new triangles are added in the array
+			for (; EdgeIndex < EdgeVertexIndices.Num(); EdgeIndex += 2)
+			{
+				if (EdgeVertexIndices[EdgeIndex] < 0)
 				{
-					if (EdgeVertexIndices[EdgeIndex] < 0)
-					{
-						continue;
-					}
-					TriangleSet.Emplace(EdgeVertexIndices[EdgeIndex + 1], EdgeVertexIndices[EdgeIndex], VertexIndex, Vertices);
+					continue;
 				}
-#ifdef DEBUG_DELAUNAY
-				//DisplayTriangles();
-#endif
+				TriangleSet.Emplace(EdgeVertexIndices[EdgeIndex + 1], EdgeVertexIndices[EdgeIndex], VertexIndex, Vertices);
 			}
-
-#ifdef DEBUG_DELAUNAY
-			// The final mesh
+#ifdef DEBUG_BOWYERWATSON
 			DisplayTriangles();
+			Wait(bDisplay);
+#endif
+		}
+
+#ifdef DEBUG_BOWYERWATSON
+		// The final mesh
+		DisplayTriangles();
+		Wait(bDisplay);
 #endif
 
-			// Find all Edges and their type (inner edge or boundary edge)
-			EdgeVertexIndices.Empty(TriangleSet.Num()*6);
+		// Find all Edges and their type (inner edge or boundary edge)
+		EdgeVertexIndices.Empty(TriangleSet.Num() * 6);
 
-			for (int32 TriangleIndex = 0; TriangleIndex < TriangleSet.Num(); TriangleIndex++)
+		for (int32 TriangleIndex = 0; TriangleIndex < TriangleSet.Num(); TriangleIndex++)
+		{
+			int32 Index = 0;
+			for (; Index < 3; Index++)
 			{
-				int32 Index = 0;
-				for (; Index < 3; Index++)
+				if (TriangleSet[TriangleIndex].VertexIndices[Index] >= VerticesCount)
 				{
-					if (TriangleSet[TriangleIndex].VertexIndices[Index] >= VerticesCount)
-					{
-						// one of the point is a corner of the bounding mesh
-						// At least, only one edge is added and this edge is necessarily an outer edge
-						break;
-					}
+					// one of the point is a corner of the bounding mesh
+					// At least, only one edge is added and this edge is necessarily an outer edge
+					break;
 				}
-				bool bIsOuter = Index < 3;
+			}
+			bool bIsOuter = Index < 3;
 
-				int32 EndVertex = TriangleSet[TriangleIndex].VertexIndices[2];
-				for (Index = 0; Index < 3; Index++)
+			int32 EndVertex = TriangleSet[TriangleIndex].VertexIndices[2];
+			for (Index = 0; Index < 3; Index++)
+			{
+				int32 StartVertex = TriangleSet[TriangleIndex].VertexIndices[Index];
+				if (StartVertex < VerticesCount && EndVertex < VerticesCount)
 				{
-					int32 StartVertex = TriangleSet[TriangleIndex].VertexIndices[Index];
-					if (StartVertex < VerticesCount && EndVertex < VerticesCount)
+					int32 Endex = 0;
+					for (; Endex < EdgeVertexIndices.Num(); Endex += 2)
 					{
-						int32 Endex = 0;
-						for (; Endex < EdgeVertexIndices.Num(); Endex += 2)
+						// Does the edge exist
+						if (EdgeVertexIndices[Endex] == EndVertex && EdgeVertexIndices[Endex + 1] == StartVertex)
 						{
-							// Does the edge exist
-							if (EdgeVertexIndices[Endex] == EndVertex && EdgeVertexIndices[Endex + 1] == StartVertex)
+							if (!bIsOuter)
 							{
-								if (!bIsOuter)
-								{
-									EdgeInstanceCount[Endex / 2]++;
-								}
-								break;
+								EdgeInstanceCount[Endex / 2]++;
 							}
+							break;
 						}
+					}
 
-						if (Endex == EdgeVertexIndices.Num())
-						{
-							// No
-							EdgeVertexIndices.Add(StartVertex);
-							EdgeVertexIndices.Add(EndVertex);
-							EdgeInstanceCount.Add(bIsOuter ? 0 : 1);
-						}
+					if (Endex == EdgeVertexIndices.Num())
+					{
+						// No
+						EdgeVertexIndices.Add(StartVertex);
+						EdgeVertexIndices.Add(EndVertex);
+						EdgeInstanceCount.Add(bIsOuter ? 0 : 1);
 					}
-					EndVertex = StartVertex;
 				}
+				EndVertex = StartVertex;
 			}
-#ifdef DEBUG_DELAUNAY
-			DisplayEdges();
-			//Wait();
+		}
+#ifdef DEBUG_BOWYERWATSON
+		DisplayEdges();
+		//Wait();
 #endif
 			// the bounding mesh vertices are removed
-			Vertices.SetNum(Vertices.Num() - 4);
-		}
+		Vertices.SetNum(Vertices.Num() - 4);
+	}
 
-		int32 OuterEdgeCount() const
+	int32 OuterEdgeCount() const
+	{
+		int32 EdgeCount = 0;
+		for (int32 Index = 0; Index < EdgeVertexIndices.Num() / 2; ++Index)
 		{
-			int32 EdgeCount = 0;
-			for (int32 Index = 0; Index < EdgeVertexIndices.Num() / 2; ++Index)
+			if (EdgeInstanceCount[Index] < 2)
 			{
-				if (EdgeInstanceCount[Index] < 2)
-				{
-					EdgeCount++;
-				}
+				EdgeCount++;
 			}
-			return EdgeCount;
 		}
+		return EdgeCount;
+	}
 
-		/**
-		 * Return the edge connected to 0 or 1 triangle
-		 */
-		void GetOuterEdges(TArray<int32>& OuterEdgeIndices) const
+	/**
+	 * Return the edge connected to 0 or 1 triangle
+	 */
+	void GetOuterEdges(TArray<int32>& OuterEdgeIndices) const
+	{
+		int32 EdgeCount = OuterEdgeCount();
+		OuterEdgeIndices.Reserve(EdgeCount);
+		for (int32 Index = 0, EdgeIndex = 0; Index < EdgeVertexIndices.Num(); ++EdgeIndex)
 		{
-			int32 EdgeCount = OuterEdgeCount();
-			OuterEdgeIndices.Reserve(EdgeCount);
-			for (int32 Index = 0, EdgeIndex = 0; Index < EdgeVertexIndices.Num(); ++EdgeIndex)
+			if (EdgeInstanceCount[EdgeIndex] < 2)
 			{
-				if (EdgeInstanceCount[EdgeIndex] < 2)
-				{
-					OuterEdgeIndices.Add(EdgeVertexIndices[Index++]);
-					OuterEdgeIndices.Add(EdgeVertexIndices[Index++]);
-				}
-				else
-				{
-					Index += 2;
-				}
+				OuterEdgeIndices.Add(EdgeVertexIndices[Index++]);
+				OuterEdgeIndices.Add(EdgeVertexIndices[Index++]);
+			}
+			else
+			{
+				Index += 2;
 			}
 		}
+	}
 
-		void GetOuterVertices(TSet<int32>& OuterVertexIndices) const
+	void GetOuterVertices(TSet<int32>& OuterVertexIndices) const
+	{
+		int32 EdgeCount = OuterEdgeCount();
+		OuterVertexIndices.Reserve(EdgeCount);
+		for (int32 Index = 0, EdgeIndex = 0; Index < EdgeVertexIndices.Num(); ++EdgeIndex)
 		{
-			int32 EdgeCount = OuterEdgeCount();
-			OuterVertexIndices.Reserve(EdgeCount);
-			for (int32 Index = 0, EdgeIndex = 0; Index < EdgeVertexIndices.Num(); ++EdgeIndex)
+			if (EdgeInstanceCount[EdgeIndex] < 2)
 			{
-				if (EdgeInstanceCount[EdgeIndex] < 2)
-				{
-					OuterVertexIndices.Add(EdgeVertexIndices[Index++]);
-					OuterVertexIndices.Add(EdgeVertexIndices[Index++]);
-				}
-				else
-				{
-					Index += 2;
-				}
+				OuterVertexIndices.Add(EdgeVertexIndices[Index++]);
+				OuterVertexIndices.Add(EdgeVertexIndices[Index++]);
+			}
+			else
+			{
+				Index += 2;
 			}
 		}
+	}
 
-		void GetMesh(TArray<int32>& Triangles)
+	void GetMesh(TArray<int32>& Triangles)
+	{
+		Triangles.Empty(TriangleSet.Num() * 3);
+		for (const FTriangle& Triangle : TriangleSet)
 		{
-			Triangles.Empty(TriangleSet.Num()*3);
-			for (const FTriangle& Triangle : TriangleSet)
-			{
-				Triangles.Append(Triangle.VertexIndices, 3);
-			}
+			Triangles.Append(Triangle.VertexIndices, 3);
 		}
+	}
 
-private:
-		int32 VerticesCount;
+#ifdef DEBUG_BOWYERWATSON
+	static bool bDisplay;
+#endif
 
-		TArray<FPoint2D>& Vertices;
+private:
+	int32 VerticesCount;
 
-		/**
-		 * An edge is defined by the indices of its vertices
-		 * So the ith edge is defined by the vertices EdgeVertexIndices[2 * ith] and EdgeVertexIndices[2 * ith + 1]
-		 */
-		TArray<int32>& EdgeVertexIndices;
+	TArray<FPoint2D>& Vertices;
 
-		TArray<FTriangle> TriangleSet;
+	/**
+	 * An edge is defined by the indices of its vertices
+	 * So the ith edge is defined by the vertices EdgeVertexIndices[2 * ith] and EdgeVertexIndices[2 * ith + 1]
+	 */
+	TArray<int32>& EdgeVertexIndices;
 
-		// It's use to mark all triangles whose circumcircles contain the next vertex
-		TArray<int32> TriangleIndices;
+	TArray<FTriangle> TriangleSet;
 
-		bool bDisplay = false;
+	// It's use to mark all triangles whose circumcircles contain the next vertex
+	TArray<int32> TriangleIndices;
 
-		/**
-		 * Use to determine if the edge is a border edge of inner edge
-		 * If EdgeInstanceCount[ith] = 2, the edge is a inner edge
-		 */
-		TArray<int32> EdgeInstanceCount;
+	/**
+	 * Use to determine if the edge is a border edge of inner edge
+	 * If EdgeInstanceCount[ith] = 2, the edge is a inner edge
+	 */
+	TArray<int32> EdgeInstanceCount;
 
-		void MakeBoundingMesh()
-		{
-			FAABB2D VerticesBBox;
-			VerticesBBox += Vertices;
-			VerticesBBox.Offset(VerticesBBox.DiagonalLength());
+	void MakeBoundingMesh()
+	{
+		FAABB2D VerticesBBox;
+		VerticesBBox += Vertices;
+		VerticesBBox.Offset(VerticesBBox.DiagonalLength());
 
-			Vertices.Emplace(VerticesBBox.GetCorner(3));
-			Vertices.Emplace(VerticesBBox.GetCorner(2));
-			Vertices.Emplace(VerticesBBox.GetCorner(0));
-			Vertices.Emplace(VerticesBBox.GetCorner(1));
+		Vertices.Emplace(VerticesBBox.GetCorner(3));
+		Vertices.Emplace(VerticesBBox.GetCorner(2));
+		Vertices.Emplace(VerticesBBox.GetCorner(0));
+		Vertices.Emplace(VerticesBBox.GetCorner(1));
 
-			TriangleSet.Emplace(VerticesCount, VerticesCount + 1, VerticesCount + 2, Vertices);
-			TriangleSet.Emplace(VerticesCount + 2, VerticesCount + 3, VerticesCount, Vertices);
-		}
+		TriangleSet.Emplace(VerticesCount, VerticesCount + 1, VerticesCount + 2, Vertices);
+		TriangleSet.Emplace(VerticesCount + 2, VerticesCount + 3, VerticesCount, Vertices);
+	}
 
-		void Init()
-		{
-			VerticesCount = Vertices.Num();
-			Vertices.Reserve(VerticesCount + 4);
-			TriangleSet.Reserve(VerticesCount);
-			TriangleIndices.Reserve(VerticesCount);
-			EdgeVertexIndices.Reserve(4 * VerticesCount);
-			EdgeInstanceCount.Reserve(2 * VerticesCount);
-		}
+	void Init()
+	{
+		VerticesCount = Vertices.Num();
+		Vertices.Reserve(VerticesCount + 4);
+		TriangleSet.Reserve(VerticesCount);
+		TriangleIndices.Reserve(VerticesCount);
+		EdgeVertexIndices.Reserve(4 * VerticesCount);
+		EdgeInstanceCount.Reserve(2 * VerticesCount);
+	}
 
-#ifdef DEBUG_DELAUNAY
+#ifdef DEBUG_BOWYERWATSON
 	void DisplayEdges()
 	{
 		if (!bDisplay)
@@ -336,7 +357,7 @@ private:
 			}
 			else
 			{
-				DisplaySegment(Vertices[EdgeVertexIndices[Index]], Vertices[EdgeVertexIndices[Index + 1]]);
+				DisplaySegment(Vertices[EdgeVertexIndices[Index]], Vertices[EdgeVertexIndices[Index + 1]], 0, EVisuProperty::PurpleCurve);
 			}
 		}
 	};
@@ -351,14 +372,41 @@ private:
 		F3DDebugSession _(TEXT("Triangles"));
 		for (int32 Index = 0; Index < TriangleSet.Num(); Index++)
 		{
-			DisplaySegment(Vertices[TriangleSet[Index].VertexIndices[0]], Vertices[TriangleSet[Index].VertexIndices[1]]);
-			DisplaySegment(Vertices[TriangleSet[Index].VertexIndices[1]], Vertices[TriangleSet[Index].VertexIndices[2]]);
-			DisplaySegment(Vertices[TriangleSet[Index].VertexIndices[2]], Vertices[TriangleSet[Index].VertexIndices[0]]);
+			DisplayTriangle(Index, EVisuProperty::YellowCurve);
 		}
 		//Wait();
 	};
-#endif
 
+	void DisplaySelectedTriangles()
+	{
+		if (!bDisplay)
+		{
+			return;
+		}
+
+		F3DDebugSession _(TEXT("Selected Triangles"));
+		for (int32 Index : TriangleIndices)
+		{
+			DisplayTriangle(Index , EVisuProperty::BlueCurve);
+		}
+		//Wait();
+	};
+
+	void DisplayTriangle(int32 Index, EVisuProperty Property)
+	{
+		if (!bDisplay)
+		{
+			return;
+		}
+
+		DisplaySegment(Vertices[TriangleSet[Index].VertexIndices[0]], Vertices[TriangleSet[Index].VertexIndices[1]], 0, EVisuProperty::BlueCurve);
+		DisplaySegment(Vertices[TriangleSet[Index].VertexIndices[1]], Vertices[TriangleSet[Index].VertexIndices[2]], 0, EVisuProperty::BlueCurve);
+		DisplaySegment(Vertices[TriangleSet[Index].VertexIndices[2]], Vertices[TriangleSet[Index].VertexIndices[0]], 0, EVisuProperty::BlueCurve);
+		//Wait();
 	};
 
+#endif
+
+};
+
 }
\ No newline at end of file
diff --git a/Engine/Source/Runtime/Datasmith/CADKernel/Private/CADKernel/Mesh/Meshers/IsoTriangulator.cpp b/Engine/Source/Runtime/Datasmith/CADKernel/Private/CADKernel/Mesh/Meshers/IsoTriangulator.cpp
--- a/Engine/Source/Runtime/Datasmith/CADKernel/Private/CADKernel/Mesh/Meshers/IsoTriangulator.cpp
+++ b/Engine/Source/Runtime/Datasmith/CADKernel/Private/CADKernel/Mesh/Meshers/IsoTriangulator.cpp
@@ -21,8 +21,7 @@
 namespace CADKernel
 {
 
-//#define ADD_TRIANGLE_2D
-//#define DEBUG_ISOTRIANGULATOR
+const double FIsoTriangulator::GeometricToMeshingToleranceFactor = 10.;
 
 FIsoTriangulator::FIsoTriangulator(FGrid& InGrid, TSharedRef<FFaceMesh> EntityMesh)
 	: Grid(InGrid)
@@ -31,17 +30,28 @@ FIsoTriangulator::FIsoTriangulator(FGrid& InGrid, TSharedRef<FFaceMesh> EntityMe
 	, InnerSegmentsIntersectionTool(InGrid)
 	, InnerToLoopSegmentsIntersectionTool(InGrid)
 	, InnerToOuterSegmentsIntersectionTool(InGrid)
+	, GeometricTolerance(Grid.GetFace()->GetCarrierSurface()->Get3DTolerance())
+	, SquareGeometricTolerance(FMath::Square(GeometricTolerance))
+	, SquareGeometricTolerance2(2. * SquareGeometricTolerance)
+	, MeshingTolerance(GeometricTolerance* GeometricToMeshingToleranceFactor)
+	, SquareMeshingTolerance(FMath::Square(MeshingTolerance))
 {
 	FinalInnerSegments.Reserve(3 * Grid.InnerNodesCount());
 	IndexOfLowerLeftInnerNodeSurroundingALoop.Reserve(Grid.GetLoopCount());
 
 #ifdef DEBUG_ISOTRIANGULATOR
+
 #ifdef DEBUG_ONLY_SURFACE_TO_DEBUG
 	if (InGrid.GetFace()->GetId() == FaceToDebug)
 #endif
 	{
 		bDisplay = true;
 	}
+
+#ifdef DEBUG_BOWYERWATSON
+bool FBowyerWatsonTriangulator::bDisplay = false;
+#endif
+
 #endif
 }
 
@@ -68,22 +78,21 @@ bool FIsoTriangulator::Triangulate()
 	BuildLoopSegments();
 
 #ifdef DEBUG_ISOTRIANGULATOR
-	Display(DisplaySpace, TEXT("FIsoTrianguler::LoopSegments"), LoopSegments, false, false, EVisuProperty::OrangeCurve);
-#endif
-
-	CheckLoopInSelfIntersectingAndFixIt();
-
-#ifdef DEBUG_ISOTRIANGULATOR
-	Display(DisplaySpace, TEXT("FIsoTrianguler::LoopSegments"), LoopSegments, true, true, EVisuProperty::OrangeCurve);
-	//Wait();
+	Display(DisplaySpace, TEXT("FIsoTrianguler::LoopSegments"), LoopSegments, false, false, EVisuProperty::GreenCurve);
 #endif
 
+	if (!FindLoopIntersectionAndFixIt())
+	{
+		return false;
+	}
 
 #ifdef DEBUG_ISOTRIANGULATOR
-	Display(DisplaySpace, TEXT("FIsoTrianguler::LoopSegments"), LoopSegments, false, false, EVisuProperty::OrangeCurve);
+	Display(DisplaySpace, TEXT("FIsoTrianguler::LoopSegments after fix intersection"), LoopSegments, true, false, EVisuProperty::OrangeCurve);
+	Wait(false);
 #endif
 
 	BuildThinZoneSegments();
+
 #ifdef DEBUG_ISOTRIANGULATOR
 	Display(DisplaySpace, TEXT("FIsoTrianguler::ThinZoneSegments"), ThinZoneSegments, false);
 	LoopSegmentsIntersectionTool.Display(TEXT("FIsoTrianguler::IntersectionTool Loop"));
@@ -117,6 +126,7 @@ bool FIsoTriangulator::Triangulate()
 	// =============================================================================================================
 #ifdef DEBUG_ISOTRIANGULATOR
 	Display(DisplaySpace, TEXT("FIsoTrianguler::Final To Loops Before"), FinalToLoops, false, false, EVisuProperty::YellowCurve);
+	//Wait(bDisplay);
 #endif
 	ConnectCellLoops();
 #ifdef DEBUG_ISOTRIANGULATOR
@@ -351,6 +361,11 @@ void FIsoTriangulator::BuildLoopSegments()
 	int32 LoopIndex = 0;
 	for (FLoopNode& Node : LoopNodes)
 	{
+		if (Node.IsDelete())
+		{
+			continue;
+		}
+
 		FIsoSegment& Segment = IsoSegmentFactory.New();
 		Segment.Init(Node, Node.GetNextNode(), ESegmentType::Loop);
 		Segment.ConnectToNode();
@@ -384,6 +399,11 @@ void FIsoTriangulator::BuildLoopSegments()
 
 	for (FLoopNode& Node : LoopNodes)
 	{
+		if (Node.IsDelete())
+		{
+			continue;
+		}
+
 		if (Node.GetConnectedSegments()[0]->IsIsoU() && Node.GetConnectedSegments()[1]->IsIsoU())
 		{
 			Node.SetAsIsoU();
@@ -394,9 +414,6 @@ void FIsoTriangulator::BuildLoopSegments()
 		}
 	}
 
-	LoopSegmentsIntersectionTool.AddSegments(LoopSegments);
-	LoopSegmentsIntersectionTool.Sort();
-
 #ifdef CADKERNEL_DEV
 	Chronos.BuildLoopSegmentsDuration += FChrono::Elapse(StartTime);
 #endif
@@ -409,6 +426,11 @@ void FIsoTriangulator::BuildThinZoneSegments()
 	TMap<int32, FLoopNode*> IndexToNode;
 	for (FLoopNode& Node : LoopNodes)
 	{
+		if (Node.IsDelete())
+		{
+			continue;
+		}
+
 		IndexToNode.Add(Node.GetFaceIndex(), &Node);
 	}
 
@@ -474,8 +496,6 @@ void FIsoTriangulator::BuildInnerSegments()
 
 	FTimePoint StartTime = FChrono::Now();
 
-	const double GeometricTolerance = Grid.GetFace()->GetCarrierSurface()->Get3DTolerance();
-
 	// Build segments according to the Grid following u then following v
 	// Build segment must not be in intersection with the loop
 	int32 NumU = Grid.GetCuttingCount(EIso::IsoU);
@@ -1199,7 +1219,7 @@ void FIsoTriangulator::ConnectCellLoops()
 	}
 }
 
-// find cell containing boundary nodes
+//#define DEBUG_BUILD_CELLS
 void FIsoTriangulator::FindCellContainingBoundaryNodes(TArray<FCell>& Cells)
 {
 	FTimePoint StartTime = FChrono::Now();
@@ -1209,6 +1229,7 @@ void FIsoTriangulator::FindCellContainingBoundaryNodes(TArray<FCell>& Cells)
 
 	const int32 CountU = Grid.GetCuttingCount(EIso::IsoU);
 	const int32 CountV = Grid.GetCuttingCount(EIso::IsoV);
+	const int32 MaxUV = Grid.GetTotalCuttingCount();
 
 	const TArray<double>& IsoUCoordinates = Grid.GetUniformCuttingCoordinatesAlongIso(EIso::IsoU);
 	const TArray<double>& IsoVCoordinates = Grid.GetUniformCuttingCoordinatesAlongIso(EIso::IsoV);
@@ -1218,22 +1239,54 @@ void FIsoTriangulator::FindCellContainingBoundaryNodes(TArray<FCell>& Cells)
 		int32 IndexU = 0;
 		int32 IndexV = 0;
 		int32 Index = 0;
+		int32 DeletedNodeCount = 0;
+
+#ifdef DEBUG_BUILD_CELLS
+		F3DDebugSession _(bDisplay, TEXT("Debug Find index"));
+#endif
+
 		for (const FLoopNode& LoopPoint : LoopNodes)
 		{
-			const FPoint2D& Coordinate = LoopPoint.Get2DPoint(EGridSpace::UniformScaled, Grid);
+			if (!LoopPoint.IsDelete())
+			{
+				const FPoint2D& Coordinate = LoopPoint.Get2DPoint(EGridSpace::UniformScaled, Grid);
+				ArrayUtils::FindCoordinateIndex(IsoUCoordinates, Coordinate.U, IndexU);
+				ArrayUtils::FindCoordinateIndex(IsoVCoordinates, Coordinate.V, IndexV);
 
-			ArrayUtils::FindCoordinateIndex(IsoUCoordinates, Coordinate.U, IndexU);
-			ArrayUtils::FindCoordinateIndex(IsoVCoordinates, Coordinate.V, IndexV);
+#ifdef DEBUG_BUILD_CELLS
+				if (bDisplay)
+				{
+					F3DDebugSession S(FString::Printf(TEXT("Index %d"), IndexV * CountU + IndexU));
+					DisplayPoint(LoopPoint.GetPoint(EGridSpace::UniformScaled, Grid), EVisuProperty::BluePoint);
+					DrawCellBoundary(IndexV * CountU + IndexU, EVisuProperty::YellowPoint);
+				}
+#endif
 
-			NodeToCellIndices.Emplace(IndexV * CountU + IndexU);
+				NodeToCellIndices.Emplace(IndexV * CountU + IndexU);
+			}
+			else
+			{
+#ifdef DEBUG_BUILD_CELLS
+				if (bDisplay)
+				{
+					F3DDebugSession S(TEXT("Deleted Node"));
+					DisplayPoint(LoopPoint.GetPoint(EGridSpace::UniformScaled, Grid), EVisuProperty::RedPoint);
+				}
+#endif
+
+				DeletedNodeCount++;
+				NodeToCellIndices.Emplace(MaxUV);
+			}
 			SortedIndex.Emplace(Index++);
 		}
-	}
 
-	Algo::Sort(SortedIndex, [&](const int32& Index1, const int32& Index2)
-		{
-			return (NodeToCellIndices[Index1] < NodeToCellIndices[Index2]);
-		});
+		Algo::Sort(SortedIndex, [&](const int32& Index1, const int32& Index2)
+			{
+				return (NodeToCellIndices[Index1] < NodeToCellIndices[Index2]);
+			});
+
+		SortedIndex.SetNum(SortedIndex.Num() - DeletedNodeCount);
+	}
 
 	int32 CountOfCellsFilled = 1;
 	{
@@ -1254,6 +1307,11 @@ void FIsoTriangulator::FindCellContainingBoundaryNodes(TArray<FCell>& Cells)
 		int32 CellIndex = NodeToCellIndices[SortedIndex[0]];
 		TArray<FLoopNode*> CellNodes;
 		CellNodes.Reserve(LoopNodeCount);
+
+#ifdef DEBUG_BUILD_CELLS
+		F3DDebugSession _(bDisplay, TEXT("Build Cells"));
+#endif
+
 		for (int32 Index : SortedIndex)
 		{
 			if (CellIndex != NodeToCellIndices[Index])
@@ -1262,12 +1320,24 @@ void FIsoTriangulator::FindCellContainingBoundaryNodes(TArray<FCell>& Cells)
 
 				CellIndex = NodeToCellIndices[Index];
 				CellNodes.Empty(LoopNodeCount);
+
+#ifdef DEBUG_BUILD_CELLS
+				DisplayCell(Cells.Last());
+#endif
 			}
 
 			FLoopNode& LoopNode = LoopNodes[Index];
-			CellNodes.Add(&LoopNode);
+			if (!LoopNode.IsDelete())
+			{
+				CellNodes.Add(&LoopNode);
+			}
 		}
 		Cells.Emplace(CellIndex, CellNodes, Grid);
+
+#ifdef DEBUG_BUILD_CELLS
+		DisplayCell(Cells.Last());
+		Wait(bDisplay);
+#endif
 	}
 	FChrono::Elapse(StartTime);
 }
@@ -1309,9 +1379,14 @@ void FIsoTriangulator::FindCandidateSegmentsToLinkInnerAndLoop()
 		int32 IndexU = 0;
 		int32 IndexV = 0;
 		int32 Index = 0;
-		for (const FLoopNode& LoopPoint : LoopNodes)
+		for (const FLoopNode& LoopNode : LoopNodes)
 		{
-			const FPoint2D& Coordinate = LoopPoint.Get2DPoint(EGridSpace::UniformScaled, Grid);
+			if (LoopNode.IsDelete())
+			{
+				continue;
+			}
+
+			const FPoint2D& Coordinate = LoopNode.Get2DPoint(EGridSpace::UniformScaled, Grid);
 
 			ArrayUtils::FindCoordinateIndex(IsoUCoordinates, Coordinate.U, IndexU);
 			ArrayUtils::FindCoordinateIndex(IsoVCoordinates, Coordinate.V, IndexV);
@@ -1713,7 +1788,7 @@ void FIsoTriangulator::SelectSegmentInCandidateSegments()
 			return (Segment1->Get3DLengthSquare(Grid)) < (Segment2->Get3DLengthSquare(Grid));
 		});
 
-	InnerSegmentsIntersectionTool.SetNum(InnerSegmentsIntersectionTool.Count() + (int32)CandidateSegments.Num());
+	InnerSegmentsIntersectionTool.Reserve(InnerSegmentsIntersectionTool.Count() + (int32)CandidateSegments.Num());
 
 	FinalToLoops.Reserve(CandidateSegments.Num());
 
@@ -1802,6 +1877,11 @@ void FIsoTriangulator::ConnectUnconnectedInnerSegments()
 		FLoopNode* Candidate = nullptr;
 		for (FLoopNode& LoopNode : LoopNodes)
 		{
+			if (LoopNode.IsDelete())
+			{
+				continue;
+			}
+
 			double Distance = LoopNode.Get2DPoint(EGridSpace::Scaled, Grid).SquareDistance(Node->Get2DPoint(EGridSpace::Scaled, Grid));
 			if (Distance < MinDistance)
 			{
@@ -1825,330 +1905,15 @@ void FIsoTriangulator::ConnectUnconnectedInnerSegments()
 	}
 }
 
-//#define DEBUG_CLOSED_OUSIDE_POINT
-bool FIsoTriangulator::TryToRemoveIntersectionByMovingTheClosedOusidePoint(const FIsoSegment* Segment0, const FIsoSegment* Segment1)
-{
-	const FIsoNode* Nodes[2][2] = { {nullptr, nullptr}, {nullptr, nullptr} };
-	Nodes[0][0] = &Segment0->GetFirstNode();
-	Nodes[0][1] = &Segment0->GetSecondNode();
-	Nodes[1][0] = &Segment1->GetFirstNode();
-	Nodes[1][1] = &Segment1->GetSecondNode();
-
-	FPoint2D Points[2][2];
-	Points[0][0] = Nodes[0][0]->Get2DPoint(EGridSpace::UniformScaled, Grid);
-	Points[0][1] = Nodes[0][1]->Get2DPoint(EGridSpace::UniformScaled, Grid);
-	Points[1][0] = Nodes[1][0]->Get2DPoint(EGridSpace::UniformScaled, Grid);
-	Points[1][1] = Nodes[1][1]->Get2DPoint(EGridSpace::UniformScaled, Grid);
-
-	int32 ProjectedNodeIndex[2] = { 0, 0 };
-	FPoint2D ProjectedPoints[2];
-	double Distance[2] = { HUGE_VALUE, HUGE_VALUE };
-	int32 OtherSegmentIndex = 0;
-	int32 SegmentIndex = 1;
-
-	TFunction<void(int32)> MoveNode = [&](int32 MovePointIndex)
-	{
-		FPoint2D MoveDirection = ProjectedPoints[MovePointIndex] - Points[MovePointIndex][ProjectedNodeIndex[MovePointIndex]];
-		MoveDirection.Normalize();
-		MoveDirection *= Grid.GetFace()->GetCarrierSurface()->Get3DTolerance();
-		FPoint2D NewCoordinate = ProjectedPoints[MovePointIndex] + MoveDirection;
-
-		const_cast<FIsoNode&>((*Nodes[MovePointIndex][ProjectedNodeIndex[MovePointIndex]])).Set2DPoint(EGridSpace::UniformScaled, Grid, NewCoordinate);
-
-#ifdef DEBUG_CLOSED_OUSIDE_POINT
-		if (bDisplay)
-		{
-			F3DDebugSession _(TEXT("New Segs"));
-			DisplayPoint(NewCoordinate, EVisuProperty::RedPoint);
-			Display(EGridSpace::UniformScaled, *Segment1, 0, EVisuProperty::RedCurve);
-			Display(EGridSpace::UniformScaled, *Segment0, 0, EVisuProperty::BlueCurve);
-			//Wait();
-		}
-#endif
-	};
-
-	TFunction<void(int32)> ProjectInnerPoint = [&](int32 OtherSegmentNodeIndex)
-	{
-		double Coordinate;
-		ProjectedPoints[OtherSegmentIndex] = ProjectPointOnSegment(Points[OtherSegmentIndex][OtherSegmentNodeIndex], Points[SegmentIndex][0], Points[SegmentIndex][1], Coordinate, false);
-		bool bProjectedPoint1IsInside = (Coordinate >= -SMALL_NUMBER && Coordinate <= 1 + SMALL_NUMBER);
-		if (bProjectedPoint1IsInside)
-		{
-			Distance[OtherSegmentIndex] = ProjectedPoints[OtherSegmentIndex].Distance(Points[OtherSegmentIndex][OtherSegmentNodeIndex]);
-		}
-
-#ifdef DEBUG_CLOSED_OUSIDE_POINT
-		if (bDisplay)
-		{
-			{
-				F3DDebugSession _(TEXT("Segs"));
-				DisplaySegment(Points[SegmentIndex][0], Points[SegmentIndex][1], 0, EVisuProperty::BlueCurve);
-				DisplayPoint(Points[OtherSegmentIndex][OtherSegmentNodeIndex], EVisuProperty::RedPoint);
-				DisplayPoint(ProjectedPoints[OtherSegmentIndex], EVisuProperty::RedPoint);
-			}
-		}
-#endif
-	};
-
-	TFunction<void()> FindInnerNodeAndProjectIt = [&]()
-	{
-		OtherSegmentIndex = SegmentIndex ? 0 : 1;
-
-		double OrientedSlop = ComputeOrientedSlope(Points[SegmentIndex][0], Points[SegmentIndex][1], Points[OtherSegmentIndex][0]);
-		ProjectedNodeIndex[OtherSegmentIndex] = (OrientedSlop < 0) ? 0 : 1;
-		ProjectInnerPoint(ProjectedNodeIndex[OtherSegmentIndex]);
-	};
-
-	SegmentIndex = 0;
-	FindInnerNodeAndProjectIt();
-	SegmentIndex = 1;
-	FindInnerNodeAndProjectIt();
-
-	if (Distance[0] == HUGE_VALUE && Distance[1] == HUGE_VALUE)
-	{
-		return false;
-	}
-
-	int32 MovePoint = (Distance[0] < Distance[1]) ? 0 : 1;
-	MoveNode(MovePoint);
-	return true;
-}
-
-//#define DEBUG_TWO_CONSECUTIVE_INTERSECTING
-bool FIsoTriangulator::TryToRemoveIntersectionOfTwoConsecutiveIntersectingSegments(const FIsoSegment* IntersectingSegment, FIsoSegment* Segment)
-{
-	FIsoSegment* NextSegment = nullptr;
-	FLoopNode* Node = nullptr;
-
-	double OrientedSlop = ComputeOrientedSlope(IntersectingSegment->GetFirstNode().Get2DPoint(EGridSpace::UniformScaled, Grid), IntersectingSegment->GetSecondNode().Get2DPoint(EGridSpace::UniformScaled, Grid), Segment->GetFirstNode().Get2DPoint(EGridSpace::UniformScaled, Grid));
-	if (OrientedSlop >= 0)
-	{
-		Node = (FLoopNode*)&Segment->GetSecondNode();
-		NextSegment = Node->GetSegmentConnectedTo(&Node->GetNextNode());
-	}
-	else
-	{
-		Node = (FLoopNode*)&Segment->GetFirstNode();
-		NextSegment = Node->GetSegmentConnectedTo(&Node->GetPreviousNode());
-	}
-
-#ifdef DEBUG_TWO_CONSECUTIVE_INTERSECTING
-	if (bDisplay)
-	{
-		F3DDebugSession _(TEXT("Intersecting Segments"));
-		Display(EGridSpace::UniformScaled, *Segment, 0, EVisuProperty::BlueCurve);
-		Display(EGridSpace::UniformScaled, *NextSegment, 0, EVisuProperty::BlueCurve);
-		Display(EGridSpace::UniformScaled, *IntersectingSegment, 0, EVisuProperty::RedCurve);
-		Display(EGridSpace::UniformScaled, *Node, 0, EVisuProperty::RedPoint);
-	}
-#endif
-
-	TSegment<FPoint2D> IntersectingSegment2D(IntersectingSegment->GetFirstNode().Get2DPoint(EGridSpace::UniformScaled, Grid), IntersectingSegment->GetSecondNode().Get2DPoint(EGridSpace::UniformScaled, Grid));
-	TSegment<FPoint2D> NextSegment2D(NextSegment->GetFirstNode().Get2DPoint(EGridSpace::UniformScaled, Grid), NextSegment->GetSecondNode().Get2DPoint(EGridSpace::UniformScaled, Grid));
-	if (!FastIntersectSegments2D(NextSegment2D, IntersectingSegment2D))
-	{
-		return false;
-	}
-
-	TSegment<FPoint2D> Segment2D(Segment->GetFirstNode().Get2DPoint(EGridSpace::UniformScaled, Grid), Segment->GetSecondNode().Get2DPoint(EGridSpace::UniformScaled, Grid));
-
-	FPoint2D Intersection1 = FindIntersectionOfSegments2D(Segment2D, IntersectingSegment2D);
-	FPoint2D Intersection2 = FindIntersectionOfSegments2D(NextSegment2D, IntersectingSegment2D);
-
-
-	FPoint2D MiddleIntersection = Intersection1.Middle(Intersection2);
-
-	FPoint2D MoveDirection = IntersectingSegment2D.Point1 - IntersectingSegment2D.Point0;
-	MoveDirection = MoveDirection.Normalize().GetPerpendicularVector();
-	MoveDirection *= Grid.GetFace()->GetCarrierSurface()->Get3DTolerance();;
-
-	FPoint2D NewCoordinate = MiddleIntersection + MoveDirection;
-	Node->Set2DPoint(EGridSpace::UniformScaled, Grid, NewCoordinate);
-
-#ifdef DEBUG_TWO_CONSECUTIVE_INTERSECTING
-	if (bDisplay)
-	{
-		{
-			F3DDebugSession _(TEXT("New position"));
-			DisplayPoint(Intersection1, EVisuProperty::RedPoint);
-			DisplayPoint(Intersection2, EVisuProperty::RedPoint);
-			Display(EGridSpace::UniformScaled, *Node, EVisuProperty::RedPoint);
-			Display(EGridSpace::UniformScaled, *Segment, 0, EVisuProperty::GreenCurve/*, true*/);
-			Display(EGridSpace::UniformScaled, *NextSegment, 0, EVisuProperty::GreenCurve/*, true*/);
-		}
-	}
-#endif
-	return true;
-}
-
-bool FIsoTriangulator::TryToRemoveIntersectionBySwappingSegments(FIsoSegment* Segment0, FIsoSegment* Segment1)
-{
-	auto SwapLoopOrientation = [&](int32 FirstNodeIndex, int32 LastNodeIndex)
-	{
-		TArray<FIsoSegment*> Segments;
-		Segments.Reserve(LastNodeIndex - FirstNodeIndex);
-		for (int32 Index = FirstNodeIndex + 1; Index < LastNodeIndex; ++Index)
-		{
-			LoopSegments[Index]->SwapOrientation();
-			Segments.Add(LoopSegments[Index]);
-		}
-
-		for (int32 ReverseIndex = Segments.Num() - 1, Index = FirstNodeIndex + 1; Index < LastNodeIndex; ++Index, --ReverseIndex)
-		{
-			LoopSegments[Index] = Segments[ReverseIndex];
-		}
-	};
-
-	if (LoopSegmentsIntersectionTool.DoesIntersect(Segment0->GetFirstNode(), Segment1->GetFirstNode()))
-	{
-		return false;
-	}
-
-	if (LoopSegmentsIntersectionTool.DoesIntersect(Segment0->GetSecondNode(), Segment1->GetSecondNode()))
-	{
-		return false;
-	}
-
-	int32 FirstNodeIndex = Segment0->GetFirstNode().GetIndex();
-	int32 LastNodeIndex = Segment1->GetFirstNode().GetIndex();
-
-	Segment0->GetSecondNode().RemoveSegment(Segment0);
-	Segment1->GetFirstNode().RemoveSegment(Segment1);
-
-	FIsoNode& Node = Segment0->GetSecondNode();
-	Segment0->SetSecondNode(Segment1->GetFirstNode());
-	Segment1->SetFirstNode(Node);
-
-	Segment1->GetFirstNode().ConnectSegment(Segment1);
-	Segment0->GetSecondNode().ConnectSegment(Segment0);
-
-#ifdef DEBUG_BY_SWAPPING_SEGMENTS
-	if (bDisplay)
-	{
-		F3DDebugSession _(TEXT("New Segments"));
-		Display(EGridSpace::UniformScaled, *Segment1, FirstNodeIndex, EVisuProperty::BlueCurve/*, true*/);
-		Display(EGridSpace::UniformScaled, *Segment0, LastNodeIndex, EVisuProperty::RedCurve/*, true*/);
-		//Wait();
-	}
-#endif
-	SwapLoopOrientation(FirstNodeIndex, LastNodeIndex);
-	return true;
-}
-
-void FIsoTriangulator::CheckLoopInSelfIntersectingAndFixIt()
-{
-	const double MaxGap = Grid.GetMinElementSize();
-	int32 LoopIndex = 0;
-
-	FIntersectionSegmentTool IntersectionTool(Grid);
-	IntersectionTool.Empty(LoopNodes.Num());
-
-	for (FIsoSegment* Segment : LoopSegments)
-	{
-		if (LoopIndex != ((FLoopNode&)Segment->GetFirstNode()).GetLoopIndex())
-		{
-			LoopIndex = ((FLoopNode&)Segment->GetFirstNode()).GetLoopIndex();
-			IntersectionTool.Empty(LoopNodes.Num());
-		}
-
-		if (const FIsoSegment* IntersectingSegment = IntersectionTool.DoesIntersect(*Segment))
-		{
-#ifdef DEBUG_CHECK_LOOP_IN_SELF_INTERSECTING_AND_FIX_IT
-			//if (bDisplay)
-			//{
-			//	{
-			//		IntersectionTool.Display(L"IntersectionTool");
-			//		F3DDebugSession _(TEXT("New Position"));
-			//		Display(EGridSpace::UniformScaled, *Segment);
-			//		Wait();
-			//	}
-			//}
-			if (bDisplay)
-			{
-				F3DDebugSession _(TEXT("Intersecting Segments"));
-				Display(EGridSpace::UniformScaled, *Segment, 0, EVisuProperty::BlueCurve/*, true*/);
-				Display(EGridSpace::UniformScaled, *IntersectingSegment, 0, EVisuProperty::RedCurve/*, true*/);
-				Wait();
-			}
-#endif
-
-			if (!TryToRemoveIntersectionBySwappingSegments(const_cast<FIsoSegment*>(IntersectingSegment), Segment))
-			{
-				if (!TryToRemoveIntersectionOfTwoConsecutiveIntersectingSegments(IntersectingSegment, Segment))
-				{
-					TryToRemoveIntersectionByMovingTheClosedOusidePoint(IntersectingSegment, Segment);
-				}
-			}
-		}
-		IntersectionTool.AddSegment(*Segment);
-	}
-
-	//if (bDisplay)
-	//{
-	//	{
-	//		F3DDebugSession A(TEXT("Segment Next"));
-	//		for (FIsoSegment* Segment : LoopSegments)
-	//		{
-	//			F3DDebugSession B(TEXT("Segment"));
-	//			Display(EGridSpace::UniformScaled, Segment->GetFirstNode(), ((FLoopNode&)Segment->GetFirstNode()).GetNextNode(), 0, EVisuProperty::BlueCurve);
-	//		}
-	//	}
-
-	//	{
-	//		F3DDebugSession C(TEXT("Segment Forward"));
-	//		for (FIsoSegment* Segment : LoopSegments)
-	//		{
-	//			F3DDebugSession D(TEXT("Segment"));
-	//			Display(EGridSpace::UniformScaled, Segment->GetSecondNode(), ((FLoopNode&)Segment->GetSecondNode()).GetPreviousNode(), 0, EVisuProperty::BlueCurve);
-	//		}
-	//	}
-	//	Wait();
-	//}
-
-}
-
-bool FIsoTriangulator::CanCycleBeFixedAndMeshed(const TArray<FIsoSegment*>& Cycle, FIntersectionSegmentTool& CycleIntersectionTool)
+bool FIsoTriangulator::CanCycleBeMeshed(const TArray<FIsoSegment*>& Cycle, FIntersectionSegmentTool& CycleIntersectionTool)
 {
-	bool bHasSelfIntersection = false;
+	bool bHasSelfIntersection = true;
 
-	for (FIsoSegment* Segment : Cycle)
+	for (const FIsoSegment* Segment : Cycle)
 	{
-		if (const FIsoSegment* IntersectingSegment = CycleIntersectionTool.DoesIntersect(*Segment))
+		if (CycleIntersectionTool.DoesIntersect(*Segment))
 		{
-			FMessage::Printf(Debug, TEXT("Remove local self intesection.\n"), Grid.GetFace()->GetId());
-
-#ifdef DEBUG_CAN_CYCLE_BE_FIXED_AND_MESHED
-			if (bDisplay)
-			{
-				F3DDebugSession _(TEXT("Segment"));
-				DisplayCycle(Cycle, TEXT("Cycle"));
-				F3DDebugSession A(TEXT("Segments"));
-				Display(EGridSpace::UniformScaled, *Segment, 0, EVisuProperty::RedCurve);
-				Display(EGridSpace::UniformScaled, *IntersectingSegment, 0, EVisuProperty::BlueCurve);
-				//Wait();
-			}
-#endif
-
-			if (!TryToRemoveIntersectionBySwappingSegments(const_cast<FIsoSegment*>(IntersectingSegment), Segment))
-			{
-				if (!TryToRemoveIntersectionOfTwoConsecutiveIntersectingSegments(IntersectingSegment, Segment))
-				{
-					if (!TryToRemoveIntersectionByMovingTheClosedOusidePoint(IntersectingSegment, Segment))
-					{
-#ifdef DEBUG_CAN_CYCLE_BE_FIXED_AND_MESHED
-						if (bDisplay)
-						{
-							F3DDebugSession _(TEXT("Failed"));
-							F3DDebugSession A(TEXT("Segments"));
-							Display(EGridSpace::UniformScaled, *Segment, 0, EVisuProperty::RedPoint);
-							Display(EGridSpace::UniformScaled, *IntersectingSegment, 0, EVisuProperty::BluePoint);
-							//Wait();
-						}
-#endif
-						return false;
-					}
-				}
-			}
+			FMessage::Printf(Log, TEXT("A cycle of the surface %d is in self intersecting. The mesh of this sector is canceled.\n"), Grid.GetFace()->GetId());
 
 #ifdef DEBUG_CAN_CYCLE_BE_FIXED_AND_MESHED
 			if (bDisplay)
@@ -2160,21 +1925,7 @@ bool FIsoTriangulator::CanCycleBeFixedAndMeshed(const TArray<FIsoSegment*>& Cycl
 				Wait();
 			}
 #endif
-
-			//FixPick2(Segment, IntersectingSegment);
-			bHasSelfIntersection = true;
-		}
-	}
-
-	if (bHasSelfIntersection)
-	{
-		for (const FIsoSegment* Segment : Cycle)
-		{
-			if (CycleIntersectionTool.DoesIntersect(*Segment))
-			{
-				FMessage::Printf(Log, TEXT("A cycle of the surface %d is in self intersecting. The mesh of this sector is canceled.\n"), Grid.GetFace()->GetId());
-				return false;
-			}
+			return false;
 		}
 	}
 
@@ -2300,9 +2051,8 @@ void FIsoTriangulator::MeshCycle(const EGridSpace Space, const TArray<FIsoSegmen
 #endif
 
 	// Check if the cycle is in self intersecting and fix it. 
-	if (!CanCycleBeFixedAndMeshed(Cycle, CycleIntersectionTool))
+	if (!CanCycleBeMeshed(Cycle, CycleIntersectionTool))
 	{
-		FMessage::Printf(Log, TEXT("A cycle of the surface %d is in self intersecting. The mesh of this sector is canceled.\n"), Grid.GetFace()->GetId());
 		return;
 	}
 
@@ -2434,8 +2184,9 @@ void FIsoTriangulator::MeshCycle(const EGridSpace Space, const TArray<FIsoSegmen
 
 	TFunction<void(FIsoSegment*, bool)> FindBestTriangle = [&](FIsoSegment* Segment, bool bOrientation)
 	{
-		FGetSlop GetSlopAtStartNode = ClockwiseSlop;
-		FGetSlop GetSlopAtEndNode = CounterClockwiseSlop;
+		using namespace IsoTriangulatorImpl;
+		GetSlopMethod GetSlopAtStartNode = ClockwiseSlop;
+		GetSlopMethod GetSlopAtEndNode = CounterClockwiseSlop;
 
 		// StartNode = A
 		FIsoNode& StartNode = bOrientation ? Segment->GetFirstNode() : Segment->GetSecondNode();
@@ -2504,35 +2255,6 @@ void FIsoTriangulator::MeshCycle(const EGridSpace Space, const TArray<FIsoSegmen
 		}
 #endif
 
-		if (&PreviousNode == &NextNode)
-		{
-			// No choice: PreviousSegment and NextSegment are connected, the triangle existe
-			if (bOrientation)
-			{
-				Segment->SetHasTriangleOnRight();
-			}
-			else
-			{
-				Segment->SetHasTriangleOnLeft();
-			}
-
-			BuildSegmentIfNeeded(&StartNode, &PreviousNode, PreviousSegment);
-			BuildSegmentIfNeeded(&PreviousNode, &EndNode, NextSegment);
-
-			Mesh->AddTriangle(EndNode.GetFaceIndex(), StartNode.GetFaceIndex(), PreviousNode.GetFaceIndex());
-
-#ifdef ADD_TRIANGLE_2D
-			if (bDisplay)
-			{
-				{
-					F3DDebugSession _(FString::Printf(TEXT("Triangle")));
-					DisplayTriangle(EGridSpace::UniformScaled, EndNode, StartNode, PreviousNode);
-				}
-			}
-#endif 
-			return;
-		}
-
 		FIsoNode* CandidatNode = nullptr;
 		FIsoSegment* StartToCandiatSegment = nullptr;
 		FIsoSegment* EndToCandiatSegment = nullptr;
@@ -2634,7 +2356,7 @@ void FIsoTriangulator::MeshCycle(const EGridSpace Space, const TArray<FIsoSegmen
 					PointCriteria = HUGE_VALUE;
 				}
 #else 
-				double PointCriteria = IsoscelesCriteria(StartPoint2D, EndPoint2D, NodePoint2D);
+				double PointCriteria = IsoTriangulatorImpl::IsoscelesCriteria(StartPoint2D, EndPoint2D, NodePoint2D);
 #endif
 				if (
 					// the candidate triangle is inside the current candidate triangle
@@ -3004,7 +2726,7 @@ void FIsoTriangulator::FindCycle(FIsoSegment* StartSegment, bool LeftSide, TArra
 
 	for (;;)
 	{
-		Segment = FindNextSegment(EGridSpace::UniformScaled, Segment, Node, ClockwiseSlop);
+		Segment = FindNextSegment(EGridSpace::UniformScaled, Segment, Node, IsoTriangulatorImpl::ClockwiseSlop);
 		if (Segment == nullptr)
 		{
 			Cycle.Empty();
@@ -3024,7 +2746,7 @@ void FIsoTriangulator::FindCycle(FIsoSegment* StartSegment, bool LeftSide, TArra
 			if (Segment->HasCycleOnLeft())
 			{
 				F3DDebugSession _(TEXT("Segment HasCycleOnLeft"));
-				Display(EGridSpace::UniformScaled, *Segment);
+				Display(EGridSpace::UniformScaled, *Segment, 0, EVisuProperty::RedPoint);
 				Wait();
 			}
 #endif
@@ -3048,7 +2770,7 @@ void FIsoTriangulator::FindCycle(FIsoSegment* StartSegment, bool LeftSide, TArra
 			if (Segment->HasCycleOnRight())
 			{
 				F3DDebugSession _(TEXT("Segment HasCycleOnRight"));
-				Display(EGridSpace::UniformScaled, *Segment);
+				Display(EGridSpace::UniformScaled, *Segment, 0, EVisuProperty::RedPoint);
 				Wait();
 			}
 #endif
@@ -3084,7 +2806,7 @@ void FIsoTriangulator::FindCycle(FIsoSegment* StartSegment, bool LeftSide, TArra
 static bool bDisplayStar = false;
 #endif
 
-FIsoSegment* FIsoTriangulator::FindNextSegment(EGridSpace Space, const FIsoSegment* StartSegment, const FIsoNode* StartNode, FGetSlop GetSlop) const
+FIsoSegment* FIsoTriangulator::FindNextSegment(EGridSpace Space, const FIsoSegment* StartSegment, const FIsoNode* StartNode, IsoTriangulatorImpl::GetSlopMethod GetSlop) const
 {
 	const FPoint2D& StartPoint = StartNode->Get2DPoint(Space, Grid);
 	const FPoint2D& EndPoint = (StartNode == &StartSegment->GetFirstNode()) ? StartSegment->GetSecondNode().Get2DPoint(Space, Grid) : StartSegment->GetFirstNode().Get2DPoint(Space, Grid);
@@ -3282,6 +3004,9 @@ void FIsoTriangulator::ConnectCellSubLoopsByNeighborhood(FCell& Cell)
 	}
 	else
 	{
+#ifdef DEBUG_BOWYERWATSON
+		FBowyerWatsonTriangulator::bDisplay = bDisplay;
+#endif
 		FBowyerWatsonTriangulator Triangulator(LoopBarycenters, EdgeVertexIndices);
 		Triangulator.Triangulate(Cell.bHasOuterLoop);
 		Triangulator.GetOuterVertices(Cell.BorderLoopIndices);
@@ -3361,11 +3086,19 @@ void FIsoTriangulator::TryToConnectTwoSubLoopsWithShortestSegment(FCell& Cell, c
 	for (int32 IndexA = 0; IndexA < LoopA.Num(); ++IndexA)
 	{
 		const FLoopNode* NodeA = LoopA[IndexA];
+		if (NodeA->IsDelete())
+		{
+			continue;
+		}
 		const FPoint2D& ACoordinates = NodeA->Get2DPoint(EGridSpace::UniformScaled, Grid);
 
 		for (int32 IndexB = 0; IndexB < LoopB.Num(); ++IndexB)
 		{
 			const FLoopNode* NodeB = LoopB[IndexB];
+			if (NodeB->IsDelete())
+			{
+				continue;
+			}
 			const FPoint2D& BCoordinates = NodeB->Get2DPoint(EGridSpace::UniformScaled, Grid);
 
 			double SquareDistance = ACoordinates.SquareDistance(BCoordinates);
@@ -3389,7 +3122,7 @@ void FIsoTriangulator::TryToConnectTwoSubLoopsWithShortestSegment(FCell& Cell, c
 		if (TryToCreateSegment(Cell, NodeA, ACoordinates, NodeB, BCoordinates, 0.1))
 		{
 			NodeA = &LoopA[MinIndexA]->GetNextNode();
-			if(Cell.Contains(NodeA))
+			if (Cell.Contains(NodeA))
 			{
 				NodeB = LoopB[MinIndexB];
 				TryToCreateSegment(Cell, NodeA, NodeA->Get2DPoint(EGridSpace::UniformScaled, Grid), NodeB, NodeB->Get2DPoint(EGridSpace::UniformScaled, Grid), 0.1);
@@ -3412,8 +3145,8 @@ void FIsoTriangulator::TryToConnectTwoLoopsWithIsocelesTriangle(FCell& Cell, con
 
 	TFunction<FIsoNode* (FIsoSegment*)> FindBestTriangle = [&](FIsoSegment* Segment) -> FIsoNode*
 	{
-		FGetSlop GetSlopAtStartNode = ClockwiseSlop;
-		FGetSlop GetSlopAtEndNode = CounterClockwiseSlop;
+		IsoTriangulatorImpl::GetSlopMethod GetSlopAtStartNode = IsoTriangulatorImpl::ClockwiseSlop;
+		IsoTriangulatorImpl::GetSlopMethod GetSlopAtEndNode = IsoTriangulatorImpl::CounterClockwiseSlop;
 
 		// StartNode = A
 		FIsoNode& StartNode = Segment->GetSecondNode();
@@ -3467,7 +3200,7 @@ void FIsoTriangulator::TryToConnectTwoLoopsWithIsocelesTriangle(FCell& Cell, con
 			// Check if the node is inside the sector (X) or outside (Z)
 			const FPoint2D& NodePoint2D = Node->Get2DPoint(EGridSpace::UniformScaled, Grid);
 
-			double PointCriteria = IsoscelesCriteria(StartPoint2D, EndPoint2D, NodePoint2D);
+			double PointCriteria = IsoTriangulatorImpl::IsoscelesCriteria(StartPoint2D, EndPoint2D, NodePoint2D);
 
 			// Triangle that are too open (more than rectangle triangle) are not tested 
 			if (PointCriteria > 0.6)
@@ -4462,187 +4195,6 @@ void FIsoTriangulator::ConnectCellSubLoopsByNeighborhood()
 
 #endif
 
-#ifdef UNUSED
-void FIsoTriangulator::FixPick(const FIsoSegment* Segment0, const FIsoSegment* Segment1)
-{
-	FPoint2D ClosedProjectedPoint;
-	FPoint2D Points[2][2];
-	int32 MinSegmentIndex = 0;
-	int32 MinPointIndex = 0;
-	double MinDistance = HUGE_VALUE;
-
-	auto ProjectPoint = [&](int32 PointToProjectSegmentIndex, int32 PointToProjectIndex, int32 SegmentIndex)
-	{
-		double Coordinate;
-		FPoint2D ProjectedPoint = ProjectPointOnSegment(Points[PointToProjectSegmentIndex][PointToProjectIndex], Points[SegmentIndex][0], Points[SegmentIndex][1], Coordinate, false);
-		double Distance = Points[PointToProjectSegmentIndex][PointToProjectIndex].SquareDistance(ProjectedPoint);
-
-		if (Distance < MinDistance)
-		{
-			MinDistance = Distance;
-			MinSegmentIndex = PointToProjectSegmentIndex;
-			MinPointIndex = PointToProjectIndex;
-			ClosedProjectedPoint = ProjectedPoint;
-		}
-	};
-
-	auto MoveNode = [&](const FIsoSegment* SegmentToMove)
-	{
-		FPoint2D MoveDirection = ClosedProjectedPoint - Points[MinSegmentIndex][MinPointIndex];
-		MoveDirection.Normalize();
-		MoveDirection *= Grid.GetFace()->GetCarrierSurface()->Get3DTolerance();;
-
-		FPoint2D NewCoordinate = ClosedProjectedPoint + MoveDirection;
-		const FIsoNode& NodeToMove = MinPointIndex ? SegmentToMove->GetSecondNode() : SegmentToMove->GetFirstNode();
-		const_cast<FIsoNode&>(NodeToMove).Set2DPoint(EGridSpace::UniformScaled, Grid, NewCoordinate);
-	};
-
-#ifdef CADKERNEL_DEV
-	if (bDisplay)
-	{
-		F3DDebugSession _(TEXT("Segment"));
-		Display(EGridSpace::UniformScaled, *Segment0, 0, EVisuProperty::BlueCurve);
-		Display(EGridSpace::UniformScaled, *Segment1, 0, EVisuProperty::BlueCurve);
-	}
-#endif
-
-	Points[0][0] = Segment0->GetFirstNode().Get2DPoint(EGridSpace::UniformScaled, Grid);
-	Points[0][1] = Segment0->GetSecondNode().Get2DPoint(EGridSpace::UniformScaled, Grid);
-	Points[1][0] = Segment1->GetFirstNode().Get2DPoint(EGridSpace::UniformScaled, Grid);
-	Points[1][1] = Segment1->GetSecondNode().Get2DPoint(EGridSpace::UniformScaled, Grid);
-
-	MinDistance = HUGE_VALUE;
-	ProjectPoint(0, 0, 1);
-	ProjectPoint(0, 1, 1);
-	ProjectPoint(1, 0, 0);
-	ProjectPoint(1, 1, 0);
-
-#ifdef CADKERNEL_DEV
-	if (bDisplay)
-	{
-		F3DDebugSession _(TEXT("Closed Point"));
-		DisplayPoint(Points[MinSegmentIndex][MinPointIndex], EVisuProperty::BluePoint);
-		DisplayPoint(ClosedProjectedPoint, EVisuProperty::RedPoint);
-	}
-#endif
-
-	MoveNode(MinSegmentIndex ? Segment1 : Segment0);
-
-#ifdef CADKERNEL_DEV
-	if (bDisplay)
-	{
-		F3DDebugSession _(TEXT("NewSegment"));
-		Display(EGridSpace::UniformScaled, *Segment0, 0, EVisuProperty::RedCurve);
-		Display(EGridSpace::UniformScaled, *Segment1, 0, EVisuProperty::RedCurve);
-		Wait();
-	}
-#endif
-};
-
-bool FIsoTriangulator::FixPick2(const FIsoSegment* SegmentToMove, const FIsoSegment* Segment1)
-{
-	FPoint2D ClosedProjectedPoint;
-	FPoint2D Points[2][2];
-	int32 MinSegmentIndex = -1;
-	int32 MinPointIndex = -1;
-	double MinDistance = HUGE_VALUE;
-
-	auto ProjectPoint = [&](int32 PointToProjectSegmentIndex, int32 PointToProjectIndex, int32 SegmentIndex)
-	{
-		double Coordinate;
-		FPoint2D ProjectedPoint = ProjectPointOnSegment(Points[PointToProjectSegmentIndex][PointToProjectIndex], Points[SegmentIndex][0], Points[SegmentIndex][1], Coordinate, false);
-
-#ifdef CADKERNEL_DEV
-		if (bDisplay)
-		{
-			F3DDebugSession _(TEXT("Closed Point"));
-			DisplayPoint(Points[PointToProjectSegmentIndex][PointToProjectIndex], EVisuProperty::BluePoint);
-			DisplayPoint(ProjectedPoint, EVisuProperty::RedPoint);
-		}
-#endif
-
-		if (Coordinate < 0 || Coordinate > 1)
-		{
-			return;
-		}
-
-		double Distance = Points[PointToProjectSegmentIndex][PointToProjectIndex].SquareDistance(ProjectedPoint);
-
-		if (Distance < MinDistance)
-		{
-			MinDistance = Distance;
-			MinSegmentIndex = PointToProjectSegmentIndex;
-			MinPointIndex = PointToProjectIndex;
-			ClosedProjectedPoint = ProjectedPoint;
-		}
-	};
-
-	auto MoveNode = [&](const FIsoSegment* SegmentToMove)
-	{
-		FPoint2D MoveDirection = ClosedProjectedPoint - Points[MinSegmentIndex][MinPointIndex];
-		MoveDirection.Normalize();
-		MoveDirection *= Grid.GetFace()->GetCarrierSurface()->Get3DTolerance();;
-
-		FPoint2D NewCoordinate = ClosedProjectedPoint + MoveDirection;
-		const FIsoNode& NodeToMove = MinPointIndex ? SegmentToMove->GetSecondNode() : SegmentToMove->GetFirstNode();
-		const_cast<FIsoNode&>(NodeToMove).Set2DPoint(EGridSpace::UniformScaled, Grid, NewCoordinate);
-	};
-
-#ifdef CADKERNEL_DEV
-	if (bDisplay)
-	{
-		F3DDebugSession _(TEXT("Segments before"));
-		Display(EGridSpace::UniformScaled, *SegmentToMove, 0, EVisuProperty::BlueCurve);
-		Display(EGridSpace::UniformScaled, *Segment1, 0, EVisuProperty::BlueCurve);
-	}
-#endif
-
-	TFunction<int32()> FindOuterNodeAndProjectIt = [&]() -> int32
-	{
-		double OrientedSlop = ComputeOrientedSlope(Points[0][0], Points[0][1], Points[1][0]);
-		return (OrientedSlop < 0) ? 0 : 1;
-	};
-
-
-	Points[0][0] = SegmentToMove->GetFirstNode().Get2DPoint(EGridSpace::UniformScaled, Grid);
-	Points[0][1] = SegmentToMove->GetSecondNode().Get2DPoint(EGridSpace::UniformScaled, Grid);
-	Points[1][0] = Segment1->GetFirstNode().Get2DPoint(EGridSpace::UniformScaled, Grid);
-	Points[1][1] = Segment1->GetSecondNode().Get2DPoint(EGridSpace::UniformScaled, Grid);
-
-	int32 OuterNode = FindOuterNodeAndProjectIt();
-	MinDistance = HUGE_VALUE;
-	ProjectPoint(0, 1, 1);
-	ProjectPoint(1, OuterNode, 0);
-
-	if (MinSegmentIndex < 0)
-	{
-		return false;
-	}
-
-#ifdef CADKERNEL_DEV
-	if (bDisplay)
-	{
-		F3DDebugSession _(TEXT("Closed Point"));
-		DisplayPoint(Points[MinSegmentIndex][MinPointIndex], EVisuProperty::BluePoint);
-		DisplayPoint(ClosedProjectedPoint, EVisuProperty::RedPoint);
-	}
-#endif
-
-	MoveNode(MinSegmentIndex ? Segment1 : SegmentToMove);
-
-#ifdef CADKERNEL_DEV
-	if (bDisplay)
-	{
-		F3DDebugSession _(TEXT("New Segments"));
-		Display(EGridSpace::UniformScaled, *SegmentToMove, 0, EVisuProperty::RedCurve);
-		Display(EGridSpace::UniformScaled, *Segment1, 0, EVisuProperty::RedCurve);
-		Wait();
-	}
-#endif
-	return true;
-};
-#endif
-
 
 
 } //namespace CADKernel
\ No newline at end of file
diff --git a/Engine/Source/Runtime/Datasmith/CADKernel/Private/CADKernel/Mesh/Meshers/IsoTriangulator.h b/Engine/Source/Runtime/Datasmith/CADKernel/Private/CADKernel/Mesh/Meshers/IsoTriangulator.h
--- a/Engine/Source/Runtime/Datasmith/CADKernel/Private/CADKernel/Mesh/Meshers/IsoTriangulator.h
+++ b/Engine/Source/Runtime/Datasmith/CADKernel/Private/CADKernel/Mesh/Meshers/IsoTriangulator.h
@@ -16,512 +16,737 @@
 //#define DEBUG_DELAUNAY
 namespace CADKernel
 {
-	class FGrid;
-	class FIntersectionSegmentTool;
-	class FFaceMesh;
-	struct FCell;
 
-	typedef TFunction<double(const FPoint2D&, const FPoint2D&, double)> FGetSlop;
+class FGrid;
+class FIntersectionSegmentTool;
+class FFaceMesh;
+struct FCell;
 
-	/**
-	 * for findNextSegment
-	 */ 
-	inline double ClockwiseSlop(const FPoint2D& StartPoint, const FPoint2D& EndPoint, double ReferenceSlope)
-	{
-		return 8. - ComputePositiveSlope(StartPoint, EndPoint, ReferenceSlope);
-	}
 
-	inline double CounterClockwiseSlop(const FPoint2D& StartPoint, const FPoint2D& EndPoint, double ReferenceSlope)
+namespace IsoTriangulatorImpl
+{
+typedef TFunction<double(const FPoint2D&, const FPoint2D&, double)> GetSlopMethod;
+typedef TFunction<FLoopNode* (FLoopNode*)> GetNextNodeMethod;
+typedef TFunction<const FLoopNode* (const FLoopNode*)> GetNextConstNodeMethod;
+typedef TFunction<const FLoopNode* (const FIsoSegment*)> GetSegmentToNodeMethod;
+typedef TPair<FLoopNode*, FLoopNode*> FLoopSection;
+}
+
+struct CADKERNEL_API FIsoTriangulatorChronos
+{
+	FDuration TriangulateDuration1 = FChrono::Init();
+	FDuration TriangulateDuration2 = FChrono::Init();
+	FDuration TriangulateDuration3 = FChrono::Init();
+	FDuration TriangulateDuration4 = FChrono::Init();
+	FDuration TriangulateDuration = FChrono::Init();
+	FDuration BuildIsoNodesDuration = FChrono::Init();
+	FDuration BuildLoopSegmentsDuration = FChrono::Init();
+	FDuration BuildInnerSegmentsDuration = FChrono::Init();
+	FDuration FindLoopSegmentOfInnerTriangulationDuration = FChrono::Init();
+	FDuration FindSegmentIsoUVSurroundingSmallLoopDuration = FChrono::Init();
+	FDuration FindIsoSegmentToLinkInnerToLoopDuration = FChrono::Init();
+	FDuration FindInnerSegmentToLinkLoopToLoopDuration = FChrono::Init();
+	FDuration FindSegmentToLinkLoopToLoopDuration = FChrono::Init();
+	FDuration FindSegmentToLinkLoopToLoopByDelaunayDuration = FChrono::Init();
+	FDuration FindSegmentToLinkInnerToLoopDuration = FChrono::Init();
+	FDuration SelectSegmentToLinkInnerToLoopsDuration = FChrono::Init();
+	FDuration TriangulateOverCycleDuration = FChrono::Init();
+	FDuration TriangulateInnerNodesDuration = FChrono::Init();
+
+	FIsoTriangulatorChronos()
+	{}
+
+	void PrintTimeElapse() const
 	{
-		return ComputePositiveSlope(StartPoint, EndPoint, ReferenceSlope);
+		FDuration IsoTriangulerDuration = FChrono::Init();
+		IsoTriangulerDuration += BuildIsoNodesDuration;
+		IsoTriangulerDuration += BuildLoopSegmentsDuration;
+		IsoTriangulerDuration += BuildInnerSegmentsDuration;
+		IsoTriangulerDuration += FindLoopSegmentOfInnerTriangulationDuration;
+		IsoTriangulerDuration += FindIsoSegmentToLinkInnerToLoopDuration;
+		IsoTriangulerDuration += FindSegmentToLinkLoopToLoopDuration;
+		IsoTriangulerDuration += FindSegmentToLinkLoopToLoopByDelaunayDuration;
+		IsoTriangulerDuration += FindSegmentToLinkInnerToLoopDuration;
+		IsoTriangulerDuration += SelectSegmentToLinkInnerToLoopsDuration;
+		IsoTriangulerDuration += TriangulateOverCycleDuration;
+		IsoTriangulerDuration += TriangulateInnerNodesDuration;
+
+		FChrono::PrintClockElapse(Log, TEXT(""), TEXT("IsoTrianguler"), IsoTriangulerDuration);
+		FChrono::PrintClockElapse(Log, TEXT("  "), TEXT("Triangulate"), TriangulateDuration);
+		FChrono::PrintClockElapse(Log, TEXT("    "), TEXT("BuildIsoNodes"), BuildIsoNodesDuration);
+		FChrono::PrintClockElapse(Log, TEXT("    "), TEXT("BuildLoopSegments"), BuildLoopSegmentsDuration);
+		FChrono::PrintClockElapse(Log, TEXT("    "), TEXT("BuildInnerSegments"), BuildInnerSegmentsDuration);
+		FChrono::PrintClockElapse(Log, TEXT("    "), TEXT("FindLoopSegmentOfInnerTriangulation"), FindLoopSegmentOfInnerTriangulationDuration);
+		FChrono::PrintClockElapse(Log, TEXT("      "), TEXT("FindSegmentIsoUVSurroundingSmallLoop"), FindSegmentIsoUVSurroundingSmallLoopDuration);
+		FChrono::PrintClockElapse(Log, TEXT("    "), TEXT("Find IsoSegment ToLink InnerToLoop"), FindIsoSegmentToLinkInnerToLoopDuration);
+		FChrono::PrintClockElapse(Log, TEXT("    "), TEXT("Find Segment ToLink LoopToLoop by Delaunay"), FindSegmentToLinkLoopToLoopByDelaunayDuration);
+		FChrono::PrintClockElapse(Log, TEXT("    "), TEXT("Find Segment ToLink LoopToLoop"), FindSegmentToLinkLoopToLoopDuration);
+		FChrono::PrintClockElapse(Log, TEXT("    "), TEXT("Find Segment ToLink InnerToLoop"), FindSegmentToLinkInnerToLoopDuration);
+		FChrono::PrintClockElapse(Log, TEXT("    "), TEXT("Select Segment ToLink InnerToLoop"), SelectSegmentToLinkInnerToLoopsDuration);
+		FChrono::PrintClockElapse(Log, TEXT("    "), TEXT("Mesh Over Cycle"), TriangulateOverCycleDuration);
+		FChrono::PrintClockElapse(Log, TEXT("    "), TEXT("Mesh Inner Nodes"), TriangulateInnerNodesDuration);
+		FChrono::PrintClockElapse(Log, TEXT("  "), TEXT("Triangulate1"), TriangulateDuration1);
+		FChrono::PrintClockElapse(Log, TEXT("  "), TEXT("Triangulate2"), TriangulateDuration2);
+		FChrono::PrintClockElapse(Log, TEXT("  "), TEXT("Triangulate3"), TriangulateDuration3);
+		FChrono::PrintClockElapse(Log, TEXT("  "), TEXT("Triangulate4"), TriangulateDuration4);
+		FChrono::PrintClockElapse(Log, TEXT("  "), TEXT("Triangulate "), TriangulateDuration);
 	}
+};
 
-	struct CADKERNEL_API FIsoTriangulatorChronos
-	{
-		FDuration TriangulateDuration1 = FChrono::Init();
-		FDuration TriangulateDuration2 = FChrono::Init();
-		FDuration TriangulateDuration3 = FChrono::Init();
-		FDuration TriangulateDuration4 = FChrono::Init();
-		FDuration TriangulateDuration = FChrono::Init();
-		FDuration BuildIsoNodesDuration = FChrono::Init();
-		FDuration BuildLoopSegmentsDuration = FChrono::Init();
-		FDuration BuildInnerSegmentsDuration = FChrono::Init();
-		FDuration FindLoopSegmentOfInnerTriangulationDuration = FChrono::Init();
-		FDuration FindSegmentIsoUVSurroundingSmallLoopDuration = FChrono::Init();
-		FDuration FindIsoSegmentToLinkInnerToLoopDuration = FChrono::Init();
-		FDuration FindInnerSegmentToLinkLoopToLoopDuration = FChrono::Init();
-		FDuration FindSegmentToLinkLoopToLoopDuration = FChrono::Init();
-		FDuration FindSegmentToLinkLoopToLoopByDelaunayDuration = FChrono::Init();
-		FDuration FindSegmentToLinkInnerToLoopDuration = FChrono::Init();
-		FDuration SelectSegmentToLinkInnerToLoopsDuration = FChrono::Init();
-		FDuration TriangulateOverCycleDuration = FChrono::Init();
-		FDuration TriangulateInnerNodesDuration = FChrono::Init();
-
-		FIsoTriangulatorChronos()
-		{}
-
-		void PrintTimeElapse() const
-		{
-			FDuration IsoTriangulerDuration = FChrono::Init();
-			IsoTriangulerDuration += BuildIsoNodesDuration;
-			IsoTriangulerDuration += BuildLoopSegmentsDuration;
-			IsoTriangulerDuration += BuildInnerSegmentsDuration;
-			IsoTriangulerDuration += FindLoopSegmentOfInnerTriangulationDuration;
-			IsoTriangulerDuration += FindIsoSegmentToLinkInnerToLoopDuration;
-			IsoTriangulerDuration += FindSegmentToLinkLoopToLoopDuration;
-			IsoTriangulerDuration += FindSegmentToLinkLoopToLoopByDelaunayDuration;
-			IsoTriangulerDuration += FindSegmentToLinkInnerToLoopDuration;
-			IsoTriangulerDuration += SelectSegmentToLinkInnerToLoopsDuration;
-			IsoTriangulerDuration += TriangulateOverCycleDuration;
-			IsoTriangulerDuration += TriangulateInnerNodesDuration;
-
-			FChrono::PrintClockElapse(Log, TEXT(""), TEXT("IsoTrianguler"), IsoTriangulerDuration);
-			FChrono::PrintClockElapse(Log, TEXT("  "), TEXT("Triangulate"), TriangulateDuration);
-			FChrono::PrintClockElapse(Log, TEXT("    "), TEXT("BuildIsoNodes"), BuildIsoNodesDuration);
-			FChrono::PrintClockElapse(Log, TEXT("    "), TEXT("BuildLoopSegments"), BuildLoopSegmentsDuration);
-			FChrono::PrintClockElapse(Log, TEXT("    "), TEXT("BuildInnerSegments"), BuildInnerSegmentsDuration);
-			FChrono::PrintClockElapse(Log, TEXT("    "), TEXT("FindLoopSegmentOfInnerTriangulation"), FindLoopSegmentOfInnerTriangulationDuration);
-			FChrono::PrintClockElapse(Log, TEXT("      "), TEXT("FindSegmentIsoUVSurroundingSmallLoop"), FindSegmentIsoUVSurroundingSmallLoopDuration);
-			FChrono::PrintClockElapse(Log, TEXT("    "), TEXT("Find IsoSegment ToLink InnerToLoop"), FindIsoSegmentToLinkInnerToLoopDuration);
-			FChrono::PrintClockElapse(Log, TEXT("    "), TEXT("Find Segment ToLink LoopToLoop by Delaunay"), FindSegmentToLinkLoopToLoopByDelaunayDuration);
-			FChrono::PrintClockElapse(Log, TEXT("    "), TEXT("Find Segment ToLink LoopToLoop"), FindSegmentToLinkLoopToLoopDuration);
-			FChrono::PrintClockElapse(Log, TEXT("    "), TEXT("Find Segment ToLink InnerToLoop"), FindSegmentToLinkInnerToLoopDuration);
-			FChrono::PrintClockElapse(Log, TEXT("    "), TEXT("Select Segment ToLink InnerToLoop"), SelectSegmentToLinkInnerToLoopsDuration);
-			FChrono::PrintClockElapse(Log, TEXT("    "), TEXT("Mesh Over Cycle"), TriangulateOverCycleDuration);
-			FChrono::PrintClockElapse(Log, TEXT("    "), TEXT("Mesh Inner Nodes"), TriangulateInnerNodesDuration);
-			FChrono::PrintClockElapse(Log, TEXT("  "), TEXT("Triangulate1"), TriangulateDuration1);
-			FChrono::PrintClockElapse(Log, TEXT("  "), TEXT("Triangulate2"), TriangulateDuration2);
-			FChrono::PrintClockElapse(Log, TEXT("  "), TEXT("Triangulate3"), TriangulateDuration3);
-			FChrono::PrintClockElapse(Log, TEXT("  "), TEXT("Triangulate4"), TriangulateDuration4);
-			FChrono::PrintClockElapse(Log, TEXT("  "), TEXT("Triangulate "), TriangulateDuration);
-		}
-	};
+class FIsoTriangulator
+{
+	friend class FParametricMesher;
 
+protected:
 
-	class FIsoTriangulator
-	{
-		friend class FParametricMesher;
-	
-	protected:
+	FGrid& Grid;
+	TSharedRef<FFaceMesh> Mesh;
 
-		FGrid& Grid;
-		TSharedRef<FFaceMesh> Mesh;
+	TArray<int32> LoopStartIndex;
+	TArray<FLoopNode> LoopNodes;
+	int32 LoopNodeCount = 0;
+	TArray<FLoopNode*> SortedLoopNodes;
 
-		TArray<int32> LoopStartIndex;
-		TArray<FLoopNode> LoopNodes;
-		int32 LoopNodeCount = 0;
-		TArray<FLoopNode*> SortedLoopNodes;
+	/**
+	 * GlobalIndexToIsoInnerNodes contains only inner nodes of the grid, if GlobalIndexToIsoInnerNodes[Index] == null, the point is outside the domain
+	 */
+	TArray<FIsoInnerNode*> GlobalIndexToIsoInnerNodes;
 
-		/**
-		 * GlobalIndexToIsoInnerNodes contains only inner nodes of the grid, if GlobalIndexToIsoInnerNodes[Index] == null, the point is outside the domain
-		 */
-		TArray<FIsoInnerNode*> GlobalIndexToIsoInnerNodes;
-		
-		/**
-		 * Static array of InnerNodes. Only used for allocation needs 
-		 */
-		TArray<FIsoInnerNode> InnerNodes;
-		int32 InnerNodeCount = 0;
+	/**
+	 * Static array of InnerNodes. Only used for allocation needs
+	 */
+	TArray<FIsoInnerNode> InnerNodes;
+	int32 InnerNodeCount = 0;
 
 
 
-		TFactory<FIsoSegment> IsoSegmentFactory;
+	TFactory<FIsoSegment> IsoSegmentFactory;
 
-		TArray<FIsoSegment*> LoopSegments;
-		TArray<FIsoSegment*> ThinZoneSegments;
-		TArray<FIsoSegment*> FinalInnerSegments;
-		TArray<FIsoSegment*> InnerToOuterSegments;
+	TArray<FIsoSegment*> LoopSegments;
+	TArray<FIsoSegment*> ThinZoneSegments;
+	TArray<FIsoSegment*> FinalInnerSegments;
+	TArray<FIsoSegment*> InnerToOuterSegments;
 
-		/**
-		 * Tool to check if a segment intersect or not existing segments.
-		 * To be optimal, depending on the segment, only a subset of segment is used. 
-		 * Checks intersection with loop.
-		 */
-		FIntersectionSegmentTool LoopSegmentsIntersectionTool;
-		FIntersectionSegmentTool InnerSegmentsIntersectionTool;
-		FIntersectionSegmentTool InnerToLoopSegmentsIntersectionTool;
-		FIntersectionSegmentTool InnerToOuterSegmentsIntersectionTool;
+	/**
+	 * Tool to check if a segment intersect or not existing segments.
+	 * To be optimal, depending on the segment, only a subset of segment is used.
+	 * Checks intersection with loop.
+	 */
+	FIntersectionSegmentTool LoopSegmentsIntersectionTool;
+	FIntersectionSegmentTool InnerSegmentsIntersectionTool;
+	FIntersectionSegmentTool InnerToLoopSegmentsIntersectionTool;
+	FIntersectionSegmentTool InnerToOuterSegmentsIntersectionTool;
 
 
-		/**
-		 * Define all the lower left index of grid node that the upper cell is surrounding a loop
-		 * This is set in FindSegmentIsoUVSurroundingSmallLoop
-		 * and use in TriangulateInnerNodes to don't generate the both cell triangles 
-		 */
-		TArray<int32> IndexOfLowerLeftInnerNodeSurroundingALoop;
+	/**
+	 * Define all the lower left index of grid node that the upper cell is surrounding a loop
+	 * This is set in FindSegmentIsoUVSurroundingSmallLoop
+	 * and use in TriangulateInnerNodes to don't generate the both cell triangles
+	 */
+	TArray<int32> IndexOfLowerLeftInnerNodeSurroundingALoop;
 
-		/**
-		 * Segments to link inner to boundary and boundary to boundary
-		 * From FindIsoSegmentToLinkInnerToBoundary
-		 * SelectSegmentInCandidateSegments
-		 */
-		TArray<FIsoSegment*> FinalToLoops;
+	/**
+	 * Segments to link inner to boundary and boundary to boundary
+	 * From FindIsoSegmentToLinkInnerToBoundary
+	 * SelectSegmentInCandidateSegments
+	 */
+	TArray<FIsoSegment*> FinalToLoops;
 
-		/**
-		 *
-		 */
-		TArray<FIsoSegment*> CandidateSegments;
+	/**
+	 *
+	 */
+	TArray<FIsoSegment*> CandidateSegments;
 
-		/**
-		 * Segments to link boundary to boundary and to complete the mesh
-		 */
-		TArray<FIsoSegment*> CandidateLoopToLoopSegments;
+	/**
+	 * Segments to link boundary to boundary and to complete the mesh
+	 */
+	TArray<FIsoSegment*> CandidateLoopToLoopSegments;
 
-		/**
-		 * Segments to link inner to boundary and to complete the mesh
-		 */
-		TArray<FIsoSegment*> CandidateInnerToLoopSegments;
+	/**
+	 * Segments to link inner to boundary and to complete the mesh
+	 */
+	TArray<FIsoSegment*> CandidateInnerToLoopSegments;
 
-		TArray<FIsoSegment*> NewTestSegments;
+	TArray<FIsoSegment*> NewTestSegments;
 
-		bool bDisplay = false;
+	bool bDisplay = false;
 
-	public:
+	bool bNeedCheckOrientation = false;
 
-		FIsoTriangulator(FGrid& InGrid, TSharedRef<FFaceMesh> EntityMesh);
+	static const double GeometricToMeshingToleranceFactor;
+	const double GeometricTolerance;
+	const double SquareGeometricTolerance;
+	const double SquareGeometricTolerance2;
+	const double MeshingTolerance;
+	const double SquareMeshingTolerance;
 
-		/**
-		 * Main method
-		 * @return false if the tessellation failed
-		 */
-		bool Triangulate();
+public:
 
-		void BuildNodes();
-		void CheckLoopInSelfIntersectingAndFixIt();
+	FIsoTriangulator(FGrid& InGrid, TSharedRef<FFaceMesh> EntityMesh);
 
-		/**
-		 * Build the segments of the loops and check if each loop is self intersecting.
-		 * @return false if the loop is self intersecting
-		 */
-		void BuildLoopSegments();
+	/**
+	 * Main method
+	 * @return false if the tessellation failed
+	 */
+	bool Triangulate();
 
-		void BuildInnerSegments();
+	void BuildNodes();
 
-		/**
-		 * Add temporary loops defining thin zones to avoid the tessellation of these zone.
-		 * These zones are tessellated in a specific process
-		 */
-		void BuildThinZoneSegments();
+	bool FindLoopIntersectionAndFixIt();
 
-		/**
-		 * Fill mesh node data (Position, normal, UV, Index) of the FFaceMesh object
-		 */
-		void FillMeshNodes();
+	/**
+	 * Build the segments of the loops and check if each loop is self intersecting.
+	 * @return false if the loop is self intersecting
+	 */
+	void BuildLoopSegments();
 
-		/**
-		 * Build the Inner Segments Intersection Tool used to check if a candidate segment crosses the inner mesh
-		 * 
-		 * The minimal set of segments of the intersection tool is the boundaries of the inner triangulation.
-		 * 
-		 * https://docs.google.com/presentation/d/1qUVOH-2kU_QXBVKyRUcdDy1Y6WGkcaJCiaS8wGjSZ6M/edit?usp=sharing
-		 * Slide "Boundary Segments Of Inner Triangulation"
-		 */
-		void BuildInnerSegmentsIntersectionTool();
+	void BuildInnerSegments();
 
-		/**
-		 * The purpose of the method is to add surrounding segments (boundary of an unitary inner grid cell) to the small loop to intersection tool to prevent traversing inner segments
-		 * A loop is inside inner segments
-		 *									|			 |
-		 *								   -----------------
-		 *									|	 XXX	 |
-		 *									|	XXXXX	 |
-		 *									|	 XXX	 |
-		 *								   -----------------
-		 *									|			 |
-		 * https://docs.google.com/presentation/d/1qUVOH-2kU_QXBVKyRUcdDy1Y6WGkcaJCiaS8wGjSZ6M/edit?usp=sharing
-		 * Slide "Find Inner Grid Cell Surrounding Small Loop"
-		 * This method finalizes BuildInnerSegmentsIntersectionTool
-		 */	
-		void FindInnerGridCellSurroundingSmallLoop();
+	/**
+	 * Add temporary loops defining thin zones to avoid the tessellation of these zone.
+	 * These zones are tessellated in a specific process
+	 */
+	void BuildThinZoneSegments();
+
+	/**
+	 * Fill mesh node data (Position, normal, UV, Index) of the FFaceMesh object
+	 */
+	void FillMeshNodes();
+
+	/**
+	 * Build the Inner Segments Intersection Tool used to check if a candidate segment crosses the inner mesh
+	 *
+	 * The minimal set of segments of the intersection tool is the boundaries of the inner triangulation.
+	 *
+	 * https://docs.google.com/presentation/d/1qUVOH-2kU_QXBVKyRUcdDy1Y6WGkcaJCiaS8wGjSZ6M/edit?usp=sharing
+	 * Slide "Boundary Segments Of Inner Triangulation"
+	 */
+	void BuildInnerSegmentsIntersectionTool();
+
+	/**
+	 * The purpose of the method is to add surrounding segments (boundary of an unitary inner grid cell) to the small loop to intersection tool to prevent traversing inner segments
+	 * A loop is inside inner segments
+	 *									|			 |
+	 *								   -----------------
+	 *									|	 XXX	 |
+	 *									|	XXXXX	 |
+	 *									|	 XXX	 |
+	 *								   -----------------
+	 *									|			 |
+	 * https://docs.google.com/presentation/d/1qUVOH-2kU_QXBVKyRUcdDy1Y6WGkcaJCiaS8wGjSZ6M/edit?usp=sharing
+	 * Slide "Find Inner Grid Cell Surrounding Small Loop"
+	 * This method finalizes BuildInnerSegmentsIntersectionTool
+	 */
+	void FindInnerGridCellSurroundingSmallLoop();
 
 #ifdef NEED_TO_CHECK_USEFULNESS
-		/**
-		 * 
-		 */
-		void CompleteIsoSegmentLoopToLoop();
+	/**
+	 *
+	 */
+	void CompleteIsoSegmentLoopToLoop();
 #endif
 
-		void ConnectCellLoops();
-		void FindCellContainingBoundaryNodes(TArray<FCell>& Cells);
+	void ConnectCellLoops();
+	void FindCellContainingBoundaryNodes(TArray<FCell>& Cells);
 
-		/**
-		 * The closest loops are connected together
-		 * To do it, a Delaunay triangulation of the loop barycenters is realized.
-		 * Each edge of this mesh defined a near loops pair
-		 * The shortest segment is then build between this pair of loops
-		 */
-		void ConnectCellSubLoopsByNeighborhood(FCell& cell);
+	/**
+	 * The closest loops are connected together
+	 * To do it, a Delaunay triangulation of the loop barycenters is realized.
+	 * Each edge of this mesh defined a near loops pair
+	 * The shortest segment is then build between this pair of loops
+	 */
+	void ConnectCellSubLoopsByNeighborhood(FCell& cell);
 
-		/**
-		 * 2nd step
-		 */
-		void ConnectCellCornerToInnerLoop(FCell& Cell);
+	/**
+	 * 2nd step
+	 */
+	void ConnectCellCornerToInnerLoop(FCell& Cell);
 
-		/**
-		 * 3rd step : IsoSegment linking 
-		 */
-		void FindIsoSegmentToLinkOuterLoopNodes(FCell& Cell);
+	/**
+	 * 3rd step : IsoSegment linking
+	 */
+	void FindIsoSegmentToLinkOuterLoopNodes(FCell& Cell);
 
-		/**
-		 * 4th step : If their is no segment candidate  
-		 */
-		void FindSegmentToLinkOuterLoopNodes(FCell& Cell);
+	/**
+	 * 4th step : If their is no segment candidate
+	 */
+	void FindSegmentToLinkOuterLoopNodes(FCell& Cell);
 
-		/**
-		 * 5th step : try to connect outer loop Extremities
-		 */
-		void AddSementToLinkOuterLoopExtremities(FCell& Cell);
+	/**
+	 * 5th step : try to connect outer loop Extremities
+	 */
+	void AddSementToLinkOuterLoopExtremities(FCell& Cell);
 
-		/**
-		 * The goal of this algorithm is to connect iso U (or V) aligned loop nodes as soon as they are nearly in the same iso V (or U) strip.
-		 * I.e.:
-		 * - Iso U aligned: NodeA.U = NodeB.U +/-TolU 
-		 * - In the same strip: each node of the segment has the same index "i" that verify: isoV[i] - TolV < Node.V < isoV[i+1] + TolV
-		 */
-		void FindIsoSegmentToLinkLoopToLoop();
+	/**
+	 * The goal of this algorithm is to connect iso U (or V) aligned loop nodes as soon as they are nearly in the same iso V (or U) strip.
+	 * I.e.:
+	 * - Iso U aligned: NodeA.U = NodeB.U +/-TolU
+	 * - In the same strip: each node of the segment has the same index "i" that verify: isoV[i] - TolV < Node.V < isoV[i+1] + TolV
+	 */
+	void FindIsoSegmentToLinkLoopToLoop();
 #ifdef UNUSED_TO_DELETE_
-		void LastChanceToCreateSegmentInCell(FCell& Cell);
-		/**
-		 * The goal of this algorithm is to connect inner node (node of the grid UV) to iso aligned loop node when they are in the same iso V (or U) strip.
-		 * I.e.:
-		 * - Iso U aligned: NodeA.U < NodeB.U  
-		 * - In the same strip: each node of the segment has the same index "i" that verify: isoV[i-1] - TolV < BoundaryNode.V or BoundaryNode.V < isoV[i+1] + TolV
-		 * https://docs.google.com/presentation/d/1qUVOH-2kU_QXBVKyRUcdDy1Y6WGkcaJCiaS8wGjSZ6M/edit?usp=sharing
-		 * Slide "Find Iso Segment To Link Inner To Loop"
-		 */
-		void FindIsoSegmentToLinkInnerToLoop();
+	void LastChanceToCreateSegmentInCell(FCell& Cell);
+	/**
+	 * The goal of this algorithm is to connect inner node (node of the grid UV) to iso aligned loop node when they are in the same iso V (or U) strip.
+	 * I.e.:
+	 * - Iso U aligned: NodeA.U < NodeB.U
+	 * - In the same strip: each node of the segment has the same index "i" that verify: isoV[i-1] - TolV < BoundaryNode.V or BoundaryNode.V < isoV[i+1] + TolV
+	 * https://docs.google.com/presentation/d/1qUVOH-2kU_QXBVKyRUcdDy1Y6WGkcaJCiaS8wGjSZ6M/edit?usp=sharing
+	 * Slide "Find Iso Segment To Link Inner To Loop"
+	 */
+	void FindIsoSegmentToLinkInnerToLoop();
 #endif
 
-		/**
-		 * The purpose of the method is select a minimal set of segments connecting loops together
-		 * The final segments will be selected with SelectSegmentInCandidateSegments
-		 */
-		void ConnectCellSubLoopsByNeighborhood();  // to rename and clean
+	/**
+	 * The purpose of the method is select a minimal set of segments connecting loops together
+	 * The final segments will be selected with SelectSegmentInCandidateSegments
+	 */
+	void ConnectCellSubLoopsByNeighborhood();  // to rename and clean
 
 #ifdef TODELETE
-		void FindSegmentToLinkLoopToLoop();
+	void FindSegmentToLinkLoopToLoop();
 #endif
 
-		/**
-		 * The purpose of the method is select a minimal set of segments
-		 * The final segments will be selected with SelectSegmentInCandidateSegments
-		 */
-		void FindCandidateSegmentsToLinkInnerToLoop();
+	/**
+	 * The purpose of the method is select a minimal set of segments
+	 * The final segments will be selected with SelectSegmentInCandidateSegments
+	 */
+	void FindCandidateSegmentsToLinkInnerToLoop();
 
 
-		void FindCandidateSegmentsToLinkInnerAndLoop();
+	void FindCandidateSegmentsToLinkInnerAndLoop();
 
-		/**
-		 * Complete the final set of segments with the best subset of segments in CandidateSegments
-		 */
-		void SelectSegmentInCandidateSegments();
+	/**
+	 * Complete the final set of segments with the best subset of segments in CandidateSegments
+	 */
+	void SelectSegmentInCandidateSegments();
 
-		/**
-		 * The goal of this algorithm is to connect unconnected inner segment extremity i.e. extremity with one or less connected segment to the closed boundary node
-		 */
-		void ConnectUnconnectedInnerSegments();
+	/**
+	 * The goal of this algorithm is to connect unconnected inner segment extremity i.e. extremity with one or less connected segment to the closed boundary node
+	 */
+	void ConnectUnconnectedInnerSegments();
 
-		/**
-		 * Finalize the tesselation between inner grid boundary and loops.
-		 * The final set of segments define a network
-		 * Each minimal cycle is tessellated independently
-		 */
-		void TriangulateOverCycle(const EGridSpace Space);
+	/**
+	 * Finalize the tessellation between inner grid boundary and loops.
+	 * The final set of segments define a network
+	 * Each minimal cycle is tessellated independently
+	 */
+	void TriangulateOverCycle(const EGridSpace Space);
 
-		/**
-		 * Find in the network a minimal cylce stating from a segment
-		 */
-		void FindCycle(FIsoSegment* StartSegment, bool bLeftSide, TArray<FIsoSegment*>& Cycle, TArray<bool>& CycleOrientation);
+	/**
+	 * Find in the network a minimal cycle stating from a segment
+	 */
+	void FindCycle(FIsoSegment* StartSegment, bool bLeftSide, TArray<FIsoSegment*>& Cycle, TArray<bool>& CycleOrientation);
 
-		/**
-		 * Generate the "Delaunay" tessellation of the cycle. 
-		 * The algorithm is based on frontal process
-		 */
-		void MeshCycle(const EGridSpace Space, const TArray<FIsoSegment*>& cycle, const TArray<bool>& cycleOrientation);
+	/**
+	 * Generate the "Delaunay" tessellation of the cycle.
+	 * The algorithm is based on frontal process
+	 */
+	void MeshCycle(const EGridSpace Space, const TArray<FIsoSegment*>& cycle, const TArray<bool>& cycleOrientation);
+
+	bool CanCycleBeMeshed(const TArray<FIsoSegment*>& Cycle, FIntersectionSegmentTool& CycleIntersectionTool);
 
-		bool CanCycleBeFixedAndMeshed(const TArray<FIsoSegment*>& Cycle, FIntersectionSegmentTool& CycleIntersectionTool);
+	bool TryToRemoveSelfIntersectionByMovingTheClosedOusidePoint(const FIsoSegment& Segment0, const FIsoSegment& Segment1);
+	bool TryToRemoveIntersectionByMovingTheClosedOusidePoint(const FIsoSegment& Segment0, const FIsoSegment& Segment1);
+	void RemoveIntersectionByMovingOutsideSegmentNodeInside(const FIsoSegment& IntersectingSegment, const FIsoSegment& Segment);
 
-		bool TryToRemoveIntersectionByMovingTheClosedOusidePoint(const FIsoSegment* Segment0, const FIsoSegment* Segment1);
-		bool TryToRemoveIntersectionBySwappingSegments(FIsoSegment* Segment0, FIsoSegment* Segment1);
+	bool TryToRemoveIntersectionBySwappingSegments(FIsoSegment& Segment0, FIsoSegment& Segment1);                                  // ok
 
 		/**
 		 * Segment and the segment next (or before Segment) are intersecting IntersectingSegment
-		 * The commun node is moved
+		 * The common node is moved
 		 */
-		bool TryToRemoveIntersectionOfTwoConsecutiveIntersectingSegments(const FIsoSegment* IntersectingSegment, FIsoSegment* Segment);
+	bool TryToRemoveIntersectionOfTwoConsecutiveIntersectingSegments(const FIsoSegment& IntersectingSegment, FIsoSegment& Segment); // ok 
 
-		//void FixPick(const FIsoSegment* Segment0, const FIsoSegment* Segment1);
-		//bool FixPick2(const FIsoSegment* Segment0, const FIsoSegment* Segment1);
+	bool RemoveNodeOfLoop(FLoopNode& NodeToRemove);
+	void RemovePickOfLoop(FIsoSegment& Segment);
 
-		//void AddSegmentsAndRemoveIntersection(FIntersectionSegmentTool& CycleIntersectionTool, const TArray<FIsoSegment*>& Cycle, const TArray<bool>& CycleOrientation);
-
-		/**
-		 * Finalization of the mesh by the tessellation of the inner grid
-		 */
-		void TriangulateInnerNodes();
+	bool RemovePickRecursively(FLoopNode* Node0, FLoopNode* Node1);
 
-		/**
-		 * Sorted loop node array used to have efficient loop proximity node research
-		 */
-		void SortLoopNodes()
+	bool CheckAndRemovePick(const FPoint2D& PreviousPoint, const FPoint2D& NodeToRemovePoint, const FPoint2D& NextPoint, FLoopNode& NodeToRemove)
+	{
+		double Slop = ComputeUnorientedSlope(NodeToRemovePoint, PreviousPoint, NextPoint);
+		bool bRemoveNode = false;
+		if (Slop < 0.03)
 		{
-			SortedLoopNodes.Reserve(LoopNodes.Num());
-			for (FLoopNode& LoopNode : LoopNodes)
+			bRemoveNode = true;
+		}
+		else if (Slop < 0.1)
+		{
+			double SquareDistance1 = SquareDistanceOfPointToSegment(PreviousPoint, NodeToRemovePoint, NextPoint);
+			double SquareDistance2 = SquareDistanceOfPointToSegment(NextPoint, NodeToRemovePoint, PreviousPoint);
+			double MinSquareDistance = FMath::Min(SquareDistance1, SquareDistance2);
+			if (MinSquareDistance < SquareGeometricTolerance2)
 			{
-				SortedLoopNodes.Add(&LoopNode);
+				bRemoveNode = true;
 			}
+		}
 
-			Algo::Sort(SortedLoopNodes, [this](const FLoopNode* Node1, const FLoopNode* Node2)
-				{
-					const FPoint2D& Node1Coordinates = Node1->Get2DPoint(EGridSpace::Default2D, Grid);
-					const FPoint2D& Node2Coordinates = Node2->Get2DPoint(EGridSpace::Default2D, Grid);
-					return (Node1Coordinates.U + Node1Coordinates.V) < (Node2Coordinates.U + Node2Coordinates.V);
-				});
+		if (bRemoveNode)
+		{
+			return RemoveNodeOfLoop(NodeToRemove);
 		}
+		return false;
+	};
+
+	/**
+	 * Finalization of the mesh by the tessellation of the inner grid
+	 */
+	void TriangulateInnerNodes();
+
+	/**
+	 * Sorted loop node array used to have efficient loop proximity node research
+	 */
+	void SortLoopNodes()
+	{
+		SortedLoopNodes.Reserve(LoopNodes.Num());
+		for (FLoopNode& LoopNode : LoopNodes)
+		{
+			SortedLoopNodes.Add(&LoopNode);
+		}
+
+		Algo::Sort(SortedLoopNodes, [this](const FLoopNode* Node1, const FLoopNode* Node2)
+			{
+				const FPoint2D& Node1Coordinates = Node1->Get2DPoint(EGridSpace::Default2D, Grid);
+				const FPoint2D& Node2Coordinates = Node2->Get2DPoint(EGridSpace::Default2D, Grid);
+				return (Node1Coordinates.U + Node1Coordinates.V) < (Node2Coordinates.U + Node2Coordinates.V);
+			});
+	}
+
+private:
+
+	// Methods used by FindLoopIntersectionAndFixIt Step1
+	void FindBestLoopExtremity(TArray<FLoopNode*>& BestStartNodeOfLoop);
+	void FindLoopIntersections(const TArray<FLoopNode*>& LoopNodes, bool bForward, TArray<TPair<double, double>>& OutIntersections);
+	EOrientation GetLoopOrientation(const FLoopNode* StartNode);
+	void FixLoopOrientation(const TArray<FLoopNode*>& LoopNodes);
+	void SwapLoopOrientation(int32 FirstSegmentIndex, int32 LastSegmentIndex);
+	void MoveIntersectingSectionBehindOppositeSection(IsoTriangulatorImpl::FLoopSection IntersectingSection, IsoTriangulatorImpl::FLoopSection OppositeSection, IsoTriangulatorImpl::GetNextNodeMethod GetNext, IsoTriangulatorImpl::GetNextNodeMethod GetPrevious);
+	void OffsetSegment(FIsoSegment& Segment, TSegment<FPoint2D>& Segment2D, TSegment<FPoint2D>& IntersectingSegment2D);
+	void OffsetNode(FLoopNode& Node, TSegment<FPoint2D>& IntersectingSegment2D);
+
+	/**
+	 * If the new position of the node is closed to the previous or next node position, it will be delete
+	 */
+	void MoveNode(FLoopNode& NodeToMove, FPoint2D& ProjectedPoint);
+
+	/**
+	 * @return false if the SubLoop is bigger than the other part of the loop
+	 */
+	bool RemoveLoopIntersections(const TArray<FLoopNode*>& NodesOfLoop, TArray<TPair<double, double>>& Intersections, bool bForward);
+
+	void RemoveLoopPicks(TArray<FLoopNode*>& NodesOfLoop, TArray<TPair<double, double>>& Intersections);
+
+	void RemoveThePick(const TArray<FLoopNode*>& NodesOfLoop, const TPair<double, double>& Intersection, bool bForward);
+	void SwapNodes(const TArray<FLoopNode*>& NodesOfLoop, const TPair<double, double>& Intersection, bool bForward);
+
+	/**
+	 * Case0: __     ___        __      ___
+	 *          \   /             \    /
+	 *           o o        =>     o  o
+	 *        __/   \___		__/    \___
+	 *
+	 * Case1:    ____    ___        ____    ___
+	 *           |   \  /   |       |   \  /   |
+	 *           |    oo    |   =>  |    o     |
+	 *           |___/  \___|       |          |
+	 *                              |    o     |
+	 * 								|___/  \___|
+	 */
+	bool SpreadCoincidentNodes(const TArray<FLoopNode*>& NodesOfLoop, TPair<double, double> Intersection);
+
+	/**
+	 * Case:  o_________o          o_________o
+	 *             o         =>
+	 *          __/ \___		        o
+	 *								 __/ \___
+	 */
+	bool MovePickBehind(const TArray<FLoopNode*>& NodesOfLoop, TPair<double, double> Intersection, bool bKeyIsExtremity);
+
+	/**
+	 * Two cases:
+	 *    - the segments of the intersection a closed parallel and in same orientation. in this case, the sub-loop is a long pick. The sub-Loop is delete
+	 *    - the loop is an inner loop closed to the border:
+	 *       _____________                _____________
+	 *   |  /             \			  |__/             \
+	 *    \/               |     =>                     |
+	 *    /\               |		   __               |
+	 *   |  \_____________/			  |  \_____________/
+	 *
+	 */
+	void TryToSwapSegmentsOrRemoveLoop(const TArray<FLoopNode*>& NodesOfLoop, const TPair<double, double>& Intersection, bool bForward);
+
+	bool RemovePickToOutside(const TArray<FLoopNode*>& NodesOfLoop, const TPair<double, double>& Intersection, const TPair<double, double>& NextIntersection, bool bForward);
+
+	/**
+	 * @return false if the SubLoop is bigger than the other part of the loop
+	 */
+	bool RemoveIntersectionsOfSubLoop(const TArray<FLoopNode*>& LoopNodesFromStartNode, TArray<TPair<double, double>>& Intersections, int32 IntersectionIndex, int32 IntersectionCount, bool bForward);
+	bool RemoveUniqueIntersection(const TArray<FLoopNode*>& LoopNodesFromStartNode, TPair<double, double> Intersection, bool bForward);
+	//void RemoveSubLoop(const TArray<FLoopNode*>& NodesOfLoop, const TPair<double, double>& Intersection, bool bForward);
+	void RemoveSubLoop(FLoopNode* StartNode, FLoopNode* EndNode, IsoTriangulatorImpl::GetNextNodeMethod GetNext);
+
+	/**
+	 * Check if the sub-loop is not bigger than the main loop.
+	 * If the segment count of the SubLoop is bigger than 1/3 of the segment count of the main loop
+	 * the length of each part are compute.
+	 * If the sub-loop length is bigger than the main loop^, there is a problem (problem orientation ?)
+	 * The surface meshing is canceled to avoid a bigger problem during the meshing step, that could cause a process crash.
+	 */
+	bool IsSubLoopBiggerThanMainLoop(const TArray<FLoopNode*>& LoopNodesFromStartNode, const TPair<double, double>& Intersection, bool bForward);
+
+	void RemoveIntersectionByMovingOutsideNodeInside(const FIsoSegment& IntersectingSegment, FLoopNode& NodeToMove);
+
+	bool CheckMainLoopConsistency();
+
+	// Methods used by FindLoopIntersectionAndFixIt Step2
+	void FillIntersectionToolWithOuterLoop();
+
+	void FixIntersectionBetweenLoops();
+
+	FIsoSegment* FindNextSegment(EGridSpace Space, const FIsoSegment* StartSegment, const FIsoNode* StartNode, IsoTriangulatorImpl::GetSlopMethod GetSlop) const;
 
-	private:
 
-		FIsoSegment* FindNextSegment(EGridSpace Space, const FIsoSegment* StartSegment, const FIsoNode* StartNode, FGetSlop GetSlop) const;
 
 #ifdef UNUSED_TO_DELETE_
-		void CreateLoopToLoopSegment(FLoopNode* NodeA, const FPoint2D& ACoordinates, FLoopNode* NodeB, const FPoint2D& BCoordinates, TArray<FIsoSegment*>& NewSegments);
+	void CreateLoopToLoopSegment(FLoopNode* NodeA, const FPoint2D& ACoordinates, FLoopNode* NodeB, const FPoint2D& BCoordinates, TArray<FIsoSegment*>& NewSegments);
 #endif
-		// ==========================================================================================
-		// 	   Create segments
-		// ==========================================================================================
-
-		/**  
-		 *  SubLoopA                  SubLoopB
-		 *      --X---X             X-----X--  
-		 *             \           /  
-		 *              \         /
-		 *               X=======X 
-		 *              /         \
-		 *             /           \
-		 *      --X---X             X-----X--  
-		 *
-		 *     ======= ShortestSegment
-		 */
-		void TryToConnectTwoSubLoopsWithShortestSegment(FCell& Cell, const TArray<FLoopNode*>& SubLoopA, const TArray<FLoopNode*>& SubLoopB);
+	// ==========================================================================================
+	// 	   Create segments
+	// ==========================================================================================
 
-		void TryToConnectTwoLoopsWithIsocelesTriangle(FCell& Cell, const TArray<FLoopNode*>& SubLoopA, const TArray<FLoopNode*>& SubLoopB);
+	/**
+	 *  SubLoopA                  SubLoopB
+	 *      --X---X             X-----X--
+	 *             \           /
+	 *              \         /
+	 *               X=======X
+	 *              /         \
+	 *             /           \
+	 *      --X---X             X-----X--
+	 *
+	 *     ======= ShortestSegment
+	 */
+	void TryToConnectTwoSubLoopsWithShortestSegment(FCell& Cell, const TArray<FLoopNode*>& SubLoopA, const TArray<FLoopNode*>& SubLoopB);
 
-		/**
-		 *    --X------X-----X--  <-- SubLoopA                         
-		 *             I                                                                     
-		 *             I    <- The most isoSegment
-		 *             I 
-		 *    ----X----X--------  <-- SubLoopB                        
-		 *
-		 */
-		void TryToConnectTwoSubLoopsWithTheMostIsoSegment(FCell& Cell, const TArray<FLoopNode*>& SubLoopA, const TArray<FLoopNode*>& SubLoopB);
+	void TryToConnectTwoLoopsWithIsocelesTriangle(FCell& Cell, const TArray<FLoopNode*>& SubLoopA, const TArray<FLoopNode*>& SubLoopB);
 
-		/**
-		 *    X--X------X-----X--  <-- SubLoop                         
-		 *    |         I                                                                     
-		 *    |         I   <- The most isoSegment
-		 *    |         I
-		 *    X----X----X--------                          
-		 *
-		 */
-		void TryToConnectVertexSubLoopWithTheMostIsoSegment(FCell& Cell, const TArray<FLoopNode*>& SubLoop);
-		bool TryToCreateSegment(FCell& Cell, FLoopNode* NodeA, const FPoint2D& ACoordinates, FIsoNode* NodeB, const FPoint2D& BCoordinates, const double FlatAngle);
+	/**
+	 *    --X------X-----X--  <-- SubLoopA
+	 *             I
+	 *             I    <- The most isoSegment
+	 *             I
+	 *    ----X----X--------  <-- SubLoopB
+	 *
+	 */
+	void TryToConnectTwoSubLoopsWithTheMostIsoSegment(FCell& Cell, const TArray<FLoopNode*>& SubLoopA, const TArray<FLoopNode*>& SubLoopB);
+
+	/**
+	 *    X--X------X-----X--  <-- SubLoop
+	 *    |         I
+	 *    |         I   <- The most isoSegment
+	 *    |         I
+	 *    X----X----X--------
+	 *
+	 */
+	void TryToConnectVertexSubLoopWithTheMostIsoSegment(FCell& Cell, const TArray<FLoopNode*>& SubLoop);
+	bool TryToCreateSegment(FCell& Cell, FLoopNode* NodeA, const FPoint2D& ACoordinates, FIsoNode* NodeB, const FPoint2D& BCoordinates, const double FlatAngle);
 
 #ifdef CADKERNEL_DEV
-	public:
-		FIsoTriangulatorChronos Chronos;
+public:
+	FIsoTriangulatorChronos Chronos;
 
-		void DisplayPixels(TArray<uint8>& Pixel) const;
+	void DisplayPixels(TArray<uint8>& Pixel) const;
 
-		void DisplayPixel(const int32 Index) const;
-		void DisplayPixel(const int32 IndexU, const int32 IndexV) const;
+	void DisplayPixel(const int32 Index) const;
+	void DisplayPixel(const int32 IndexU, const int32 IndexV) const;
 
-		void DisplayIsoNodes(EGridSpace Space) const;
-		void Display(EGridSpace Space, const FIsoNode& Node, FIdent Ident = 0, EVisuProperty Property = EVisuProperty::BluePoint) const;
-		void Display(EGridSpace Space, const FIsoNode& NodeA, const FIsoNode& NodeB, FIdent Ident = 0, EVisuProperty Property = EVisuProperty::Element) const;
-		void Display(EGridSpace Space, const FIsoSegment& Segment, FIdent Ident = 0, EVisuProperty Property = EVisuProperty::Element, bool bDisplayOrientation = false) const;
-		void Display(EGridSpace Space, const TCHAR* Message, const TArray<FIsoSegment*>& Segment, bool bDisplayNode, bool bDisplayOrientation = false, EVisuProperty Property = EVisuProperty::Element) const;
+	void DisplayIsoNodes(EGridSpace Space) const;
+	void Display(EGridSpace Space, const FIsoNode& Node, FIdent Ident = 0, EVisuProperty Property = EVisuProperty::BluePoint) const;
+	void Display(FString Message, EGridSpace Space, const TArray<const FIsoNode*>& Points, EVisuProperty Property) const;
 
-		void DisplayLoops(EGridSpace Space, const TCHAR* Message, const TArray<FLoopNode>& Nodes, bool bDisplayNode, EVisuProperty Property = EVisuProperty::Element) const;
+	void Display(EGridSpace Space, const FIsoNode& NodeA, const FIsoNode& NodeB, FIdent Ident = 0, EVisuProperty Property = EVisuProperty::Element) const;
+	void Display(EGridSpace Space, const FIsoSegment& Segment, FIdent Ident = 0, EVisuProperty Property = EVisuProperty::Element, bool bDisplayOrientation = false) const;
+	void Display(EGridSpace Space, const TCHAR* Message, const TArray<FIsoSegment*>& Segment, bool bDisplayNode, bool bDisplayOrientation = false, EVisuProperty Property = EVisuProperty::Element) const;
 
-		void DisplayCycle(const TArray<FIsoSegment*>& Cycle, const TCHAR* Message) const;
+	void DisplayLoops(EGridSpace Space, const TCHAR* Message, const TArray<FLoopNode>& Nodes, bool bDisplayNode, EVisuProperty Property = EVisuProperty::Element) const;
+	void DisplayLoops(const TCHAR* Message, bool bOneNode = true, bool bSplitBySegment = false, bool bDisplayNext = false, bool bDisplayPrevious = false) const;
+	void DisplayLoop(EGridSpace Space, const TCHAR* Message, const TArray<FLoopNode*>& Nodes, bool bDisplayNode, EVisuProperty Property = EVisuProperty::Element) const;
 
-		void DisplayCells(const TArray<FCell>& Cells) const;
-		void DisplayCell(const FCell& Cell) const;
+	void DisplayCycle(const TArray<FIsoSegment*>& Cycle, const TCHAR* Message) const;
 
-		void DisplayTriangle(EGridSpace Space, const FIsoNode& NodeA, const FIsoNode& NodeB, const FIsoNode& NodeC) const;
+	void DisplayCells(const TArray<FCell>& Cells) const;
+	void DisplayCell(const FCell& Cell) const;
+	void DrawCellBoundary(int32 Index, EVisuProperty Property) const;
+
+	void DisplayTriangle(EGridSpace Space, const FIsoNode& NodeA, const FIsoNode& NodeB, const FIsoNode& NodeC) const;
 #endif
 
-	};
+};
 
-	/**
-	 * Criteria to find the optimal "Delaunay" triangle starting from the segment AB to a set of point P
-	 * A "Delaunay" triangle is an equilateral triangle
-	 * The optimal value is the smallest value.
-	 */
-	inline double CotangentCriteria(const FPoint& APoint, const FPoint& BPoint, const FPoint& PPoint, FPoint& OutNormal)
-	{
-		const double BigValue = HUGE_VALUE;
+namespace IsoTriangulatorImpl
+{
 
-		FPoint PA = APoint - PPoint;
-		FPoint PB = BPoint - PPoint;
+/**
+ * Criteria to find the optimal "Delaunay" triangle starting from the segment AB to a set of point P
+ * A "Delaunay" triangle is an equilateral triangle
+ * The optimal value is the smallest value.
+ */
+inline double CotangentCriteria(const FPoint& APoint, const FPoint& BPoint, const FPoint& PPoint, FPoint& OutNormal)
+{
+	const double BigValue = HUGE_VALUE;
 
-		// the ratio between the scalar product PA.PB (=|PA| |PB| cos (A,P,B) )
-		// with the norm of the cross product |PA^PB| (=|PA| |PB| |sin(A,P,B)|)
-		// is compute. 
-		double ScalareProduct = PA * PB;
-		OutNormal = PA ^ PB;
-		double NormOFScalarProduct = sqrt(OutNormal * OutNormal);
+	FPoint PA = APoint - PPoint;
+	FPoint PB = BPoint - PPoint;
 
-		// PPoint is aligned with (A,B)
-		if (NormOFScalarProduct < SMALL_NUMBER)
-		{
-			return BigValue;
-		}
+	// the ratio between the scalar product PA.PB (=|PA| |PB| cos (A,P,B) )
+	// with the norm of the cross product |PA^PB| (=|PA| |PB| |sin(A,P,B)|)
+	// is compute. 
+	double ScalareProduct = PA * PB;
+	OutNormal = PA ^ PB;
+	double NormOFScalarProduct = sqrt(OutNormal * OutNormal);
 
-		// return Cotangent value 
-		return ScalareProduct / NormOFScalarProduct;
+	// PPoint is aligned with (A,B)
+	if (NormOFScalarProduct < SMALL_NUMBER)
+	{
+		return BigValue;
 	}
 
-	inline double CotangentCriteria(const FPoint2D& APoint, const FPoint2D& BPoint, const FPoint2D& PPoint)
-	{
-		const double BigValue = HUGE_VALUE;
+	// return Cotangent value 
+	return ScalareProduct / NormOFScalarProduct;
+}
 
-		FPoint2D PA = APoint - PPoint;
-		FPoint2D PB = BPoint - PPoint;
+inline double CotangentCriteria(const FPoint2D& APoint, const FPoint2D& BPoint, const FPoint2D& PPoint)
+{
+	const double BigValue = HUGE_VALUE;
 
-		// the ratio between the scalar product PA.PB (=|PA| |PB| cos (A,P,B) )
-		// with the norm of the cross product |PA^PB| (=|PA| |PB| |sin(A,P,B)|)
-		// is compute. 
-		double ScalareProduct = PA * PB;
-		double OutNormal = PA ^ PB;
-		double NormOFPointProduct = FMath::Abs(OutNormal);
+	FPoint2D PA = APoint - PPoint;
+	FPoint2D PB = BPoint - PPoint;
 
-		if (NormOFPointProduct < SMALL_NUMBER)
-		{
-			// PPoint is aligned with (A,B)
-			return BigValue;
-		}
+	// the ratio between the scalar product PA.PB (=|PA| |PB| cos (A,P,B) )
+	// with the norm of the cross product |PA^PB| (=|PA| |PB| |sin(A,P,B)|)
+	// is compute. 
+	double ScalareProduct = PA * PB;
+	double OutNormal = PA ^ PB;
+	double NormOFPointProduct = FMath::Abs(OutNormal);
 
-		// return Cotangent value 
-		return ScalareProduct / NormOFPointProduct;
+	if (NormOFPointProduct < SMALL_NUMBER)
+	{
+		// PPoint is aligned with (A,B)
+		return BigValue;
 	}
 
-	template<class PointType>
-	inline double IsoscelesCriteria(const PointType& APoint, const PointType& BPoint, const PointType& IsoscelesVertex)
+	// return Cotangent value 
+	return ScalareProduct / NormOFPointProduct;
+}
+
+template<class PointType>
+inline double IsoscelesCriteria(const PointType& APoint, const PointType& BPoint, const PointType& IsoscelesVertex)
+{
+	double Coord = CoordinateOfProjectedPointOnSegment(IsoscelesVertex, APoint, BPoint, false);
+	return FMath::Abs(Coord - 0.5);
+}
+
+template<class PointType>
+inline double EquilateralCriteria(const PointType& SegmentA, const PointType& SegmentB, const PointType& Point)
+{
+	double Criteria1 = FMath::Abs(CoordinateOfProjectedPointOnSegment(SegmentA, SegmentB, Point, false) - 0.5);
+	double Criteria2 = FMath::Abs(CoordinateOfProjectedPointOnSegment(Point, SegmentA, SegmentB, false) - 0.5);
+	double Criteria3 = FMath::Abs(CoordinateOfProjectedPointOnSegment(SegmentB, Point, SegmentA, false) - 0.5);
+	return Criteria1 + Criteria2 + Criteria3;
+}
+
+inline FLoopNode* GetNextNodeImpl(FLoopNode* Node)
+{
+	return &Node->GetNextNode();
+}
+
+inline FLoopNode* GetPreviousNodeImpl(FLoopNode* Node)
+{
+	return &Node->GetPreviousNode();
+}
+
+inline const FLoopNode* GetNextConstNodeImpl(const FLoopNode* Node)
+{
+	return &Node->GetNextNode();
+}
+
+inline const FLoopNode* GetPreviousConstNodeImpl(const FLoopNode* Node)
+{
+	return &Node->GetPreviousNode();
+}
+
+inline const FLoopNode* GetFirstNode(const FIsoSegment* Segment)
+{
+	return (const FLoopNode*)&Segment->GetFirstNode();
+};
+
+inline const FLoopNode* GetSecondNode(const FIsoSegment* Segment)
+{
+	return (const FLoopNode*)&Segment->GetSecondNode();
+};
+
+inline void GetLoopNodeStartingFrom(FLoopNode* StartNode, bool bForward, TArray<FLoopNode*>& Loop)
+{
+	using namespace IsoTriangulatorImpl;
+	GetNextNodeMethod GetNext = bForward ? GetNextNodeImpl : GetPreviousNodeImpl;
+	FLoopNode* Node = StartNode;
+	Loop.Add(StartNode);
+	for (Node = GetNext(Node); Node != StartNode; Node = GetNext(Node))
 	{
-		double Coord = CoordinateOfProjectedPointOnSegment(IsoscelesVertex, APoint, BPoint, false);
-		return FMath::Abs(Coord - 0.5);
+		Loop.Add(Node);
 	}
+}
+
+inline FLoopNode* GetNodeAt(const TArray<FLoopNode*>& NodesOfLoop, const int32 Index)
+{
+	FLoopNode* Node = NodesOfLoop[Index];
+#ifdef DEBUG_LOOP_INTERSECTION_AND_FIX_IT		
+	Wait(Node->IsDelete());
+#endif
+	ensureCADKernel(!Node->IsDelete());
+	return Node;
+};
+
+inline int32 NextIndex(const int32 NodeCount, int32 StartIndex)
+{
+	++StartIndex;
+	return (StartIndex >= NodeCount) ? 0 : StartIndex;
+};
 
-	template<class PointType>
-	inline double EquilateralCriteria(const PointType& SegmentA, const PointType& SegmentB, const PointType& Point)
+inline void RemoveDeletedNodes(TArray<FLoopNode*>& NodesOfLoop)
+{
+	int32 Index = NodesOfLoop.IndexOfByPredicate([](FLoopNode* Node) { return Node->IsDelete(); });
+	if (Index == INDEX_NONE)
 	{
-		double Criteria1 = FMath::Abs(CoordinateOfProjectedPointOnSegment(SegmentA, SegmentB, Point, false) - 0.5);
-		double Criteria2 = FMath::Abs(CoordinateOfProjectedPointOnSegment(Point, SegmentA, SegmentB, false) - 0.5);
-		double Criteria3 = FMath::Abs(CoordinateOfProjectedPointOnSegment(SegmentB, Point, SegmentA, false) - 0.5);
-		return Criteria1 + Criteria2 + Criteria3;
+		return;
 	}
+	int32 NewIndex = Index;
+	for (; Index < NodesOfLoop.Num(); ++Index)
+	{
+		if (!NodesOfLoop[Index]->IsDelete())
+		{
+			NodesOfLoop[NewIndex++] = NodesOfLoop[Index];
+		}
+	}
+	NodesOfLoop.SetNum(NewIndex);
+}
+
+/**
+ * for findNextSegment
+ */
+inline double ClockwiseSlop(const FPoint2D& StartPoint, const FPoint2D& EndPoint, double ReferenceSlope)
+{
+	return 8. - ComputePositiveSlope(StartPoint, EndPoint, ReferenceSlope);
+}
+
+inline double CounterClockwiseSlop(const FPoint2D& StartPoint, const FPoint2D& EndPoint, double ReferenceSlope)
+{
+	return ComputePositiveSlope(StartPoint, EndPoint, ReferenceSlope);
+}
+
+
+} // namespace FIsoTriangulatorImpl
 
 } // namespace CADKernel
 
diff --git a/Engine/Source/Runtime/Datasmith/CADKernel/Private/CADKernel/Mesh/Meshers/IsoTriangulator/DefineForDebug.h b/Engine/Source/Runtime/Datasmith/CADKernel/Private/CADKernel/Mesh/Meshers/IsoTriangulator/DefineForDebug.h
--- a/Engine/Source/Runtime/Datasmith/CADKernel/Private/CADKernel/Mesh/Meshers/IsoTriangulator/DefineForDebug.h
+++ b/Engine/Source/Runtime/Datasmith/CADKernel/Private/CADKernel/Mesh/Meshers/IsoTriangulator/DefineForDebug.h
@@ -2,21 +2,19 @@
 // Refactor in progress
 #pragma once
 
-
 static int32 FaceToDebug = 0;
-#define DEBUG_ONLY_SURFACE_TO_DEBUG
-
-//#define SURFACE_TO_DEBUG SurfaceToDebug
-//#define SURFACE_TO_DEBUG 0
-static bool BoolDisplayDebugMeshStep=false;
 
-//#define DisplayDebugMeshSurfaceBySurfaceStep
+//#define ADD_TRIANGLE_2D
+//#define DEBUG_ISOTRIANGULATOR
+//
+//#define DEBUG_LOOP_INTERSECTION_AND_FIX_IT		
+//#define DEBUG_REMOVE_LOOP_INTERSECTIONS
+//#define DEBUG_REMOVE_UNIQUE_INTERSECTION
+//#define DEBUG_SWAP_SEGMENTS_OR_REMOVE
+//#define DEBUG_FIND_LOOP_INTERSECTION_AND_FIX_IT
+//#define DEBUG_FIND_LOOP_INTERSECTIONS		
 
-//#define DISPLAY_CRITERIA_GRID
-
-//#define FIND_CYCLE
-//#define DEBUG_FIND_CYCLE
 
+#define DEBUG_ONLY_SURFACE_TO_DEBUG
+static bool BoolDisplayDebugMeshStep=false;
 
-//#define DEBUG_MESH_EDGE
-//#define DEBUG_GetPreferredUVCoordinatesFromNeighbours
\ No newline at end of file
diff --git a/Engine/Source/Runtime/Datasmith/CADKernel/Private/CADKernel/Mesh/Meshers/IsoTriangulator/IntersectionSegmentTool.cpp b/Engine/Source/Runtime/Datasmith/CADKernel/Private/CADKernel/Mesh/Meshers/IsoTriangulator/IntersectionSegmentTool.cpp
--- a/Engine/Source/Runtime/Datasmith/CADKernel/Private/CADKernel/Mesh/Meshers/IsoTriangulator/IntersectionSegmentTool.cpp
+++ b/Engine/Source/Runtime/Datasmith/CADKernel/Private/CADKernel/Mesh/Meshers/IsoTriangulator/IntersectionSegmentTool.cpp
@@ -16,6 +16,40 @@ const FIsoSegment* FIntersectionSegmentTool::DoesIntersect(const FIsoSegment& Se
 	return DoesIntersect(Segment.GetFirstNode(), Segment.GetSecondNode());
 }
 
+FIsoSegment* FIntersectionSegmentTool::DoesIntersect(const FIsoSegment& Segment)
+{
+	return const_cast<FIsoSegment*> (static_cast<const FIntersectionSegmentTool*>(this)->DoesIntersect(Segment.GetFirstNode(), Segment.GetSecondNode()));
+}
+
+bool FIntersectionSegmentTool::FindIntersections(const FIsoNode& StartNode, const FIsoNode& EndNode, TArray<const FIsoSegment*>& OutIntersectedSegments) const
+{
+	const FPoint2D& StartPoint = StartNode.Get2DPoint(EGridSpace::UniformScaled, Grid);
+	const FPoint2D& EndPoint = EndNode.Get2DPoint(EGridSpace::UniformScaled, Grid);
+	const FSegment4IntersectionTools StartEndSegment(StartPoint, EndPoint);
+
+	OutIntersectedSegments.Empty(10);
+
+	for (const FSegment4IntersectionTools& Segment : Segments)
+	{
+		if (&Segment.IsoSegment->GetFirstNode() == &StartNode || &Segment.IsoSegment->GetSecondNode() == &StartNode)
+		{
+			continue;
+		}
+
+		if (&Segment.IsoSegment->GetFirstNode() == &EndNode || &Segment.IsoSegment->GetSecondNode() == &EndNode)
+		{
+			continue;
+		}
+
+		if (IntersectSegments2D(Segment.Segment2D, StartEndSegment.Segment2D))
+		{
+			OutIntersectedSegments.Add(Segment.IsoSegment);
+		}
+	}
+
+	return OutIntersectedSegments.Num() > 0;
+}
+
 const FIsoSegment* FIntersectionSegmentTool::DoesIntersect(const FIsoNode& StartNode, const FIsoNode& EndNode) const
 {
 	const FPoint2D& StartPoint = StartNode.Get2DPoint(EGridSpace::UniformScaled, Grid);
@@ -26,6 +60,11 @@ const FIsoSegment* FIntersectionSegmentTool::DoesIntersect(const FIsoNode& Start
 
 	for (const FSegment4IntersectionTools& Segment : Segments)
 	{
+		if (Segment.IsoSegment != nullptr && Segment.IsoSegment->IsDelete())
+		{
+			continue;
+		}
+
 		if (bSegmentsAreSorted && AxisMax < Segment.AxisMin)
 		{
 			break;
@@ -111,7 +150,7 @@ bool FIntersectionSegmentTool::DoesIntersect(const FPoint2D& StartPoint, const F
 
 FSegment4IntersectionTools::FSegment4IntersectionTools(const FGrid& Grid, const FIsoSegment& InSegment)
 	: Segment2D(InSegment.GetFirstNode().Get2DPoint(EGridSpace::UniformScaled, Grid), InSegment.GetSecondNode().Get2DPoint(EGridSpace::UniformScaled, Grid))
-	, Boundary(InSegment.GetFirstNode().Get2DPoint(EGridSpace::UniformScaled, Grid), InSegment.GetSecondNode().Get2DPoint(EGridSpace::UniformScaled, Grid))
+	, Boundary(Segment2D.Point0, Segment2D.Point1)
 	, IsoSegment(&InSegment)
 {
 	AxisMin = Boundary[EIso::IsoU].Min + Boundary[EIso::IsoV].Min;
diff --git a/Engine/Source/Runtime/Datasmith/CADKernel/Private/CADKernel/Mesh/Meshers/IsoTriangulator/IntersectionSegmentTool.h b/Engine/Source/Runtime/Datasmith/CADKernel/Private/CADKernel/Mesh/Meshers/IsoTriangulator/IntersectionSegmentTool.h
--- a/Engine/Source/Runtime/Datasmith/CADKernel/Private/CADKernel/Mesh/Meshers/IsoTriangulator/IntersectionSegmentTool.h
+++ b/Engine/Source/Runtime/Datasmith/CADKernel/Private/CADKernel/Mesh/Meshers/IsoTriangulator/IntersectionSegmentTool.h
@@ -5,6 +5,7 @@
 #include "CADKernel/Math/Boundary.h"
 #include "CADKernel/Math/Geometry.h"
 #include "CADKernel/UI/Message.h"
+#include "CADKernel/UI/Visu.h"
 
 namespace CADKernel
 {
@@ -85,9 +86,38 @@ namespace CADKernel
 			Segments.Reserve(InMaxNum);
 		}
 
-		void SetNum(int32 NewMaxNum)
+		void RemoveLast()
 		{
-			Segments.Reserve(NewMaxNum);
+			ensureCADKernel(!bSegmentsAreSorted);
+			Segments.RemoveAt(Segments.Num() - 1);
+			bSegmentsAreSorted = false;
+		}
+
+		void SetCount(int32 NewCount)
+		{
+			ensureCADKernel(NewCount < Segments.Num() && !bSegmentsAreSorted);
+			while(NewCount != Segments.Num())
+			{
+				Segments.RemoveAt(Segments.Num() - 1);
+			}
+		}
+
+		/**
+		 * @return true if the segment is found and removed
+		 */
+		bool Remove(const FIsoSegment* Segment)
+		{
+			int32 SegmentIndex = Segments.IndexOfByPredicate([Segment](const FSegment4IntersectionTools& SegmentIter)
+				{
+					return SegmentIter.IsoSegment == Segment;
+				});
+
+			if (SegmentIndex != INDEX_NONE)
+			{
+				Segments.RemoveAt(SegmentIndex);
+				return true;
+			}
+			return false;
 		}
 
 		int32 Count()
@@ -95,6 +125,24 @@ namespace CADKernel
 			return Segments.Num();
 		}
 
+		bool Update(const FIsoSegment* Segment)
+		{
+			int32 SegmentIndex = Segments.IndexOfByPredicate([Segment](const FSegment4IntersectionTools& SegmentIter)
+				{
+					return SegmentIter.IsoSegment == Segment;
+				});
+
+			if (SegmentIndex != INDEX_NONE)
+			{
+				Segments.RemoveAt(SegmentIndex);
+				Segments.EmplaceAt(SegmentIndex, Grid, *Segment);
+				bSegmentsAreSorted = false;
+				return true;
+			}
+			return false;
+
+		}
+
 		void AddSegments(FIsoSegment** InNewSegments, int32 Count)
 		{
 			bSegmentsAreSorted = false;
@@ -129,6 +177,7 @@ namespace CADKernel
 		}
 
 		const FIsoSegment* DoesIntersect(const FIsoSegment& Segment) const;
+		FIsoSegment* DoesIntersect(const FIsoSegment& Segment);
 
 		/**
 		 * WARNING StartPoint, EndPoint must be defined in EGridSpace::UniformScaled
@@ -146,6 +195,8 @@ namespace CADKernel
 		 */
 		const FIsoSegment* DoesIntersect(const FIsoNode& StartNode, const FIsoNode& EndNode) const;
 
+		bool FindIntersections(const FIsoNode& StartNode, const FIsoNode& EndNode, TArray<const FIsoSegment*>& OutIntersections) const;
+
 		/**
 		 * segments are sorted by DMin increasing
 		 */
@@ -156,8 +207,9 @@ namespace CADKernel
 		}
 
 #ifdef CADKERNEL_DEV
-		void Display(const TCHAR* Message) const;
+		void Display(const TCHAR* Message, EVisuProperty Property = EVisuProperty::BlueCurve) const;
 #endif
+
 	};
 } // namespace CADKernel
 
diff --git a/Engine/Source/Runtime/Datasmith/CADKernel/Private/CADKernel/Mesh/Meshers/IsoTriangulator/IsoNode.h b/Engine/Source/Runtime/Datasmith/CADKernel/Private/CADKernel/Mesh/Meshers/IsoTriangulator/IsoNode.h
--- a/Engine/Source/Runtime/Datasmith/CADKernel/Private/CADKernel/Mesh/Meshers/IsoTriangulator/IsoNode.h
+++ b/Engine/Source/Runtime/Datasmith/CADKernel/Private/CADKernel/Mesh/Meshers/IsoTriangulator/IsoNode.h
@@ -42,6 +42,8 @@ enum class EIsoNodeStates : uint16
 	NearlyIsoU = 0x1000u,  // Loop nodes
 	NearlyIsoV = 0x2000u,  // Loop nodes
 
+	Delete = 0x8000u,
+
 	All = 0xFFFFu
 };
 
@@ -76,6 +78,19 @@ public:
 
 	virtual ~FIsoNode() = default;
 
+	virtual void Delete()
+	{
+		ConnectedSegments.Empty();
+		FaceIndex = -1;
+
+		States = EIsoNodeStates::Delete;
+	}
+
+	bool IsDelete() const 
+	{
+		return (States & EIsoNodeStates::Delete) == EIsoNodeStates::Delete;
+	}
+
 	const int32 GetIndex() const
 	{
 		return Index;
@@ -98,23 +113,14 @@ public:
 
 	FIsoSegment* GetSegmentConnectedTo(const FIsoNode* Node) const;
 
-	void RemoveSegment(FIsoSegment* const Segment)
+	void ConnectSegment(FIsoSegment& Segment)
 	{
-		int32 SegmentIndex = 0;
-		if (ConnectedSegments.Find(Segment, SegmentIndex))
-		{
-			ConnectedSegments.RemoveAt(SegmentIndex);
-		}
-	}
-
-	void ConnectSegment(FIsoSegment* Segment)
-	{
-		ConnectedSegments.Add(Segment);
+		ConnectedSegments.Add(&Segment);
 	}
 
-	void DisconnectSegment(FIsoSegment* Segment)
+	void DisconnectSegment(FIsoSegment& Segment)
 	{
-		ConnectedSegments.Remove(Segment);
+		ConnectedSegments.RemoveSingle(&Segment);
 	}
 
 	const virtual bool IsALoopNode() const = 0;
@@ -230,6 +236,9 @@ public:
 	virtual bool operator==(const FIsoNode& OtherNode) const = 0;
 	virtual bool IsEqualTo(const FLoopNode& OtherNode) const = 0;
 	virtual bool IsEqualTo(const FIsoInnerNode& OtherNode) const = 0;
+
+	virtual uint32 GetTypeHash() const = 0;
+
 };
 
 /**
@@ -249,6 +258,15 @@ public:
 	{
 	};
 
+	virtual void Delete() override
+	{
+		FIsoNode::Delete();
+
+		//LoopIndex = -1;
+		ConnectedLoopNodes[0] = nullptr;
+		ConnectedLoopNodes[1] = nullptr;
+	}
+
 	virtual bool IsEqualTo(const FLoopNode& OtherNode) const
 	{
 		return ((GetIndex() == OtherNode.GetIndex()) && (GetLoopIndex() == OtherNode.GetLoopIndex()));
@@ -293,6 +311,16 @@ public:
 		return *ConnectedLoopNodes[1];
 	}
 
+	FIsoSegment& GetPreviousSegment() const
+	{
+		return *GetSegmentConnectedTo(ConnectedLoopNodes[0]);
+	}
+
+	FIsoSegment& GetNextSegment() const
+	{
+		return *GetSegmentConnectedTo(ConnectedLoopNodes[1]);
+	}
+
 	/**
 	 * Return the 2d coordinate of the node according to the space
 	 */
@@ -395,8 +423,17 @@ public:
 		}
 	}
 
+	virtual uint32 GetTypeHash() const override
+	{
+		return HashCombine(GetIndex(), GetLoopIndex() + 1);
+	}
 };
 
+inline uint32 GetTypeHash(const FIsoNode& Node)
+{
+	return Node.GetTypeHash();
+}
+
 class FIsoInnerNode : public FIsoNode
 {
 public:
@@ -501,6 +538,10 @@ public:
 		Id += StartId;
 	}
 
+	virtual uint32 GetTypeHash() const override
+	{
+		return ::GetTypeHash(GetIndex());
+	}
 };
 
 
diff --git a/Engine/Source/Runtime/Datasmith/CADKernel/Private/CADKernel/Mesh/Meshers/IsoTriangulator/IsoSegment.h b/Engine/Source/Runtime/Datasmith/CADKernel/Private/CADKernel/Mesh/Meshers/IsoTriangulator/IsoSegment.h
--- a/Engine/Source/Runtime/Datasmith/CADKernel/Private/CADKernel/Mesh/Meshers/IsoTriangulator/IsoSegment.h
+++ b/Engine/Source/Runtime/Datasmith/CADKernel/Private/CADKernel/Mesh/Meshers/IsoTriangulator/IsoSegment.h
@@ -7,312 +7,344 @@
 
 namespace CADKernel
 {
-	enum class ESegmentType : uint8
-	{
-		Unknown = 0,
-		Loop = 1,
-		ThinZone = 2,
-		IsoU,
-		IsoV,
-		LoopToLoop, // IsBoundaryToInner use that xxxToBoundary are bigger than BoundaryToBoundary
-		IsoLoopToLoop,
-		InnerToLoopU,
-		InnerToLoopV,
-		InnerToLoop,
-	};
 
-	enum class EIsoSegmentStates : uint8
-	{
-		None            = 0x00u,	// No flags.
+enum class ESegmentType : uint8
+{
+	Unknown = 0,
+	Loop = 1,
+	ThinZone = 2,
+	IsoU,
+	IsoV,
+	LoopToLoop, // IsBoundaryToInner use that xxxToBoundary are bigger than BoundaryToBoundary
+	IsoLoopToLoop,
+	InnerToLoopU,
+	InnerToLoopV,
+	InnerToLoop,
+};
+
+enum class EIsoSegmentStates : uint8
+{
+	None = 0x00u,	// No flags.
 
-		CandidateFlag   = 0x01u,
+	CandidateFlag = 0x01u,
 
-		LeftCycle       = 0x02u,
-		RightCycle      = 0x04u,
+	LeftCycle = 0x02u,
+	RightCycle = 0x04u,
 
-		Degenerate      = 0x08u,
+	Degenerate = 0x08u,
 
-		LeftTriangle    = 0x10u,
-		RightTriangle   = 0x20u,
-		SegmentComplete = 0x30u,
+	LeftTriangle = 0x10u,
+	RightTriangle = 0x20u,
+	SegmentComplete = 0x30u,
 
-		StatusIsoU      = 0x40u,
-		StatusIsoV      = 0x80u,
+	StatusIsoU = 0x40u,
+	StatusIsoV = 0x80u,
 
-		All             = 0xFFu
-	};
+	Delete = 0xFFu,
+	All = 0xFFu
+};
 
-	ENUM_CLASS_FLAGS(EIsoSegmentStates);
+ENUM_CLASS_FLAGS(EIsoSegmentStates);
 
 
-	class FGrid;
+class FGrid;
 
-	class FIsoSegment
+class FIsoSegment
+{
+protected:
+	FIsoNode* FirstNode;
+	FIsoNode* SecondNode;
+	ESegmentType Type;
+	EIsoSegmentStates States;
+
+public:
+
+	FIsoSegment()
+		: FirstNode(nullptr)
+		, SecondNode(nullptr)
+		, Type(ESegmentType::Unknown)
+		, States(EIsoSegmentStates::None)
 	{
-	protected:
-		FIsoNode* FirstNode;
-		FIsoNode* SecondNode;
-		ESegmentType Type;
-		EIsoSegmentStates States;
-
-	public:
+	}
 
-		FIsoSegment()
-			: FirstNode(nullptr)
-			, SecondNode(nullptr)
-			, Type(ESegmentType::Unknown)
-			, States(EIsoSegmentStates::None)
-		{
+	void Init(FIsoNode& InFirstNode, FIsoNode& InSecondNode, const ESegmentType InType)
+	{
+		States = EIsoSegmentStates::None;
+		FirstNode = &InFirstNode;
+		SecondNode = &InSecondNode;
+
+		Type = InType;
+		switch (Type)
+		{
+		case ESegmentType::IsoU:
+			InFirstNode.SetLinkedToIso(EIsoLink::IsoUNext);
+			InSecondNode.SetLinkedToIso(EIsoLink::IsoUPrevious);
+			break;
+		case ESegmentType::IsoV:
+			InFirstNode.SetLinkedToIso(EIsoLink::IsoVNext);
+			InSecondNode.SetLinkedToIso(EIsoLink::IsoVPrevious);
+			break;
 		}
+	}
 
-		void Init(FIsoNode& InFirstNode, FIsoNode& InSecondNode, const ESegmentType InType)
+	void ConnectToNode()
+	{
+		TFunction<bool()> CheckExistingSegment = [&]()
 		{
-			States = EIsoSegmentStates::None;
-			FirstNode = &InFirstNode;
-			SecondNode = &InSecondNode;
-
-			Type = InType;
-			switch (Type)
+			if (IsItAlreadyDefined(FirstNode, SecondNode))
 			{
-			case ESegmentType::IsoU:
-				InFirstNode.SetLinkedToIso(EIsoLink::IsoUNext);
-				InSecondNode.SetLinkedToIso(EIsoLink::IsoUPrevious);
-				break;
-			case ESegmentType::IsoV:
-				InFirstNode.SetLinkedToIso(EIsoLink::IsoVNext);
-				InSecondNode.SetLinkedToIso(EIsoLink::IsoVPrevious);
-				break;
+				//Wait();
+				return false;
 			}
-		}
-
-		void ConnectToNode()
-		{
-			TFunction<bool()> CheckExistingSegment = [&]()
-			{
-				if(IsItAlreadyDefined(FirstNode, SecondNode))
-				{
-					//Wait();
-					return false;
-				}
-				return true;
-			};
+			return true;
+		};
 
-			ensureCADKernel(CheckExistingSegment());
+		ensureCADKernel(CheckExistingSegment());
 
-			FirstNode->ConnectSegment(this);
-			SecondNode->ConnectSegment(this);
-		}
+		FirstNode->ConnectSegment(*this);
+		SecondNode->ConnectSegment(*this);
+	}
 
-		static bool	IsItAlreadyDefined(const FIsoNode* StartNode, const FIsoNode* EndNode)
+	static bool	IsItAlreadyDefined(const FIsoNode* StartNode, const FIsoNode* EndNode)
+	{
+		for (const FIsoSegment* Segment : StartNode->GetConnectedSegments())
 		{
-			for (const FIsoSegment* Segment : StartNode->GetConnectedSegments())
+			const FIsoNode& OtherNode = (StartNode == &Segment->GetFirstNode()) ? Segment->GetSecondNode() : Segment->GetFirstNode();
+			if (&OtherNode == EndNode)
 			{
-				const FIsoNode& OtherNode = (StartNode == &Segment->GetFirstNode()) ? Segment->GetSecondNode() : Segment->GetFirstNode();
-				if (&OtherNode == EndNode)
-				{
-					return true;
-				}
+				return true;
 			}
-			return false;
 		}
+		return false;
+	}
 
-		void Clean()
-		{
-			FirstNode = nullptr;
-			SecondNode = nullptr;
-		}
+	void Clean()
+	{
+		FirstNode = nullptr;
+		SecondNode = nullptr;
+		States = EIsoSegmentStates::Delete;
+	}
 
-		void SetSelected()
-		{
-			States &= ~EIsoSegmentStates::CandidateFlag;
-		}
+	bool IsDelete() const
+	{
+		return States == EIsoSegmentStates::Delete;
+	}
 
-		void SetCandidate()
-		{
-			States |= EIsoSegmentStates::CandidateFlag;
-		}
+	void SetSelected()
+	{
+		States &= ~EIsoSegmentStates::CandidateFlag;
+	}
 
-		bool IsACandidate() const // [0 Pi/2]
-		{
-			return (States & EIsoSegmentStates::CandidateFlag) == EIsoSegmentStates::CandidateFlag;
-		}
+	void SetCandidate()
+	{
+		States |= EIsoSegmentStates::CandidateFlag;
+	}
 
-		bool HasCycleOnLeft()
-		{
-			return (States & EIsoSegmentStates::LeftCycle) == EIsoSegmentStates::LeftCycle;
-		}
+	bool IsACandidate() const // [0 Pi/2]
+	{
+		return (States & EIsoSegmentStates::CandidateFlag) == EIsoSegmentStates::CandidateFlag;
+	}
 
-		void SetHaveCycleOnLeft()
-		{
-			States |= EIsoSegmentStates::LeftCycle;
-		}
+	bool HasCycleOnLeft()
+	{
+		return (States & EIsoSegmentStates::LeftCycle) == EIsoSegmentStates::LeftCycle;
+	}
 
-		bool HasCycleOnRight()
-		{
-			return (States & EIsoSegmentStates::RightCycle) == EIsoSegmentStates::RightCycle;
-		}
+	void SetHaveCycleOnLeft()
+	{
+		States |= EIsoSegmentStates::LeftCycle;
+	}
 
-		void SetHaveCycleOnRight()
-		{
-			States |= EIsoSegmentStates::RightCycle;
-		}
+	bool HasCycleOnRight()
+	{
+		return (States & EIsoSegmentStates::RightCycle) == EIsoSegmentStates::RightCycle;
+	}
 
-		bool HasTriangleOnLeft()
-		{
-			return (States & EIsoSegmentStates::LeftTriangle) == EIsoSegmentStates::LeftTriangle;
-		}
+	void SetHaveCycleOnRight()
+	{
+		States |= EIsoSegmentStates::RightCycle;
+	}
 
-		void SetHasTriangleOnLeft()
-		{
-			States |= EIsoSegmentStates::LeftTriangle;
-		}
+	bool HasTriangleOnLeft()
+	{
+		return (States & EIsoSegmentStates::LeftTriangle) == EIsoSegmentStates::LeftTriangle;
+	}
 
-		bool HasTriangleOnRight()
-		{
-			return (States & EIsoSegmentStates::RightTriangle) == EIsoSegmentStates::RightTriangle;
-		}
+	void SetHasTriangleOnLeft()
+	{
+		States |= EIsoSegmentStates::LeftTriangle;
+	}
 
-		bool HasTriangleOnRightAndLeft()
-		{
-			return (States & EIsoSegmentStates::RightTriangle) == EIsoSegmentStates::RightTriangle && (States & EIsoSegmentStates::LeftTriangle) == EIsoSegmentStates::LeftTriangle;
-		}
+	bool HasTriangleOnRight()
+	{
+		return (States & EIsoSegmentStates::RightTriangle) == EIsoSegmentStates::RightTriangle;
+	}
 
-		void SetHasTriangleOnRight()
-		{
-			States |= EIsoSegmentStates::RightTriangle;
-		}
+	bool HasTriangleOnRightAndLeft()
+	{
+		return (States & EIsoSegmentStates::RightTriangle) == EIsoSegmentStates::RightTriangle && (States & EIsoSegmentStates::LeftTriangle) == EIsoSegmentStates::LeftTriangle;
+	}
 
-		void ResetHasTriangle()
-		{
-			States &= ~EIsoSegmentStates::SegmentComplete;
-		}
+	void SetHasTriangleOnRight()
+	{
+		States |= EIsoSegmentStates::RightTriangle;
+	}
 
-		void SetAsIsoU()
-		{
-			States |= EIsoSegmentStates::StatusIsoU;
-		}
+	void ResetHasTriangle()
+	{
+		States &= ~EIsoSegmentStates::SegmentComplete;
+	}
 
-		bool IsIsoU() const
-		{
-			return (States & EIsoSegmentStates::StatusIsoU) == EIsoSegmentStates::StatusIsoU;
-		}
+	void SetAsIsoU()
+	{
+		States |= EIsoSegmentStates::StatusIsoU;
+	}
 
-		void SetAsIsoV()
-		{
-			States |= EIsoSegmentStates::StatusIsoV;
-		}
+	bool IsIsoU() const
+	{
+		return (States & EIsoSegmentStates::StatusIsoU) == EIsoSegmentStates::StatusIsoU;
+	}
 
-		bool IsIsoV() const
-		{
-			return (States & EIsoSegmentStates::StatusIsoV) == EIsoSegmentStates::StatusIsoV;
-		}
+	void SetAsIsoV()
+	{
+		States |= EIsoSegmentStates::StatusIsoV;
+	}
 
-		void SetAsDegenerated()
-		{
-			States |= EIsoSegmentStates::Degenerate;
-		}
+	bool IsIsoV() const
+	{
+		return (States & EIsoSegmentStates::StatusIsoV) == EIsoSegmentStates::StatusIsoV;
+	}
 
-		bool IsDegenerated() const
-		{
-			return (States & EIsoSegmentStates::Degenerate) == EIsoSegmentStates::Degenerate;
-		}
+	void SetAsDegenerated()
+	{
+		States |= EIsoSegmentStates::Degenerate;
+	}
 
-		const FIsoNode& GetFirstNode() const
-		{
-			return *FirstNode;
-		}
+	bool IsDegenerated() const
+	{
+		return (States & EIsoSegmentStates::Degenerate) == EIsoSegmentStates::Degenerate;
+	}
 
-		FIsoNode& GetFirstNode()
-		{
-			return *FirstNode;
-		}
+	const FIsoNode& GetFirstNode() const
+	{
+		return *FirstNode;
+	}
 
-		const FIsoNode& GetSecondNode() const
-		{
-			return *SecondNode;
-		}
+	FIsoNode& GetFirstNode()
+	{
+		return *FirstNode;
+	}
 
-		FIsoNode& GetSecondNode()
-		{
-			return *SecondNode;
-		}
+	const FIsoNode& GetSecondNode() const
+	{
+		return *SecondNode;
+	}
 
-		void SetFirstNode(FIsoNode& NewNode)
-		{
-			FirstNode = &NewNode;
+	FIsoNode& GetSecondNode()
+	{
+		return *SecondNode;
+	}
 
-			((FLoopNode*)SecondNode)->SetPreviousConnectedNode((FLoopNode*)FirstNode);
-			((FLoopNode*)FirstNode)->SetNextConnectedNode((FLoopNode*)SecondNode);
-		}
+	void SetFirstNode(FIsoNode& NewNode)
+	{
+		FirstNode = &NewNode;
 
-		void SetSecondNode(FIsoNode& NewNode)
-		{
-			SecondNode = &NewNode;
-			
-			((FLoopNode*)SecondNode)->SetPreviousConnectedNode((FLoopNode*)FirstNode);
-			((FLoopNode*)FirstNode)->SetNextConnectedNode((FLoopNode*)SecondNode);
-		}
+		((FLoopNode*)SecondNode)->SetPreviousConnectedNode((FLoopNode*)FirstNode);
+		((FLoopNode*)FirstNode)->SetNextConnectedNode((FLoopNode*)SecondNode);
+	}
 
-		void SwapOrientation()
-		{
-			FIsoNode* TempNode = SecondNode;
-			SecondNode = FirstNode;
-			FirstNode = TempNode;
+	void SetSecondNode(FIsoNode& NewNode)
+	{
+		SecondNode = &NewNode;
 
-			((FLoopNode*)FirstNode)->SetNextConnectedNode((FLoopNode*)SecondNode);
-			((FLoopNode*)SecondNode)->SetPreviousConnectedNode((FLoopNode*)FirstNode);
-		}
+		((FLoopNode*)SecondNode)->SetPreviousConnectedNode((FLoopNode*)FirstNode);
+		((FLoopNode*)FirstNode)->SetNextConnectedNode((FLoopNode*)SecondNode);
+	}
 
-		const FIsoNode& GetOtherNode(const FIsoNode* Node) const
-		{
-			return (FirstNode == Node) ? *SecondNode : *FirstNode;
-		}
+	void SwapOrientation()
+	{
+		FIsoNode* TempNode = SecondNode;
+		SecondNode = FirstNode;
+		FirstNode = TempNode;
 
-		FIsoNode& GetOtherNode(const FIsoNode* Node)
-		{
-			return (FirstNode == Node) ? *SecondNode : *FirstNode;
-		}
+		((FLoopNode*)FirstNode)->SetNextConnectedNode((FLoopNode*)SecondNode);
+		((FLoopNode*)SecondNode)->SetPreviousConnectedNode((FLoopNode*)FirstNode);
+	}
 
-		const ESegmentType GetType() const
-		{
-			return Type;
-		}
+	const FIsoNode& GetOtherNode(const FIsoNode* Node) const
+	{
+		return (FirstNode == Node) ? *SecondNode : *FirstNode;
+	}
 
-		bool IsLoopToInner() const
-		{
-			return (Type >= ESegmentType::LoopToLoop);
-		}
+	FIsoNode& GetOtherNode(const FIsoNode* Node)
+	{
+		return (FirstNode == Node) ? *SecondNode : *FirstNode;
+	}
 
-		double Get2DLengthSquare(EGridSpace Space,  const FGrid& Grid) const
-		{
-			const FPoint2D& Point1 = FirstNode->Get2DPoint(Space, Grid);
-			const FPoint2D& Point2 = SecondNode->Get2DPoint(Space, Grid);
-			return Point1.SquareDistance(Point2);
-		}
+	const ESegmentType GetType() const
+	{
+		return Type;
+	}
 
-		double Get3DLengthSquare(const FGrid& Grid) const
-		{
-			const FPoint& Point1 = FirstNode->Get3DPoint(Grid);
-			const FPoint& Point2 = SecondNode->Get3DPoint(Grid);
-			return Point1.SquareDistance(Point2);
-		}
-	};
+	double Get2DLengthSquare(EGridSpace Space, const FGrid& Grid) const
+	{
+		const FPoint2D& Point1 = FirstNode->Get2DPoint(Space, Grid);
+		const FPoint2D& Point2 = SecondNode->Get2DPoint(Space, Grid);
+		return Point1.SquareDistance(Point2);
+	}
 
-	inline FIsoSegment* FIsoNode::GetSegmentConnectedTo(const FIsoNode* Node) const
+	double Get3DLengthSquare(const FGrid& Grid) const
 	{
-		for (FIsoSegment* Segment : ConnectedSegments)
+		const FPoint& Point1 = FirstNode->Get3DPoint(Grid);
+		const FPoint& Point2 = SecondNode->Get3DPoint(Grid);
+		return Point1.SquareDistance(Point2);
+	}
+};
+
+inline uint32 GetTypeHash(const FIsoSegment& Segment)
+{
+	uint32 StartHash = GetTypeHash(Segment.GetFirstNode());
+	uint32 EndHash = GetTypeHash(Segment.GetSecondNode());
+	if (StartHash < EndHash)
+	{
+		return HashCombine(StartHash, EndHash);
+	}
+	else
+	{
+		return HashCombine(EndHash, StartHash);
+	}
+};
+
+inline uint32 GetTypeHash(const FIsoSegment& Segment0, const FIsoSegment& Segment1)
+{
+	uint32 StartHash = GetTypeHash(Segment0);
+	uint32 EndHash = GetTypeHash(Segment1);
+	if (StartHash < EndHash)
+	{
+		return HashCombine(StartHash, EndHash);
+	}
+	else
+	{
+		return HashCombine(EndHash, StartHash);
+	}
+};
+
+
+inline FIsoSegment* FIsoNode::GetSegmentConnectedTo(const FIsoNode* Node) const
+{
+	for (FIsoSegment* Segment : ConnectedSegments)
+	{
+		if (&Segment->GetFirstNode() == Node)
 		{
-			if (&Segment->GetFirstNode() == Node)
-			{
-				return Segment;
-			}
-			if (&Segment->GetSecondNode() == Node)
-			{
-				return Segment;
-			}
+			return Segment;
+		}
+		if (&Segment->GetSecondNode() == Node)
+		{
+			return Segment;
 		}
-		return nullptr;
 	}
+	return nullptr;
+}
 
 } // namespace CADKernel
 
diff --git a/Engine/Source/Runtime/Datasmith/CADKernel/Private/CADKernel/Mesh/Meshers/IsoTriangulator/IsoTriangulatorDisplay.cpp b/Engine/Source/Runtime/Datasmith/CADKernel/Private/CADKernel/Mesh/Meshers/IsoTriangulator/IsoTriangulatorDisplay.cpp
--- a/Engine/Source/Runtime/Datasmith/CADKernel/Private/CADKernel/Mesh/Meshers/IsoTriangulator/IsoTriangulatorDisplay.cpp
+++ b/Engine/Source/Runtime/Datasmith/CADKernel/Private/CADKernel/Mesh/Meshers/IsoTriangulator/IsoTriangulatorDisplay.cpp
@@ -41,6 +41,21 @@ void FIsoTriangulator::Display(EGridSpace Space, const FIsoNode& Node, FIdent Id
 	return DisplayPoint(Node.GetPoint(Space, Grid), Property, Ident);
 }
 
+void FIsoTriangulator::Display(FString Message, EGridSpace Space, const TArray<const FIsoNode*>& Nodes, EVisuProperty Property) const
+{
+	if (!bDisplay)
+	{
+		return;
+	}
+
+	Open3DDebugSession(Message);
+	for (int32 Index = 0; Index < Nodes.Num(); ++Index)
+	{
+		Display(Space, *Nodes[Index], Nodes[Index]->GetIndex(), Property);
+	}
+	Close3DDebugSession();
+}
+
 void FIsoTriangulator::Display(EGridSpace Space, const TCHAR* Message, const TArray<FIsoSegment*>& Segments, bool bDisplayNode, bool bDisplayOrientation, EVisuProperty Property) const
 {
 	if (!bDisplay)
@@ -72,14 +87,14 @@ void FIsoTriangulator::Display(EGridSpace Space, const TCHAR* Message, const TAr
 	{
 		for (FIsoSegment* Segment : Segments)
 		{
-			Display(Space, Segment->GetFirstNode(), Segment->GetFirstNode().GetFaceIndex());
-			Display(Space, Segment->GetSecondNode(), Segment->GetSecondNode().GetFaceIndex());
+			Display(Space, Segment->GetFirstNode(), Segment->GetFirstNode().GetId(), (EVisuProperty)(Property - 1));
+			Display(Space, Segment->GetSecondNode(), Segment->GetSecondNode().GetId(), (EVisuProperty)(Property - 1));
 		}
 	}
 	//Wait();
 }
 
-void FIsoTriangulator::DisplayLoops(EGridSpace Space, const TCHAR* Message, const TArray<FLoopNode>& Nodes, bool bDisplayNode, EVisuProperty Property) const
+void FIsoTriangulator::DisplayLoops(EGridSpace Space, const TCHAR* Message, const TArray<FLoopNode>& Nodes, bool bDisplayNode, EVisuProperty NodeVisuProperty) const
 {
 	if (!bDisplay)
 	{
@@ -90,8 +105,9 @@ void FIsoTriangulator::DisplayLoops(EGridSpace Space, const TCHAR* Message, cons
 
 	for (const FLoopNode& Node : Nodes)
 	{
+		if (!Node.IsDelete())
 		{
-			Display(Space, Node, Node.GetNextNode(), 0, Property);
+			Display(Space, Node, Node.GetNextNode(), 0, NodeVisuProperty);
 		}
 	}
 
@@ -99,7 +115,48 @@ void FIsoTriangulator::DisplayLoops(EGridSpace Space, const TCHAR* Message, cons
 	{
 		for (const FLoopNode& Node : Nodes)
 		{
-			Display(Space, Node, Node.GetFaceIndex());
+			if (Node.IsDelete())
+			{
+				Display(Space, Node, Node.GetFaceIndex(), EVisuProperty::PurplePoint);
+			}
+			else
+			{
+				Display(Space, Node, Node.GetFaceIndex(), (EVisuProperty) ((int32)NodeVisuProperty + 1));
+			}
+		}
+	}
+}
+
+void FIsoTriangulator::DisplayLoop(EGridSpace Space, const TCHAR* Message, const TArray<FLoopNode*>& Nodes, bool bDisplayNode, EVisuProperty NodeVisuProperty) const
+{
+	if (!bDisplay)
+	{
+		return;
+	}
+
+	F3DDebugSession _(Message);
+
+	int32 Index = 0;
+	for (const FLoopNode* Node : Nodes)
+	{
+		if (!Node->IsDelete())
+		{
+			Display(Space, *Node, Node->GetNextNode(), Index++, (EVisuProperty)((int32)NodeVisuProperty + 1));
+		}
+	}
+
+	if (bDisplayNode)
+	{
+		for (const FLoopNode* Node : Nodes)
+		{
+			if (Node->IsDelete())
+			{
+				Display(Space, *Node, Node->GetFaceIndex(), EVisuProperty::PurplePoint);
+			}
+			else
+			{
+				Display(Space, *Node, Node->GetFaceIndex(), (EVisuProperty)((int32)NodeVisuProperty));
+			}
 		}
 	}
 }
@@ -164,6 +221,61 @@ void FIsoTriangulator::DisplayCycle(const TArray<FIsoSegment*>& Cycle, const TCH
 	}
 }
 
+void FIsoTriangulator::DisplayLoops(const TCHAR* Message, bool bOneNode, bool bSplitBySegment, bool bDisplayNext, bool bDisplayPrevious) const
+{
+	if (bDisplay)
+	{
+		F3DDebugSession _(Message);
+		{
+			{
+				F3DDebugSession A(!bOneNode, TEXT("Node Cloud"));
+				int32 Index = 0;
+				for (const FLoopNode& Node : LoopNodes)
+				{
+					if (Node.IsDelete())
+					{
+						Display(EGridSpace::UniformScaled, Node, Index++, EVisuProperty::PurplePoint);
+					}
+					Display(EGridSpace::UniformScaled, Node, Index++, EVisuProperty::BluePoint);
+				}
+			}
+
+			{
+				F3DDebugSession C(!bOneNode, TEXT("Segment"));
+				int32 Index = 0;
+				for (FIsoSegment* Segment : LoopSegments)
+				{
+					F3DDebugSession D(bSplitBySegment, TEXT("Segment"));
+					Display(EGridSpace::UniformScaled, Segment->GetFirstNode(), Segment->GetSecondNode(), Index++, EVisuProperty::BlueCurve);
+				}
+			}
+
+			if(bDisplayNext)
+			{
+				F3DDebugSession A(TEXT("Node Next"));
+				int32 Index = 0;
+				for (FIsoSegment* Segment : LoopSegments)
+				{
+					F3DDebugSession B(bSplitBySegment, TEXT("Segment"));
+					Display(EGridSpace::UniformScaled, Segment->GetFirstNode(), ((const FLoopNode&) Segment->GetFirstNode()).GetNextNode(), Index++, EVisuProperty::BlueCurve);
+				}
+			}
+
+			if(bDisplayPrevious)
+			{
+				F3DDebugSession C(TEXT("Node Forward"));
+				int32 Index = 0;
+				for (FIsoSegment* Segment : LoopSegments)
+				{
+					F3DDebugSession B(bSplitBySegment, TEXT("Segment"));
+					Display(EGridSpace::UniformScaled, Segment->GetSecondNode(), ((const FLoopNode&)Segment->GetSecondNode()).GetPreviousNode(), Index++, EVisuProperty::BlueCurve);
+				}
+			}
+
+		}
+		Wait(false);
+	}
+}
 
 void FIsoTriangulator::DisplayCells(const TArray<FCell>& Cells) const
 {
@@ -180,6 +292,24 @@ void FIsoTriangulator::DisplayCells(const TArray<FCell>& Cells) const
 	}
 }
 
+void FIsoTriangulator::DrawCellBoundary(int32 Index, EVisuProperty Property) const
+{
+	int32 IndexU;
+	int32 IndexV;
+	Grid.UVIndexFromGlobalIndex(Index, IndexU, IndexV);
+
+	DisplayPoint(Grid.GetInner2DPoint(EGridSpace::UniformScaled, IndexU, IndexV), Property, Index);
+	DisplayPoint(Grid.GetInner2DPoint(EGridSpace::UniformScaled, IndexU, IndexV + 1), Property, Index);
+	DisplayPoint(Grid.GetInner2DPoint(EGridSpace::UniformScaled, IndexU + 1, IndexV), Property, Index);
+	DisplayPoint(Grid.GetInner2DPoint(EGridSpace::UniformScaled, IndexU + 1, IndexV + 1), Property, Index);
+
+	DisplaySegment(Grid.GetInner2DPoint(EGridSpace::UniformScaled, IndexU, IndexV), Grid.GetInner2DPoint(EGridSpace::UniformScaled, IndexU, IndexV + 1), Index, EVisuProperty(Property + 1));
+	DisplaySegment(Grid.GetInner2DPoint(EGridSpace::UniformScaled, IndexU + 1, IndexV), Grid.GetInner2DPoint(EGridSpace::UniformScaled, IndexU + 1, IndexV + 1), Index, EVisuProperty(Property + 1));
+	DisplaySegment(Grid.GetInner2DPoint(EGridSpace::UniformScaled, IndexU, IndexV), Grid.GetInner2DPoint(EGridSpace::UniformScaled, IndexU + 1, IndexV), Index, EVisuProperty(Property + 1));
+	DisplaySegment(Grid.GetInner2DPoint(EGridSpace::UniformScaled, IndexU, IndexV + 1), Grid.GetInner2DPoint(EGridSpace::UniformScaled, IndexU + 1, IndexV + 1), Index, EVisuProperty(Property + 1));
+}
+
+
 void FIsoTriangulator::DisplayCell(const FCell& Cell) const
 {
 	if (!bDisplay)
@@ -187,27 +317,32 @@ void FIsoTriangulator::DisplayCell(const FCell& Cell) const
 		return;
 	}
 
+
 	F3DDebugSession _(FString::Printf(TEXT("Cell %d"), Cell.Id));
+
+	DrawCellBoundary(Cell.Id, EVisuProperty::YellowPoint);
 	for (const TArray<FLoopNode*>& Nodes : Cell.SubLoops)
 	{
 		for (const FLoopNode* Node : Nodes)
 		{
-			DisplayPoint(Node->GetPoint(EGridSpace::UniformScaled, Grid), EVisuProperty::OrangePoint, Cell.Id);
+			DisplayPoint(Node->GetPoint(EGridSpace::UniformScaled, Grid), EVisuProperty::BluePoint, Cell.Id);
 		}
 	}
 }
 
-void FIntersectionSegmentTool::Display(const TCHAR* Message) const
+void FIntersectionSegmentTool::Display(const TCHAR* Message, EVisuProperty Property) const
 {
 	if (!Grid.bDisplay)
 	{
 		return;
 	}
-	
+
+	int32 Index = 0;
+
 	Open3DDebugSession(Message);
 	for (const FSegment4IntersectionTools& Segment : Segments)
 	{
-		DisplaySegment(Segment.Segment2D[0], Segment.Segment2D[1]);
+		DisplaySegment(Segment.Segment2D[0], Segment.Segment2D[1], Index++, Property);
 	}
 	Close3DDebugSession();
 }
diff --git a/Engine/Source/Runtime/Datasmith/CADKernel/Private/CADKernel/Mesh/Meshers/IsoTriangulatorRemoveIntersection.cpp b/Engine/Source/Runtime/Datasmith/CADKernel/Private/CADKernel/Mesh/Meshers/IsoTriangulatorRemoveIntersection.cpp
new file mode 100644
--- /dev/null
+++ b/Engine/Source/Runtime/Datasmith/CADKernel/Private/CADKernel/Mesh/Meshers/IsoTriangulatorRemoveIntersection.cpp
@@ -0,0 +1,2171 @@
+// Copyright Epic Games, Inc. All Rights Reserved.
+
+#include "CADKernel/Mesh/Meshers/IsoTriangulator.h"
+
+#include "CADKernel/Math/Geometry.h"
+#include "CADKernel/Math/Point.h"
+#include "CADKernel/Mesh/Meshers/IsoTriangulator/IsoNode.h"
+#include "CADKernel/Mesh/Meshers/IsoTriangulator/IsoSegment.h"
+#include "Algo/AllOf.h"
+
+namespace CADKernel
+{
+
+//#define DEBUG_FIND_LOOP_INTERSECTION_AND_FIX_IT
+
+void FIsoTriangulator::FindBestLoopExtremity(TArray<FLoopNode*>& BestStartNodeOfLoop)
+{
+	double UMin = HUGE_VAL;
+	double UMax = -HUGE_VAL;
+	double VMin = HUGE_VAL;
+	double VMax = -HUGE_VAL;
+
+	FLoopNode* ExtremityNodes[4] = { nullptr, nullptr, nullptr, nullptr };
+
+	FLoopNode* BestNode = nullptr;
+	double OptimalSlop = 9.;
+	int32 LoopIndex = 0;
+
+	//F3DDebugSession _(TEXT("FindBestLoopExtremity"));
+
+	TFunction<void(FLoopNode*)> CompareWithSlopAt = [&](FLoopNode* Node)
+	{
+		FLoopNode& PreviousNode = Node->GetPreviousNode();
+		FLoopNode& NextNode = Node->GetNextNode();
+		double Slop = ComputePositiveSlope(Node->Get2DPoint(EGridSpace::UniformScaled, Grid), PreviousNode.Get2DPoint(EGridSpace::UniformScaled, Grid), NextNode.Get2DPoint(EGridSpace::UniformScaled, Grid));
+
+		//F3DDebugSession A(FString::Printf(TEXT("Node Slop %f"), Slop));
+		//Display(EGridSpace::UniformScaled, *Node, 0, EVisuProperty::RedPoint);
+
+		if ((Slop > OptimalSlop) == (LoopIndex == 0))
+		{
+			OptimalSlop = Slop;
+			BestNode = Node;
+		}
+	};
+
+	TFunction<void()> FindLoopExtremity = [&]()
+	{
+		BestNode = nullptr;
+		OptimalSlop = (LoopIndex == 0) ? -1 : 9.;
+
+		for (FLoopNode* Node : ExtremityNodes)
+		{
+			CompareWithSlopAt(Node);
+		}
+		BestStartNodeOfLoop.Add(BestNode);
+
+		// init for next loop
+		UMin = HUGE_VAL;
+		UMax = -HUGE_VAL;
+		VMin = HUGE_VAL;
+		VMax = -HUGE_VAL;
+
+		for (FLoopNode*& Node : ExtremityNodes)
+		{
+			Node = nullptr;
+		}
+	};
+
+	for (FLoopNode& Node : LoopNodes)
+	{
+		if (Node.GetLoopIndex() != LoopIndex)
+		{
+			FindLoopExtremity();
+			LoopIndex = Node.GetLoopIndex();
+		}
+
+		const FPoint2D& Point = Node.Get2DPoint(EGridSpace::UniformScaled, Grid);
+
+		if (Point.U > UMax)
+		{
+			UMax = Point.U;
+			ExtremityNodes[0] = &Node;
+		}
+		if (Point.U < UMin)
+		{
+			UMin = Point.U;
+			ExtremityNodes[1] = &Node;
+		}
+
+		if (Point.V > VMax)
+		{
+			VMax = Point.V;
+			ExtremityNodes[2] = &Node;
+		}
+		if (Point.V < VMin)
+		{
+			VMin = Point.V;
+			ExtremityNodes[3] = &Node;
+		}
+	}
+	FindLoopExtremity();
+}
+
+EOrientation FIsoTriangulator::GetLoopOrientation(const FLoopNode* StartNode)
+{
+	using namespace IsoTriangulatorImpl;
+	double UMin = HUGE_VAL;
+	double UMax = -HUGE_VAL;
+	double VMin = HUGE_VAL;
+	double VMax = -HUGE_VAL;
+
+	const FLoopNode* ExtremityNodes[4] = { nullptr, nullptr, nullptr, nullptr };
+
+	const FLoopNode* BestNode = nullptr;
+
+	int32 LoopIndex = StartNode->GetLoopIndex();
+	double OptimalSlop = (LoopIndex == 0) ? -1 : 9.;
+
+	TFunction<void(const FLoopNode*)> CompareWithSlopAt = [&](const FLoopNode* Node)
+	{
+		FLoopNode& PreviousNode = Node->GetPreviousNode();
+		FLoopNode& NextNode = Node->GetNextNode();
+		double Slop = ComputePositiveSlope(Node->Get2DPoint(EGridSpace::UniformScaled, Grid), PreviousNode.Get2DPoint(EGridSpace::UniformScaled, Grid), NextNode.Get2DPoint(EGridSpace::UniformScaled, Grid));
+
+		if ((Slop > OptimalSlop) == (LoopIndex == 0))
+		{
+			OptimalSlop = Slop;
+			BestNode = Node;
+		}
+	};
+
+	TFunction<void()> FindLoopExtremity = [&]()
+	{
+		BestNode = nullptr;
+
+		for (const FLoopNode* Node : ExtremityNodes)
+		{
+			CompareWithSlopAt(Node);
+		}
+
+		// init for next loop
+		UMin = HUGE_VAL;
+		UMax = -HUGE_VAL;
+		VMin = HUGE_VAL;
+		VMax = -HUGE_VAL;
+
+		for (const FLoopNode*& Node : ExtremityNodes)
+		{
+			Node = nullptr;
+		}
+	};
+
+	const FLoopNode* Node = GetNextConstNodeImpl(StartNode);
+	for (; Node != StartNode; Node = GetNextConstNodeImpl(Node))
+	{
+		const FPoint2D& Point = Node->Get2DPoint(EGridSpace::UniformScaled, Grid);
+
+		if (Point.U > UMax)
+		{
+			UMax = Point.U;
+			ExtremityNodes[0] = Node;
+		}
+		if (Point.U < UMin)
+		{
+			UMin = Point.U;
+			ExtremityNodes[1] = Node;
+		}
+
+		if (Point.V > VMax)
+		{
+			VMax = Point.V;
+			ExtremityNodes[2] = Node;
+		}
+		if (Point.V < VMin)
+		{
+			VMin = Point.V;
+			ExtremityNodes[3] = Node;
+		}
+	}
+	FindLoopExtremity();
+
+	if (LoopIndex == 0)
+	{
+		return OptimalSlop > 4 ? EOrientation::Front : EOrientation::Back;
+	}
+	else
+	{
+		return OptimalSlop < 4 ? EOrientation::Front : EOrientation::Back;
+	}
+}
+
+//#define DEBUG_FIND_LOOP_INTERSECTIONS
+void FIsoTriangulator::FindLoopIntersections(const TArray<FLoopNode*>& NodesOfLoop, bool bForward, TArray<TPair<double, double>>& OutIntersections)
+{
+	using namespace IsoTriangulatorImpl;
+
+	GetNextNodeMethod GetNext = bForward ? GetNextNodeImpl : GetPreviousNodeImpl;
+	GetSegmentToNodeMethod GetFirst = bForward ? GetFirstNode : GetSecondNode;
+	GetSegmentToNodeMethod GetSecond = bForward ? GetSecondNode : GetFirstNode;
+
+	FLoopNode *const* StartNodePtr = NodesOfLoop.FindByPredicate([](const FLoopNode* Node) { return !Node->IsDelete(); });
+	ensureCADKernel(*StartNodePtr);
+
+	FLoopNode* StartNode = *StartNodePtr;
+
+	TArray<const FIsoSegment*> IntersectedSegments;
+
+	int32 SegmentCount = 1;
+	int32 SegmentIndex = 1;
+	FLoopNode* Node = nullptr;
+	FLoopNode* NextNode = nullptr;
+
+	TFunction<void()> FindSegmentIntersection = [&]()
+	{
+		FIsoSegment* Segment = Node->GetSegmentConnectedTo(NextNode);
+		ensureCADKernel(Segment);
+
+		if (LoopSegmentsIntersectionTool.FindIntersections(*Node, *NextNode, IntersectedSegments))
+		{
+			for (const FIsoSegment* IntersectedSegment : IntersectedSegments)
+			{
+				const FLoopNode* IntersectedSegmentFirstNode = GetFirst(IntersectedSegment);
+				int32 IntersectionIndex = 0;
+				FLoopNode* TmpNode = StartNode;
+				while (TmpNode != IntersectedSegmentFirstNode)
+				{
+					++IntersectionIndex;
+					TmpNode = GetNext(TmpNode);
+				}
+
+				TSegment<FPoint2D> SegmentPoints(Node->Get2DPoint(EGridSpace::UniformScaled, Grid), NextNode->Get2DPoint(EGridSpace::UniformScaled, Grid));
+				TSegment<FPoint2D> IntersectedSegmentPoints(GetFirst(IntersectedSegment)->Get2DPoint(EGridSpace::UniformScaled, Grid), GetSecond(IntersectedSegment)->Get2DPoint(EGridSpace::UniformScaled, Grid));
+
+				double SegmentIntersectionCoordinate;
+				FPoint2D IntersectionPoint = FindIntersectionOfSegments2D(SegmentPoints, IntersectedSegmentPoints, SegmentIntersectionCoordinate);
+
+				double SquareLength = IntersectedSegmentPoints.SquareLength();
+				double IntersectedSegmentIntersectionCoordinate = (SquareLength > SMALL_NUMBER_SQUARE) ? FMath::Sqrt(IntersectedSegmentPoints[0].SquareDistance(IntersectionPoint) / IntersectedSegmentPoints.SquareLength()) : 0;
+				IntersectedSegmentIntersectionCoordinate += IntersectionIndex;
+
+				SegmentIntersectionCoordinate += SegmentIndex;
+				if (FMath::IsNearlyEqual(SegmentIntersectionCoordinate, SegmentCount))
+				{
+					SegmentIntersectionCoordinate = IntersectedSegmentIntersectionCoordinate;
+					IntersectedSegmentIntersectionCoordinate = 0;
+				}
+
+				if (OutIntersections.Num() && FMath::IsNearlyEqual(OutIntersections.Last().Key, IntersectedSegmentIntersectionCoordinate) && FMath::IsNearlyEqual(OutIntersections.Last().Value, SegmentIntersectionCoordinate))
+				{
+					continue;
+				}
+				OutIntersections.Emplace(IntersectedSegmentIntersectionCoordinate, SegmentIntersectionCoordinate);
+
+#ifdef DEBUG_FIND_LOOP_INTERSECTIONS		
+				if (bDisplay)
+				{
+					F3DDebugSession _(*FString::Printf(TEXT("Intersection %f %f"), IntersectedSegmentIntersectionCoordinate, SegmentIntersectionCoordinate));
+					Display(EGridSpace::UniformScaled, *Node, *NextNode, 0, EVisuProperty::BluePoint);
+					Display(EGridSpace::UniformScaled, *IntersectedSegment, 0, EVisuProperty::BluePoint);
+					DisplayPoint(IntersectionPoint, EVisuProperty::RedPoint);
+				}
+#endif
+			}
+		}
+
+		LoopSegmentsIntersectionTool.AddSegment(*Segment);
+	};
+
+	LoopSegmentsIntersectionTool.Empty(NodesOfLoop.Num());
+
+	Node = GetNext(StartNode);
+	for (; Node != StartNode; Node = GetNext(Node), ++SegmentCount);
+
+	Node = StartNode;
+	NextNode = GetNext(Node);
+
+#ifdef DEBUG_FIND_LOOP_INTERSECTIONS		
+	if (bDisplay)
+	{
+		F3DDebugSession _(TEXT("Start Node"));
+		Display(EGridSpace::UniformScaled, *Node, *NextNode, 0, EVisuProperty::BlueCurve);
+		Display(EGridSpace::UniformScaled, *Node, 0, EVisuProperty::RedPoint);
+	}
+#endif
+
+	LoopSegmentsIntersectionTool.AddSegment(*StartNode->GetSegmentConnectedTo(NextNode));
+
+	for (Node = NextNode; Node != StartNode; Node = NextNode, ++SegmentIndex)
+	{
+		NextNode = GetNext(Node);
+		FindSegmentIntersection();
+	}
+
+}
+
+bool FIsoTriangulator::RemoveLoopIntersections(const TArray<FLoopNode*>& NodesOfLoop, TArray<TPair<double, double>>& Intersections, bool bForward)
+{
+	using namespace IsoTriangulatorImpl;
+	GetNextNodeMethod GetNext = bForward ? GetNextNodeImpl : GetPreviousNodeImpl;
+	GetNextNodeMethod GetPrevious = bForward ? GetPreviousNodeImpl : GetNextNodeImpl;
+	const int32 NodeCount = NodesOfLoop.Num();
+
+#ifdef DEBUG_REMOVE_LOOP_INTERSECTIONS		
+	TFunction<void(const TPair<double, double>&)> DisplayIntersection = [&](const TPair<double, double>& Intersection)
+	{
+		if (bDisplay)
+		{
+			FLoopNode* Segment0End = GetNodeAt(NodesOfLoop, NextIndex(NodeCount, (int32)Intersection.Key));
+			FLoopNode* Segment1Start = GetNodeAt(NodesOfLoop, (int32)Intersection.Value);
+
+			Display(EGridSpace::UniformScaled, *Segment0End, *GetPrevious(Segment0End), 0, EVisuProperty::RedCurve);
+			Display(EGridSpace::UniformScaled, *Segment1Start, *GetNext(Segment1Start), 0, EVisuProperty::RedCurve);
+			Display(EGridSpace::UniformScaled, *GetPrevious(Segment0End), 0, EVisuProperty::RedPoint);
+			Display(EGridSpace::UniformScaled, *Segment1Start, 0, EVisuProperty::RedPoint);
+		}
+	};
+#endif
+
+	Algo::Sort(Intersections, [&](const TPair<double, double>& Intersection1, const TPair<double, double>& Intersection2)
+		{
+			return (Intersection1.Key < Intersection2.Key);
+		});
+
+	// Nodes are not yet deleted in the sub loop
+	for (int32 IntersectionIndex = 0; IntersectionIndex < Intersections.Num(); )
+	{
+#ifdef DEBUG_REMOVE_LOOP_INTERSECTIONS		
+		if (bDisplay)
+		{
+			DisplayLoops(TEXT("RemoveLoopIntersections"));
+		}
+#endif
+
+		const TPair<double, double>& Intersection = Intersections[IntersectionIndex];
+
+#ifdef DEBUG_REMOVE_LOOP_INTERSECTIONS		
+		F3DDebugSession _(bDisplay, TEXT("Intersected Segments"));
+		DisplayIntersection(Intersection);
+#endif
+
+		bool bIntersectionForward = true;
+
+		int32 NextIntersectionIndex = IntersectionIndex + 1;
+		int32 IntersectionCount = 1;
+		for (; NextIntersectionIndex < Intersections.Num(); ++NextIntersectionIndex)
+		{
+			if (Intersections[NextIntersectionIndex].Value > Intersection.Value)
+			{
+				break;
+			}
+			++IntersectionCount;
+
+#ifdef DEBUG_REMOVE_LOOP_INTERSECTIONS		
+			DisplayIntersection(Intersections[NextIntersectionIndex]);
+#endif
+		}
+
+		if (IntersectionCount == 1)
+		{
+			NextIntersectionIndex = IntersectionIndex + 1;
+			for (; NextIntersectionIndex < Intersections.Num(); ++NextIntersectionIndex)
+			{
+				if (Intersections[NextIntersectionIndex].Key > Intersection.Value)
+				{
+					break;
+				}
+
+				bIntersectionForward = false;
+				++IntersectionCount;
+#ifdef DEBUG_REMOVE_LOOP_INTERSECTIONS		
+				DisplayIntersection(Intersections[NextIntersectionIndex]);
+#endif
+				break;
+			}
+		}
+
+		// Stating from this point 
+		// the process must not delete node after the first intersection
+		if (IntersectionCount == 1)
+		{
+			if (!RemoveUniqueIntersection(NodesOfLoop, Intersections[IntersectionIndex], bForward))
+			{
+				return false;
+			}
+			IntersectionIndex++;
+		}
+		else if (!bIntersectionForward)
+		{
+			if (!RemovePickToOutside(NodesOfLoop, Intersections[IntersectionIndex], Intersections[IntersectionIndex + 1], bForward))
+			{
+				return false;
+			}
+			IntersectionIndex += 2;
+		}
+		else
+		{
+			if (!RemoveIntersectionsOfSubLoop(NodesOfLoop, Intersections, IntersectionIndex, IntersectionCount, bForward))
+			{
+				return false;
+			}
+			IntersectionIndex += IntersectionCount;
+		}
+	}
+
+	return true;
+}
+
+void FIsoTriangulator::RemoveLoopPicks(TArray<FLoopNode*>& NodesOfLoop, TArray<TPair<double, double>>& Intersections)
+{
+	for (FLoopNode* Node : NodesOfLoop)
+	{
+		if (Node->IsDelete())
+		{
+			continue;
+		}
+		RemovePickRecursively(Node, &Node->GetNextNode());
+	}
+	IsoTriangulatorImpl::RemoveDeletedNodes(NodesOfLoop);
+}
+
+bool FIsoTriangulator::RemovePickToOutside(const TArray<FLoopNode*>& NodesOfLoop, const TPair<double, double>& Intersection, const TPair<double, double>& NextIntersection, bool bForward)
+{
+	const TPair<double, double> OutSideLoop(Intersection.Value, NextIntersection.Value);
+	if (IsSubLoopBiggerThanMainLoop(NodesOfLoop, OutSideLoop, bForward))
+	{
+		return false;
+	}
+
+	using namespace IsoTriangulatorImpl;
+	GetNextNodeMethod GetNext = bForward ? GetNextNodeImpl : GetPreviousNodeImpl;
+	GetNextNodeMethod GetPrevious = bForward ? GetPreviousNodeImpl : GetNextNodeImpl;
+
+	const int32 NodeCount = NodesOfLoop.Num();
+
+	TFunction<FPoint2D(double, FLoopNode*, FLoopNode*)> IntersectingPoint = [&](double Coordinate, FLoopNode* Start, FLoopNode* End) -> FPoint2D
+	{
+		const FPoint2D& StartPoint = Start->Get2DPoint(EGridSpace::UniformScaled, Grid);
+		const FPoint2D& EndPoint = End->Get2DPoint(EGridSpace::UniformScaled, Grid);
+
+		int32 StartIndex = (int32)Coordinate;
+		Coordinate -= StartIndex;
+		FPoint2D Intersection = PointOnSegment(StartPoint, EndPoint, Coordinate);
+
+		//{
+		//	F3DDebugSession _(TEXT("Intersected Point"));
+		//	DisplayPoint(StartPoint, EVisuProperty::RedPoint);
+		//	DisplayPoint(EndPoint, EVisuProperty::RedPoint);
+		//	DisplayPoint(Intersection, EVisuProperty::BluePoint);
+		//}
+
+		return Intersection;
+	};
+
+	FLoopNode* TmpNode = GetNodeAt(NodesOfLoop, NextIndex(NodeCount, (int32) Intersection.Value));
+	FLoopNode* StartNode = GetPrevious(TmpNode);
+	FPoint2D FirstIntersection = IntersectingPoint(Intersection.Value, StartNode, TmpNode);
+
+	FLoopNode* EndNode = GetNodeAt(NodesOfLoop, (int32)NextIntersection.Value);
+	TmpNode = GetNext(EndNode);
+	FPoint2D SecondIntersection = IntersectingPoint(NextIntersection.Value, EndNode, TmpNode);
+
+	FPoint2D MiddlePoint = FirstIntersection.Middle(SecondIntersection);
+
+	TmpNode = GetNext(StartNode);
+	while (TmpNode && (TmpNode != EndNode) && !TmpNode->IsDelete())
+	{
+		RemoveNodeOfLoop(*TmpNode);
+		TmpNode = GetNext(StartNode);
+	}
+
+	FPoint2D MoveDirection = SecondIntersection - FirstIntersection;
+	double Length = MoveDirection.Length();
+	if (FMath::IsNearlyZero(Length))
+	{
+		FLoopNode* StartSegment = GetNodeAt(NodesOfLoop, (int32)Intersection.Key);
+		MoveDirection = GetNext(StartSegment)->Get2DPoint(EGridSpace::UniformScaled, Grid) - StartSegment->Get2DPoint(EGridSpace::UniformScaled, Grid);
+		Length = MoveDirection.Length();
+	}
+
+	MoveDirection /= Length;
+	MoveDirection = MoveDirection.GetPerpendicularVector();
+	MoveDirection *= GeometricTolerance;
+	MiddlePoint += MoveDirection;
+
+	if(!EndNode->IsDelete())
+	{
+		EndNode->Set2DPoint(EGridSpace::UniformScaled, Grid, MiddlePoint);
+	}
+
+	return true;
+}
+
+void FIsoTriangulator::RemoveSubLoop(FLoopNode* StartNode, FLoopNode* EndNode, IsoTriangulatorImpl::GetNextNodeMethod NextNode)
+{
+	FLoopNode* Node = NextNode(StartNode);
+	while (Node && (Node != EndNode) && !Node->IsDelete())
+	{
+		RemoveNodeOfLoop(*Node);
+		Node = NextNode(StartNode);
+	}
+
+	//Display(EGridSpace::UniformScaled, TEXT("Loop"), LoopSegments, true);
+	//Wait();
+}
+
+void FIsoTriangulator::MoveIntersectingSectionBehindOppositeSection(IsoTriangulatorImpl::FLoopSection IntersectingSection, IsoTriangulatorImpl::FLoopSection OppositeSection, IsoTriangulatorImpl::GetNextNodeMethod GetNext, IsoTriangulatorImpl::GetNextNodeMethod GetPrevious)
+{
+	using namespace IsoTriangulatorImpl;
+
+	FLoopNode* FirstNodeIntersectingSection = IntersectingSection.Key;
+	FLoopNode* LastNodeIntersectingSection = IntersectingSection.Value;
+
+	FLoopNode* FirstNodeOppositeSection = OppositeSection.Key;
+	FLoopNode* LastNodeOppositeSection = OppositeSection.Value;
+
+	int32 OppositeSectionCount = 1;
+	for (FLoopNode* SegmentNode = FirstNodeOppositeSection; SegmentNode != LastNodeOppositeSection; SegmentNode = GetNext(SegmentNode), ++OppositeSectionCount);
+
+	TArray<const FPoint2D*> OppositeSectionPoint;
+	OppositeSectionPoint.Reserve(OppositeSectionCount);
+
+	for (FLoopNode* SegmentNode = FirstNodeOppositeSection; SegmentNode != LastNodeOppositeSection; SegmentNode = GetNext(SegmentNode))
+	{
+		OppositeSectionPoint.Add(&SegmentNode->Get2DPoint(EGridSpace::UniformScaled, Grid));
+	}
+	OppositeSectionPoint.Add(&LastNodeOppositeSection->Get2DPoint(EGridSpace::UniformScaled, Grid));
+
+#ifdef DEBUG_MOVE_INTERSECTING_SECTION_BEHIND_OPPOSITE_SECTION
+	{
+		F3DDebugSession _(TEXT("IntersectingSection"));
+		for (FLoopNode* Node = FirstNodeIntersectingSection; Node != GetNext(LastNodeIntersectingSection); Node = GetNext(Node))
+		{
+			DisplayPoint(Node->Get2DPoint(EGridSpace::UniformScaled, Grid), EVisuProperty::PurplePoint);
+		}
+	}
+	{
+		F3DDebugSession _(TEXT("OppositeSection"));
+		for (FLoopNode* Node = FirstNodeOppositeSection; Node != GetNext(LastNodeOppositeSection); Node = GetNext(Node))
+		{
+			DisplayPoint(Node->Get2DPoint(EGridSpace::UniformScaled, Grid), EVisuProperty::YellowPoint);
+		}
+		Wait();
+	}
+#endif
+
+	double Coordinate;
+	FLoopNode* NextNode = nullptr;
+	for (FLoopNode* Node = GetNext(FirstNodeIntersectingSection); Node != LastNodeIntersectingSection;)
+	{
+		FLoopNode* NodeToProceed = Node;
+		Node = GetNext(Node);
+
+		FPoint2D  CandiatePosition;
+		double MinSquareDistance = HUGE_VALUE;
+		FPoint2D PointToProject = NodeToProceed->Get2DPoint(EGridSpace::UniformScaled, Grid);
+		for (int32 Index = 1; Index < OppositeSectionCount; ++Index)
+		{
+			FPoint2D ProjectedPoint = ProjectPointOnSegment(PointToProject, *OppositeSectionPoint[Index - 1], *OppositeSectionPoint[Index], Coordinate);
+			double SquareDistance = PointToProject.SquareDistance(ProjectedPoint);
+			if (SquareDistance < MinSquareDistance)
+			{
+				MinSquareDistance = SquareDistance;
+				CandiatePosition = ProjectedPoint;
+			}
+		}
+		MoveNode(*NodeToProceed, CandiatePosition);
+	}
+
+	for (FLoopNode* Node = GetNext(FirstNodeIntersectingSection); Node != LastNodeIntersectingSection; )
+	{
+		FLoopNode* NodeToProceed = Node;
+		Node = GetNext(Node);
+		if (CheckAndRemovePick(NodeToProceed->GetPreviousNode().Get2DPoint(EGridSpace::UniformScaled, Grid), NodeToProceed->Get2DPoint(EGridSpace::UniformScaled, Grid), NodeToProceed->GetNextNode().Get2DPoint(EGridSpace::UniformScaled, Grid), *NodeToProceed))
+		{
+			Node = GetPrevious(Node);
+		}
+	}
+}
+
+bool FIsoTriangulator::RemoveIntersectionsOfSubLoop(const TArray<FLoopNode*>& NodesOfLoop, TArray<TPair<double, double>>& LoopIntersections, int32 IntersectionIndex, int32 IntersectionCount, bool bForward)
+{
+	using namespace IsoTriangulatorImpl;
+	GetNextNodeMethod GetNext = bForward ? GetNextNodeImpl : GetPreviousNodeImpl;
+	GetNextNodeMethod GetPrevious = bForward ? GetPreviousNodeImpl : GetNextNodeImpl;
+
+	const int32 NodeCount = NodesOfLoop.Num();
+
+	TFunction<void(FLoopNode*, const FPoint2D&, FPoint2D&)> MoveNodeToProjection = [&](FLoopNode* NodeToProject, const FPoint2D& PointToProject, FPoint2D& ProjectedPoint)
+	{
+		FPoint2D MoveDirection = ProjectedPoint - PointToProject;
+		MoveDirection.Normalize();
+		MoveDirection *= GeometricTolerance;
+		ProjectedPoint += MoveDirection;
+		NodeToProject->Set2DPoint(EGridSpace::UniformScaled, Grid, ProjectedPoint);
+	};
+
+	TFunction<void(FLoopNode*, const FPoint2D&, const FPoint2D&)> ProjectNodeOnSegment = [&](FLoopNode* NodeToProject, const FPoint2D& Point0, const FPoint2D& Point1)
+	{
+		const FPoint2D& PointToProject = NodeToProject->Get2DPoint(EGridSpace::UniformScaled, Grid);
+
+		FPoint2D ProjectedPoint;
+		double Coordinate;
+		ProjectedPoint = ProjectPointOnSegment(PointToProject, Point0, Point1, Coordinate);
+
+		MoveNodeToProjection(NodeToProject, PointToProject, ProjectedPoint);
+	};
+
+	TFunction<void(const int32, const int32, const int32)> ProjectNodesOnSegment = [&](const int32 StartIndex, const int32 EndIndex, const int32 SegmnentEndIndex)
+	{
+		FLoopNode* Node = GetNodeAt(NodesOfLoop, StartIndex);
+		FLoopNode* StopNode = GetNodeAt(NodesOfLoop, EndIndex);
+		FLoopNode* EndSegment = GetNodeAt(NodesOfLoop, SegmnentEndIndex);
+
+		StopNode = GetNext(StopNode);
+		FLoopNode* StartSegment = GetPrevious(EndSegment);
+
+		const FPoint2D& EndPoint = EndSegment->Get2DPoint(EGridSpace::UniformScaled, Grid);
+		const FPoint2D& StartPoint = StartSegment->Get2DPoint(EGridSpace::UniformScaled, Grid);
+
+		for (; Node != StopNode; Node = GetNext(Node))
+		{
+			ProjectNodeOnSegment(Node, StartPoint, EndPoint);
+		}
+
+		for (Node = GetNodeAt(NodesOfLoop, StartIndex); Node != StopNode; )
+		{
+			FLoopNode* NodeToProcess = Node;
+			Node = GetNext(Node);
+			CheckAndRemovePick(NodeToProcess->GetPreviousNode().Get2DPoint(EGridSpace::UniformScaled, Grid), NodeToProcess->Get2DPoint(EGridSpace::UniformScaled, Grid), NodeToProcess->GetNextNode().Get2DPoint(EGridSpace::UniformScaled, Grid), *NodeToProcess);
+		}
+	};
+
+	TFunction<void(FLoopNode*, FLoopNode*)> MoveNodeBehindOther = [&](FLoopNode* NodeToMove, FLoopNode* Node1Side1)
+	{
+		FLoopNode* Node0Side1 = GetPrevious(Node1Side1);
+		const FPoint2D& Point0 = Node0Side1->Get2DPoint(EGridSpace::UniformScaled, Grid);
+		const FPoint2D& Point1 = Node1Side1->Get2DPoint(EGridSpace::UniformScaled, Grid);
+
+		const FPoint2D& PointToMove = NodeToMove->Get2DPoint(EGridSpace::UniformScaled, Grid);
+
+		FPoint2D MoveDirection = Point1 - Point0;
+		MoveDirection.Normalize();
+		MoveDirection = MoveDirection.GetPerpendicularVector();
+		MoveDirection *= GeometricTolerance;
+
+		FPoint2D NewCoordinate = Point1 + MoveDirection;
+		NodeToMove->Set2DPoint(EGridSpace::UniformScaled, Grid, NewCoordinate);
+	};
+
+	for (int32 Index = IntersectionCount - 1; Index >= 0; --Index)
+	{
+		int32 SecondIntersectionIndex = IntersectionIndex + Index;
+		if (Index > 0)
+		{
+			const TPair<double, double>& SecondIntersection = LoopIntersections[SecondIntersectionIndex];
+			const TPair<double, double>& FirstIntersection = LoopIntersections[SecondIntersectionIndex - 1];
+
+			int32 Side0NodeCount = (int32)SecondIntersection.Key - (int32)FirstIntersection.Key;
+			int32 Side1NodeCount = (int32)FirstIntersection.Value - (int32)SecondIntersection.Value;
+
+			int32 IndexSide0 = NextIndex(NodeCount, (int32)FirstIntersection.Key);
+			int32 IndexSide1 = NextIndex(NodeCount, (int32)SecondIntersection.Value);
+
+			if (Side0NodeCount == 0)
+			{
+				ProjectNodesOnSegment(IndexSide1, (int32)FirstIntersection.Value, IndexSide0);
+			}
+			else if (Side1NodeCount == 0)
+			{
+				ProjectNodesOnSegment(IndexSide0, (int32)SecondIntersection.Key, IndexSide1);
+			}
+			else if (Side0NodeCount == 1 && Side1NodeCount == 1)
+			{
+				FLoopNode* NodeSide0 = GetNodeAt(NodesOfLoop, IndexSide0);
+				FLoopNode* NodeSide1 = GetNodeAt(NodesOfLoop, IndexSide1);
+
+				double SlopSide0 = ComputeUnorientedSlope(GetPrevious(NodeSide0)->Get2DPoint(EGridSpace::UniformScaled, Grid), NodeSide0->Get2DPoint(EGridSpace::UniformScaled, Grid), GetNext(NodeSide0)->Get2DPoint(EGridSpace::UniformScaled, Grid));
+				double SlopSide1 = ComputeUnorientedSlope(GetPrevious(NodeSide1)->Get2DPoint(EGridSpace::UniformScaled, Grid), NodeSide1->Get2DPoint(EGridSpace::UniformScaled, Grid), GetNext(NodeSide1)->Get2DPoint(EGridSpace::UniformScaled, Grid));
+
+				if (SlopSide0 < SlopSide1)
+				{
+					MoveNodeBehindOther(NodeSide1, NodeSide0);
+				}
+				else
+				{
+					MoveNodeBehindOther(NodeSide0, NodeSide1);
+				}
+			}
+			else
+			{
+				FLoopSection IntersectingSection;
+				IntersectingSection.Key = GetPrevious(GetNodeAt(NodesOfLoop, IndexSide0));
+				IntersectingSection.Value = GetNext(GetNodeAt(NodesOfLoop, (int32)SecondIntersection.Key));
+
+				FLoopSection OppositeSection;
+				OppositeSection.Key = GetPrevious(GetNodeAt(NodesOfLoop, IndexSide1));
+				OppositeSection.Value = GetNext(GetNodeAt(NodesOfLoop, (int32)FirstIntersection.Value));
+
+				int32 IntersectingSectionCount = (int32)FirstIntersection.Value - IndexSide1;
+				int32 OppositeSectionCount = (int32)SecondIntersection.Key - IndexSide0;
+				if (OppositeSectionCount < IntersectingSectionCount)
+				{
+					Swap(IntersectingSection, OppositeSection);
+				}
+
+				MoveIntersectingSectionBehindOppositeSection(IntersectingSection, OppositeSection, GetNext, GetPrevious);
+
+#ifdef DEBUG_REMOVE_INTERSECTIONS		
+				DisplayLoops(TEXT("RemoveLoopIntersections first step"), false, true);
+#endif
+
+				// Check if there is no more intersection, otherwise fix last intersections by moving node of OppositeSection
+				// As we don't know the existence of intersection, the process is chesk and fix
+				{
+					LoopSegmentsIntersectionTool.Empty(NodesOfLoop.Num());
+					for (FLoopNode* Node = IntersectingSection.Key; Node != IntersectingSection.Value; Node = GetNext(Node))
+					{
+						LoopSegmentsIntersectionTool.AddSegment(*Node->GetSegmentConnectedTo(GetNext(Node)));
+					}
+				}
+
+				for (FLoopNode* Node = OppositeSection.Key; Node != OppositeSection.Value; Node = GetNext(Node))
+				{
+					if (const FIsoSegment* Intersection = LoopSegmentsIntersectionTool.DoesIntersect(*Node, *GetNext(Node)))
+					{
+						RemoveIntersectionByMovingOutsideNodeInside(*Intersection, *GetNext(Node));
+					}
+				}
+
+#ifdef DEBUG_REMOVE_INTERSECTIONS		
+				DisplayLoops(TEXT("RemoveLoopIntersections second step"), false, true);
+#endif
+
+				// Remove the pick if neeeded
+				for (FLoopNode* Node = OppositeSection.Key; Node != OppositeSection.Value;)
+				{
+					FLoopNode* NodeToProceed = Node;
+					Node = GetNext(Node);
+					if (CheckAndRemovePick(NodeToProceed->GetPreviousNode().Get2DPoint(EGridSpace::UniformScaled, Grid), NodeToProceed->Get2DPoint(EGridSpace::UniformScaled, Grid), NodeToProceed->GetNextNode().Get2DPoint(EGridSpace::UniformScaled, Grid), *NodeToProceed))
+					{
+						Node = GetPrevious(Node);
+					}
+				}
+
+#ifdef DEBUG_REMOVE_INTERSECTIONS		
+				DisplayLoops(TEXT("RemoveLoopIntersections after remove pick"), false, true);
+				Wait();
+#endif
+			}
+			--Index;
+		}
+		else
+		{
+			if (!RemoveUniqueIntersection(NodesOfLoop, LoopIntersections[IntersectionIndex], bForward))
+			{
+				return false;
+			}
+		}
+	}
+
+	return true;
+}
+
+bool FIsoTriangulator::IsSubLoopBiggerThanMainLoop(const TArray<FLoopNode*>& NodesOfLoop, const TPair<double, double>& Intersection, bool bForward)
+{
+	using namespace IsoTriangulatorImpl;
+	GetNextConstNodeMethod GetNext = bForward ? GetNextConstNodeImpl : GetPreviousConstNodeImpl;
+	GetNextConstNodeMethod GetPrevious = bForward ? GetPreviousConstNodeImpl : GetNextConstNodeImpl;
+
+	const int32 NodeCount = NodesOfLoop.Num();
+
+	TFunction<double(const FLoopNode*, const FLoopNode*)> ComputeLength = [&](const FLoopNode* Start, const FLoopNode* End) -> double
+	{
+		double Length = 0;
+		const FLoopNode* NextNode = nullptr;
+		for (const FLoopNode* Node = Start; Node != End; Node = NextNode)
+		{
+			NextNode = GetNext(Node);
+			Length += Node->Get3DPoint(Grid).Distance(NextNode->Get3DPoint(Grid));
+		}
+		return Length;
+	};
+
+	// check the sampling of the sub-loop vs the other part of the loop (deleted nodes are counted)
+	int32 SubLoopSegmentCount = (int32)Intersection.Value - (int32)Intersection.Key;
+	int32 OtherSegmentCount = (int32)Intersection.Key + NodeCount - (int32)Intersection.Value;
+	if (SubLoopSegmentCount * 4 > OtherSegmentCount)
+	{
+		// The sub-loop is quite sampled to be bigger than the main part of the loop
+		// it's better to comput the length of each parts
+
+		// the sub-loop is not yet process so there are not deleted points
+		const FLoopNode* SubLoopStartNode = GetNodeAt(NodesOfLoop, NextIndex(NodeCount, (int32)Intersection.Key));
+		const FLoopNode* SubLoopEndNode = GetNodeAt(NodesOfLoop, (int32)Intersection.Value);
+
+		const FLoopNode* MainLoopStartNode = GetNext(SubLoopEndNode);
+		const FLoopNode* MainLoopEndNode = GetPrevious(SubLoopStartNode);
+
+		// compute each sub-loop length
+		double SubLoopLength = ComputeLength(SubLoopStartNode, SubLoopEndNode);
+		double MainLoopLength = ComputeLength(MainLoopStartNode, MainLoopEndNode);
+
+		if (SubLoopLength > MainLoopLength)
+		{
+			return true;
+		}
+	}
+
+	return false;
+}
+
+bool FIsoTriangulator::RemoveUniqueIntersection(const TArray<FLoopNode*>& NodesOfLoop, TPair<double, double> Intersection, bool bForward)
+{
+	const int32 NodesCount = NodesOfLoop.Num();
+	int32 SubLoopSegmentCount = (int32)Intersection.Value - (int32)Intersection.Key;
+	const int32 OtherSubLoopSegmentCount = NodesCount - (int32)Intersection.Value + (int32)Intersection.Key;
+
+	bool bIntersectionKeyIsExtremity = FMath::IsNearlyEqual(Intersection.Key, (int32)Intersection.Key);
+	bool bIntersectionValueIsExtremity = FMath::IsNearlyEqual(Intersection.Value, (int32)Intersection.Value);
+
+	// TODO
+	/*
+	if(bIntersectionKeyIsExtremity && bIntersectionValueIsExtremity)
+	{   //       _   _ 
+		// Cas:   \./
+		//       _/ \_
+
+		return SpreadCoincidentNodes(NodesOfLoop, Intersection);
+	}
+	else if(bIntersectionKeyIsExtremity || bIntersectionValueIsExtremity)
+	{
+		// Cas:  ________
+		//         _/\_
+		return MovePickBehind(NodesOfLoop, Intersection, bIntersectionKeyIsExtremity);
+	}
+	*/
+
+	if (OtherSubLoopSegmentCount < SubLoopSegmentCount)
+	{
+		Swap(Intersection.Value, Intersection.Key);
+		SubLoopSegmentCount = OtherSubLoopSegmentCount;
+		bNeedCheckOrientation = true;
+	}
+
+#ifdef DEBUG_REMOVE_UNIQUE_INTERSECTION
+	if (bDisplay)
+	{
+		DisplayLoops(TEXT("before removed"));
+		Wait();
+	}
+#endif
+
+	switch (SubLoopSegmentCount)
+	{
+	case 0:
+		ensureCADKernel(false);
+		break;
+	case 1:
+	case 2:
+		RemoveThePick(NodesOfLoop, Intersection, bForward);
+		break;
+	case 3:
+		SwapNodes(NodesOfLoop, Intersection, bForward);
+		break;
+	default:
+		TryToSwapSegmentsOrRemoveLoop(NodesOfLoop, Intersection, bForward);
+		break;
+	}
+
+#ifdef DEBUG_REMOVE_UNIQUE_INTERSECTION
+	if (bDisplay)
+	{
+		DisplayLoops(TEXT("After removed"));
+		Wait();
+	}
+#endif
+	return true;
+};
+
+bool FIsoTriangulator::SpreadCoincidentNodes(const TArray<FLoopNode*>& NodesOfLoop, TPair<double, double> Intersection)
+{
+	// TODO
+	// 	   Use case:
+	//	      FilesToProcess.append([1, 0, 0, 1, r"D:/Data/Cad Files/SolidWorks/p014 - Unreal Sport Bike/Headset front left brake.SLDPRT"])
+	//        selection = [1967]
+	
+	//using namespace IsoTriangulatorImpl;
+	//GetNextNodeMethod GetNext = bForward ? GetNextNodeImpl : GetPreviousNodeImpl;
+
+	//FLoopNode* Node0 = GetNodeAt(NodesOfLoop, (int32)Intersection.Key);
+	//FLoopNode* PreviousNode0 = Node0->
+	//FLoopNode* Node1 = GetNext(Node0);
+	//FLoopNode* NextNode = GetNext(Node1);
+
+	//const FPoint2D* PreviousPoint = &PreviousNode->Get2DPoint(EGridSpace::UniformScaled, Grid);
+	//const FPoint2D* Point0 = &Node0->Get2DPoint(EGridSpace::UniformScaled, Grid);
+	//const FPoint2D* Point1 = &Node1->Get2DPoint(EGridSpace::UniformScaled, Grid);
+	//const FPoint2D* NextPoint = &NextNode->Get2DPoint(EGridSpace::UniformScaled, Grid);
+	return true;
+}
+
+bool FIsoTriangulator::MovePickBehind(const TArray<FLoopNode*>& NodesOfLoop, TPair<double, double> Intersection, bool bKeyIsExtremity)
+{
+	// TODO
+	return true;
+}
+
+
+void FIsoTriangulator::RemoveIntersectionByMovingOutsideNodeInside(const FIsoSegment& IntersectingSegment, FLoopNode& NodeToMove)
+{
+	const FIsoNode* Nodes[2] = { nullptr, nullptr };
+	Nodes[0] = &IntersectingSegment.GetFirstNode();
+	Nodes[1] = &IntersectingSegment.GetSecondNode();
+
+	FPoint2D IntersectingPoints[2];
+	IntersectingPoints[0] = IntersectingSegment.GetFirstNode().Get2DPoint(EGridSpace::UniformScaled, Grid);
+	IntersectingPoints[1] = IntersectingSegment.GetSecondNode().Get2DPoint(EGridSpace::UniformScaled, Grid);
+
+	bool bEndNodeIsOutside = true;
+	FPoint2D PointToMove = NodeToMove.Get2DPoint(EGridSpace::UniformScaled, Grid);
+
+#ifdef DEBUG_CLOSED_OUSIDE_POINT
+	if (bDisplay)
+	{
+		F3DDebugSession _(TEXT("Outside Point"));
+		DisplayPoint(PointToMove, bEndNodeIsOutside ? EVisuProperty::GreenPoint : EVisuProperty::YellowPoint);
+	}
+#endif
+
+	double Coordinate;
+	FPoint2D ProjectedPoint = ProjectPointOnSegment(PointToMove, IntersectingPoints[0], IntersectingPoints[1], Coordinate);
+	MoveNode(NodeToMove, ProjectedPoint);
+}
+
+void FIsoTriangulator::MoveNode(FLoopNode& NodeToMove, FPoint2D& ProjectedPoint)
+{
+	const FPoint2D& PointToMove = NodeToMove.Get2DPoint(EGridSpace::UniformScaled, Grid);
+
+	FPoint2D ProjectedDirection = ProjectedPoint - PointToMove;
+	ProjectedDirection.Normalize();
+	ProjectedDirection *= GeometricTolerance;
+
+	FPoint2D NewCoordinate = ProjectedPoint + ProjectedDirection;
+
+	FLoopNode& PreviousNode = NodeToMove.GetPreviousNode();
+	FLoopNode& NextNode = NodeToMove.GetNextNode();
+
+	if ((PreviousNode.Get2DPoint(EGridSpace::UniformScaled, Grid).SquareDistance(NewCoordinate) < SquareGeometricTolerance2) ||
+		(NextNode.Get2DPoint(EGridSpace::UniformScaled, Grid).SquareDistance(NewCoordinate) < SquareGeometricTolerance2))
+	{
+		RemoveNodeOfLoop(NodeToMove);
+	}
+	else
+	{
+		NodeToMove.Set2DPoint(EGridSpace::UniformScaled, Grid, NewCoordinate);
+	}
+
+#ifdef DEBUG_MOVE_NODE
+	if (bDisplay)
+	{
+		{
+			F3DDebugSession _(TEXT("Point To Move"));
+			DisplayPoint(PointToMove, EVisuProperty::YellowPoint);
+		}
+		{
+			F3DDebugSession _(TEXT("Projected Point"));
+			DisplayPoint(ProjectedPoint, EVisuProperty::GreenPoint);
+		}
+		{
+			F3DDebugSession _(TEXT("New Position"));
+			DisplayPoint(NewCoordinate, EVisuProperty::BluePoint);
+			Wait(false);
+		}
+	}
+#endif
+}
+
+void FIsoTriangulator::RemoveThePick(const TArray<FLoopNode*>& NodesOfLoop, const TPair<double, double>& Intersection, bool bForward)
+{
+	using namespace IsoTriangulatorImpl;
+	GetNextNodeMethod GetNext = bForward ? GetNextNodeImpl : GetPreviousNodeImpl;
+
+	FLoopNode* PreviousNode = GetNodeAt(NodesOfLoop, (int32)Intersection.Key);
+	FLoopNode* Node0 = GetNext(PreviousNode);
+	FLoopNode* Node1 = GetNext(Node0);
+	FLoopNode* NextNode = GetNext(Node1);
+
+	const FPoint2D* PreviousPoint = &PreviousNode->Get2DPoint(EGridSpace::UniformScaled, Grid);
+	const FPoint2D* Point0 = &Node0->Get2DPoint(EGridSpace::UniformScaled, Grid);
+	const FPoint2D* Point1 = &Node1->Get2DPoint(EGridSpace::UniformScaled, Grid);
+	const FPoint2D* NextPoint = &NextNode->Get2DPoint(EGridSpace::UniformScaled, Grid);
+
+	double Slop0 = ComputeUnorientedSlope(*Point0, *PreviousPoint, *Point1);
+	double Slop1 = ComputeUnorientedSlope(*Point1, *PreviousPoint, *NextPoint);
+
+	if (Slop0 < Slop1)
+	{
+		RemoveNodeOfLoop(*Node0);
+		Node0 = PreviousNode;
+	}
+	else
+	{
+		RemoveNodeOfLoop(*Node1);
+		Node1 = NextNode;
+	}
+};
+
+bool FIsoTriangulator::RemovePickRecursively(FLoopNode* Node0, FLoopNode* Node1)
+{
+	FLoopNode* PreviousNode = &Node0->GetPreviousNode();
+	FLoopNode* NextNode = &Node1->GetNextNode();
+
+	const FPoint2D* PreviousPoint = &PreviousNode->Get2DPoint(EGridSpace::UniformScaled, Grid);
+	const FPoint2D* Point0 = &Node0->Get2DPoint(EGridSpace::UniformScaled, Grid);
+	const FPoint2D* Point1 = &Node1->Get2DPoint(EGridSpace::UniformScaled, Grid);
+	const FPoint2D* NextPoint = &NextNode->Get2DPoint(EGridSpace::UniformScaled, Grid);
+
+	bool bNodeRemoved = true;
+	while (true)
+	{
+		bool bNodeIsToDelete = (Point0->SquareDistance(*Point1) < SquareGeometricTolerance2);
+		if (bNodeIsToDelete)
+		{
+			RemoveNodeOfLoop(*Node0);
+		}
+
+		if (bNodeIsToDelete || CheckAndRemovePick(*PreviousPoint, *Point0, *Point1, *Node0))
+		{
+			if (PreviousNode->IsDelete())
+			{
+				break;
+			}
+
+			Point0 = PreviousPoint;
+			Node0 = PreviousNode;
+			PreviousNode = &PreviousNode->GetPreviousNode();
+			PreviousPoint = &PreviousNode->Get2DPoint(EGridSpace::UniformScaled, Grid);
+			continue;
+		}
+
+		if (CheckAndRemovePick(*Point0, *Point1, *NextPoint, *Node1))
+		{
+			if (NextNode->IsDelete())
+			{
+				break;
+			}
+
+			Point1 = NextPoint;
+			Node1 = NextNode;
+			NextNode = &NextNode->GetNextNode();
+			NextPoint = &NextNode->Get2DPoint(EGridSpace::UniformScaled, Grid);
+			continue;
+		}
+		break;
+	}
+	return bNodeRemoved;
+}
+
+
+// Case  ______a  c       ______a__b
+//              \/ \o	            \o
+//       ______d/\b/	  ______d__c/
+//	    			  
+void FIsoTriangulator::SwapNodes(const TArray<FLoopNode*>& NodesOfLoop, const TPair<double, double>& Intersection, bool bForward)
+{
+	using namespace IsoTriangulatorImpl;
+	GetNextNodeMethod GetNext = bForward ? GetNextNodeImpl : GetPreviousNodeImpl;
+
+	const int32 NodeCount = NodesOfLoop.Num();
+
+	FLoopNode* Node0 = GetNodeAt(NodesOfLoop, NextIndex(NodeCount, (int32)Intersection.Key));
+	FLoopNode* Pick = GetNext(Node0);
+	FLoopNode* Node1 = GetNext(Pick);
+
+	const FPoint2D Point0Copy = Node0->Get2DPoint(EGridSpace::UniformScaled, Grid);
+	const FPoint2D* Point1Ptr = &Node1->Get2DPoint(EGridSpace::UniformScaled, Grid);
+
+	Node0->Set2DPoint(EGridSpace::UniformScaled, Grid, *Point1Ptr);
+	Node1->Set2DPoint(EGridSpace::UniformScaled, Grid, Point0Copy);
+};
+
+void FIsoTriangulator::TryToSwapSegmentsOrRemoveLoop(const TArray<FLoopNode*>& NodesOfLoop, const TPair<double, double>& Intersection, bool bForward)
+{
+	using namespace IsoTriangulatorImpl;
+	const GetNextNodeMethod GetNext = bForward ? GetNextNodeImpl : GetPreviousNodeImpl;
+	const GetNextNodeMethod GetPrevious = bForward ? GetPreviousNodeImpl : GetNextNodeImpl;
+
+	const int32 NodeCount = NodesOfLoop.Num();
+
+	const int32 Segment0StartIndex = NextIndex(NodeCount, (int32)Intersection.Key);
+	const int32 Segment1EndIndex = (int32)Intersection.Value;
+
+	FLoopNode* Segment0_Node1 = GetNodeAt(NodesOfLoop, Segment0StartIndex);
+	FLoopNode* Segment1_Node0 = GetNodeAt(NodesOfLoop, Segment1EndIndex);
+
+	FLoopNode* Segment0_Node0 = GetPrevious(Segment0_Node1);
+	FLoopNode* Segment1_Node1 = GetNext(Segment1_Node0);
+
+	const FPoint2D& Segment0_Point0 = Segment0_Node0->Get2DPoint(EGridSpace::UniformScaled, Grid);
+	const FPoint2D& Segment0_Point1 = Segment0_Node1->Get2DPoint(EGridSpace::UniformScaled, Grid);
+	const FPoint2D& Segment1_Point0 = Segment1_Node0->Get2DPoint(EGridSpace::UniformScaled, Grid);
+	const FPoint2D& Segment1_Point1 = Segment1_Node1->Get2DPoint(EGridSpace::UniformScaled, Grid);
+
+#ifdef DEBUG_SWAP_SEGMENTS_OR_REMOVE
+	if (bDisplay)
+	{
+		F3DDebugSession _(TEXT("Intersected Segments"));
+		Display(EGridSpace::UniformScaled, *Segment0_Node0, *Segment0_Node1, 0, EVisuProperty::RedCurve);
+		Display(EGridSpace::UniformScaled, *Segment1_Node0, *Segment1_Node1, 0, EVisuProperty::RedCurve);
+		Display(EGridSpace::UniformScaled, *Segment0_Node0, 0, EVisuProperty::RedPoint);
+		Display(EGridSpace::UniformScaled, *Segment1_Node0, 0, EVisuProperty::RedPoint);
+		Wait(false);
+	}
+#endif
+
+	double Slop = ComputeSlope(Segment0_Point0, Segment0_Point1);
+	Slop = ComputeUnorientedSlope(Segment1_Point0, Segment1_Point1, Slop);
+
+	bool bIsFixed = false;
+	if (Slop < 2)
+	{
+		FIsoSegment* Segment0 = Segment0_Node0->GetSegmentConnectedTo(Segment0_Node1);
+		FIsoSegment* Segment1 = Segment1_Node0->GetSegmentConnectedTo(Segment1_Node1);
+		bIsFixed = TryToRemoveIntersectionBySwappingSegments(*Segment0, *Segment1);
+
+#ifdef DEBUG_SWAP_SEGMENTS_OR_REMOVE
+		if (bDisplay)
+		{
+			F3DDebugSession _(TEXT("New Segments"));
+			Display(EGridSpace::UniformScaled, *Segment0_Node0, Segment0_Node0->GetNextNode(), 0, EVisuProperty::BlueCurve);
+			Display(EGridSpace::UniformScaled, *Segment1_Node1, Segment1_Node1->GetPreviousNode(), 0, EVisuProperty::BlueCurve);
+			Display(EGridSpace::UniformScaled, *Segment0_Node0, 0, EVisuProperty::BluePoint);
+			Display(EGridSpace::UniformScaled, *Segment1_Node0, 0, EVisuProperty::BluePoint);
+		}
+#endif
+	}
+
+	if (!bIsFixed)
+	{
+		RemoveSubLoop(Segment0_Node0, Segment1_Node1, GetNext);
+	}
+};
+
+
+void FIsoTriangulator::FillIntersectionToolWithOuterLoop()
+{
+	for (FLoopNode& Node : LoopNodes)
+	{
+		if (Node.GetLoopIndex() != 0)
+		{
+			break;
+		}
+		LoopSegmentsIntersectionTool.AddSegment(*Node.GetSegmentConnectedTo(&Node.GetNextNode()));
+	}
+}
+
+bool FIsoTriangulator::FindLoopIntersectionAndFixIt()
+{
+	// move points for test
+	{
+		//FPoint2D NewCoordinate = LoopNodes[56].Get2DPoint(EGridSpace::UniformScaled, Grid);
+		//FPoint2D NewCoordinate53 = LoopNodes[55].Get2DPoint(EGridSpace::UniformScaled, Grid);
+		//NewCoordinate += FPoint(2, 0);
+		//const_cast<FLoopNode&>(LoopNodes[56]).Set2DPoint(EGridSpace::UniformScaled, Grid, NewCoordinate);
+
+		//NewCoordinate += FPoint(0.3, 1.0);
+		//const_cast<FLoopNode&>(LoopNodes[55]).Set2DPoint(EGridSpace::UniformScaled, Grid, NewCoordinate);
+		//NewCoordinate += FPoint(-0.1, -0.3);
+		//const_cast<FLoopNode&>(LoopNodes[54]).Set2DPoint(EGridSpace::UniformScaled, Grid, NewCoordinate);
+		//const_cast<FLoopNode&>(LoopNodes[53]).Set2DPoint(EGridSpace::UniformScaled, Grid, NewCoordinate53);
+
+		//FPoint2D NewCoordinate42 = LoopNodes[43].Get2DPoint(EGridSpace::UniformScaled, Grid);
+		//NewCoordinate42 += FPoint(0, -0.1);
+		//const_cast<FLoopNode&>(LoopNodes[42]).Set2DPoint(EGridSpace::UniformScaled, Grid, NewCoordinate42);
+
+
+		//NewCoordinate = LoopNodes[53].Get2DPoint(EGridSpace::UniformScaled, Grid);
+		//NewCoordinate += FPoint(2, 0);
+		//const_cast<FLoopNode&>(LoopNodes[53]).Set2DPoint(EGridSpace::UniformScaled, Grid, NewCoordinate);
+
+
+		//NewCoordinate = LoopNodes[49].Get2DPoint(EGridSpace::UniformScaled, Grid);
+		//NewCoordinate += FPoint(-0.05, -1);
+		//const_cast<FLoopNode&>(LoopNodes[49]).Set2DPoint(EGridSpace::UniformScaled, Grid, NewCoordinate);
+
+		//FPoint2D NewCoordinate = LoopNodes[126].Get2DPoint(EGridSpace::UniformScaled, Grid);
+		//NewCoordinate += FPoint(0, 3);
+		//const_cast<FLoopNode&>(LoopNodes[126]).Set2DPoint(EGridSpace::UniformScaled, Grid, NewCoordinate);
+	}
+
+	TArray<FLoopNode*> BestStartNodeOfLoop;
+	FindBestLoopExtremity(BestStartNodeOfLoop);
+
+#ifdef DEBUG_LOOP_INTERSECTION_AND_FIX_IT		
+	if (bDisplay)
+	{
+		Display(EGridSpace::UniformScaled, TEXT("Loops Orientation"), LoopSegments, false, true, EVisuProperty::BlueCurve);
+		DisplayLoops(TEXT("FindLoopIntersectionAndFixIt Before"), false, true);
+		F3DDebugSession _(TEXT("BestStartNodeOfLoop"));
+		for (const FLoopNode* Node : BestStartNodeOfLoop)
+		{
+			Display(EGridSpace::UniformScaled, *Node, 0, EVisuProperty::BluePoint);
+		}
+		Wait();
+	}
+#endif
+
+	TArray<FLoopNode*> LoopNodesFromStartNode;
+	TArray<TPair<double, double>> Intersections;
+
+	// for each loop, start by the best node, find all intersections
+	bool bIsAnOuterLoop = true;
+	for (FLoopNode* StartNode : BestStartNodeOfLoop)
+	{
+		bNeedCheckOrientation = false;
+		LoopNodesFromStartNode.Empty(LoopNodes.Num());
+		Intersections.Empty(5);
+
+		// LoopNodesFromStartNode is the set of node of the loop oriented as if the loop is an external loop
+		// The use of FLoopNode::GetNextNode() is not recommended
+		// Prefer to use GetSegmentToNodeMethode GetNext() set with GetFirstNode of GetSecondNode according to bOuterLoop
+
+		IsoTriangulatorImpl::GetLoopNodeStartingFrom(StartNode, bIsAnOuterLoop, LoopNodesFromStartNode);
+
+#ifdef DEBUG_LOOP_INTERSECTION_AND_FIX_IT		
+		DisplayLoop(EGridSpace::UniformScaled, TEXT("LoopIntersections: start"), LoopNodesFromStartNode, true, EVisuProperty::YellowPoint);
+		Wait(bDisplay);
+#endif
+
+		RemoveLoopPicks(LoopNodesFromStartNode, Intersections);
+
+		if (LoopNodesFromStartNode.Num() == 0)
+		{
+			continue;
+		}
+
+#ifdef DEBUG_LOOP_INTERSECTION_AND_FIX_IT		
+		DisplayLoop(EGridSpace::UniformScaled, TEXT("LoopIntersections: remove pick"), LoopNodesFromStartNode, true, EVisuProperty::YellowPoint);
+		Wait(bDisplay);
+#endif
+		// At this step, LoopNodesFromStartNode cannot have deleted node
+		FindLoopIntersections(LoopNodesFromStartNode, bIsAnOuterLoop, Intersections);
+
+		// WARNING From this step, RemoveLoopIntersections can delete nodes. 
+		RemoveLoopIntersections(LoopNodesFromStartNode, Intersections, bIsAnOuterLoop);
+
+#ifdef DEBUG_LOOP_INTERSECTION_AND_FIX_IT		
+ 		DisplayLoop(EGridSpace::UniformScaled, TEXT("LoopIntersections: remove self intersection"), LoopNodesFromStartNode, true, EVisuProperty::YellowPoint);
+		Wait(bDisplay);
+#endif
+
+		RemoveLoopPicks(LoopNodesFromStartNode, Intersections);
+
+		if (LoopNodesFromStartNode.Num() == 0)
+		{
+			continue;
+		}
+
+#ifdef DEBUG_LOOP_INTERSECTION_AND_FIX_IT		
+		DisplayLoop(EGridSpace::UniformScaled, TEXT("LoopIntersections: remove pick"), LoopNodesFromStartNode, true, EVisuProperty::YellowPoint);
+		Wait(bDisplay);
+#endif
+
+		FixLoopOrientation(LoopNodesFromStartNode);
+
+		bIsAnOuterLoop = false;
+	}
+
+	if (!CheckMainLoopConsistency())
+	{
+		return false;
+	}
+
+
+	if (Grid.GetLoopCount() > 1)
+	{
+		// Remove intersection between loops
+		FixIntersectionBetweenLoops();
+
+#ifdef DEBUG_LOOP_INTERSECTION_AND_FIX_IT		
+		if (bDisplay)
+		{
+			DisplayLoops(TEXT("FindLoopIntersectionAndFixIt Step2"), false, true);
+			Wait();
+		}
+#endif
+	}
+	else
+	{
+		LoopSegmentsIntersectionTool.Empty(LoopSegments.Num());
+		for (FIsoSegment* Segment : LoopSegments)
+		{
+			LoopSegmentsIntersectionTool.AddSegment(*Segment);
+		}
+	}
+
+	if (!CheckMainLoopConsistency())
+	{
+		return false;
+	}
+
+#ifdef DEBUG_LOOP_INTERSECTION_AND_FIX_IT		
+	if (bDisplay)
+	{
+		DisplayLoops(TEXT("FindLoopIntersectionAndFixIt 3"), false, true);
+		Wait();
+	}
+#endif
+
+	return true;
+}
+
+void FIsoTriangulator::FixLoopOrientation(const TArray<FLoopNode*>& NodesOfLoop)
+{
+	FLoopNode const* const* StartNode = NodesOfLoop.FindByPredicate([](const FLoopNode* Node) { return !Node->IsDelete(); });
+	if (StartNode != nullptr)
+	{
+		EOrientation Orientation = GetLoopOrientation(*StartNode);
+		if (Orientation == EOrientation::Back)
+		{
+			const int32 LoopIndex = (*StartNode)->GetLoopIndex();
+			const int32 FirstSegmentIndex = LoopSegments.IndexOfByPredicate([&](const FIsoSegment* Segment) { return ((const FLoopNode&)Segment->GetFirstNode()).GetLoopIndex() == LoopIndex; });
+			int32 LastSegmentIndex = (Grid.GetLoopCount() == LoopIndex + 1) ? LoopSegments.Num() : LoopSegments.IndexOfByPredicate([&](const FIsoSegment* Segment) { return ((const FLoopNode&)Segment->GetFirstNode()).GetLoopIndex() > LoopIndex; });
+			if (LastSegmentIndex == INDEX_NONE)
+			{
+				LastSegmentIndex = LoopSegments.Num();
+			}
+			SwapLoopOrientation(FirstSegmentIndex, LastSegmentIndex);
+#ifdef DEBUG_LOOP_ORIENTATION
+			Display(EGridSpace::UniformScaled, TEXT("After orientation"), LoopSegments, true, true);
+			DisplayLoops(TEXT("After orientation"), false, true, true);
+			Wait();
+#endif
+		}
+	}
+}
+
+bool FIsoTriangulator::CheckMainLoopConsistency()
+{
+	int32 OuterNodeCount = 0;
+	for (const FLoopNode& Node : LoopNodes)
+	{
+		if (Node.GetLoopIndex() != 0)
+		{
+			break;
+		}
+
+		if (!Node.IsDelete())
+		{
+			++OuterNodeCount;
+			if (OuterNodeCount > 2)
+			{
+				return true;
+			}
+		}
+	}
+	return false;
+}
+
+//#define DEBUG_CLOSED_OUSIDE_POINT
+void FIsoTriangulator::RemoveIntersectionByMovingOutsideSegmentNodeInside(const FIsoSegment& IntersectingSegment, const FIsoSegment& Segment)
+{
+	const FIsoNode* Nodes[2][2] = { {nullptr, nullptr}, {nullptr, nullptr} };
+	Nodes[0][0] = &IntersectingSegment.GetFirstNode();
+	Nodes[0][1] = &IntersectingSegment.GetSecondNode();
+	Nodes[1][0] = &Segment.GetFirstNode();
+	Nodes[1][1] = &Segment.GetSecondNode();
+
+	FPoint2D IntersectingPoints[2];
+	IntersectingPoints[0] = IntersectingSegment.GetFirstNode().Get2DPoint(EGridSpace::UniformScaled, Grid);
+	IntersectingPoints[1] = IntersectingSegment.GetSecondNode().Get2DPoint(EGridSpace::UniformScaled, Grid);
+
+	FPoint2D SegmentPoints[2];
+	SegmentPoints[0] = Segment.GetFirstNode().Get2DPoint(EGridSpace::UniformScaled, Grid);
+	SegmentPoints[1] = Segment.GetSecondNode().Get2DPoint(EGridSpace::UniformScaled, Grid);
+
+	double SquareLenghtIntersectingSegment = IntersectingPoints[0].SquareDistance(IntersectingPoints[1]);
+	double SquareLenghtSegment = SegmentPoints[0].SquareDistance(SegmentPoints[1]);
+
+	if (SquareLenghtSegment > 10 * SquareLenghtIntersectingSegment)
+	{
+		return RemoveIntersectionByMovingOutsideSegmentNodeInside(Segment, const_cast<FIsoSegment&>(IntersectingSegment));
+	}
+
+	bool bFirstNodeIsOutside = false;
+	FPoint2D PointToMove = Segment.GetSecondNode().Get2DPoint(EGridSpace::UniformScaled, Grid);
+
+	// Is Second node, the outside node ?
+	double OrientedSlop = ComputeOrientedSlope(IntersectingPoints[0], IntersectingPoints[1], PointToMove);
+	if (OrientedSlop > 0)
+	{
+		PointToMove = Segment.GetFirstNode().Get2DPoint(EGridSpace::UniformScaled, Grid);
+		bFirstNodeIsOutside = true;
+	}
+
+#ifdef DEBUG_CLOSED_OUSIDE_POINT
+	if (bDisplay)
+	{
+		F3DDebugSession _(TEXT("Outside Point"));
+		DisplayPoint(PointToMove, bFirstNodeIsOutside ? EVisuProperty::GreenPoint : EVisuProperty::YellowPoint);
+	}
+#endif
+
+	double Coordinate;
+	FPoint2D ProjectedPoint = ProjectPointOnSegment(PointToMove, IntersectingPoints[0], IntersectingPoints[1], Coordinate);
+
+	FPoint2D MoveDirection = IntersectingPoints[1] - IntersectingPoints[0];
+	MoveDirection.Normalize();
+	MoveDirection = MoveDirection.GetPerpendicularVector();
+	MoveDirection *= GeometricTolerance;
+
+	FPoint2D NewCoordinate = ProjectedPoint + MoveDirection;
+
+	FLoopNode& Node = (FLoopNode&) const_cast<FIsoNode&>(!bFirstNodeIsOutside ? Segment.GetSecondNode() : Segment.GetFirstNode());
+	FLoopNode& PreviousNode = Node.GetPreviousNode();
+	FLoopNode& NextNode = Node.GetNextNode();
+
+	if ((PreviousNode.Get2DPoint(EGridSpace::UniformScaled, Grid).SquareDistance(NewCoordinate) < SquareGeometricTolerance2) ||
+		(NextNode.Get2DPoint(EGridSpace::UniformScaled, Grid).SquareDistance(NewCoordinate) < SquareGeometricTolerance2))
+	{
+		RemoveNodeOfLoop(Node);
+		return;
+	}
+
+	Node.Set2DPoint(EGridSpace::UniformScaled, Grid, NewCoordinate);
+
+#ifdef DEBUG_CLOSED_OUSIDE_POINT
+	if (bDisplay)
+	{
+		F3DDebugSession _(TEXT("New Segs"));
+		DisplayPoint(NewCoordinate, EVisuProperty::BluePoint);
+		Display(EGridSpace::UniformScaled, Segment, 0, EVisuProperty::BlueCurve);
+		Display(EGridSpace::UniformScaled, IntersectingSegment, 0, EVisuProperty::RedCurve);
+	}
+#endif
+}
+
+
+void FIsoTriangulator::OffsetSegment(FIsoSegment& Segment, TSegment<FPoint2D>& Segment2D, TSegment<FPoint2D>& IntersectingSegment2D)
+{
+	FPoint2D SegmentTangent = IntersectingSegment2D.Point1 - IntersectingSegment2D.Point0;
+	SegmentTangent.Normalize();
+	FPoint2D MoveDirection = SegmentTangent.GetPerpendicularVector();
+	MoveDirection *= GeometricTolerance;
+
+	FPoint2D NewPoint0 = Segment2D.Point0 + MoveDirection;
+	FPoint2D NewPoint1 = Segment2D.Point1 + MoveDirection;
+
+	Segment.GetFirstNode().Set2DPoint(EGridSpace::UniformScaled, Grid, NewPoint0);
+	Segment.GetSecondNode().Set2DPoint(EGridSpace::UniformScaled, Grid, NewPoint1);
+}
+
+void FIsoTriangulator::OffsetNode(FLoopNode& Node, TSegment<FPoint2D>& IntersectingSegment2D)
+{
+	FPoint2D SegmentTangent = IntersectingSegment2D.Point1 - IntersectingSegment2D.Point0;
+	SegmentTangent.Normalize();
+	FPoint2D MoveDirection = SegmentTangent.GetPerpendicularVector();
+	MoveDirection *= GeometricTolerance;
+
+	FPoint2D NewPoint = Node.Get2DPoint(EGridSpace::UniformScaled, Grid) + MoveDirection;
+	Node.Set2DPoint(EGridSpace::UniformScaled, Grid, NewPoint);
+}
+
+//#define DEBUG_TWO_CONSECUTIVE_INTERSECTING
+bool FIsoTriangulator::TryToRemoveIntersectionOfTwoConsecutiveIntersectingSegments(const FIsoSegment& IntersectingSegment, FIsoSegment& Segment)
+{
+	FLoopNode* Node = nullptr;
+	FLoopNode* PreviousNode = nullptr;
+	FLoopNode* NextNode = nullptr;
+
+	TSegment<FPoint2D> IntersectingSegment2D(IntersectingSegment.GetFirstNode().Get2DPoint(EGridSpace::UniformScaled, Grid), IntersectingSegment.GetSecondNode().Get2DPoint(EGridSpace::UniformScaled, Grid));
+	TSegment<FPoint2D> Segment2D(Segment.GetFirstNode().Get2DPoint(EGridSpace::UniformScaled, Grid), Segment.GetSecondNode().Get2DPoint(EGridSpace::UniformScaled, Grid));
+
+	double IntersectingSegmentSlop = ComputeOrientedSlope(IntersectingSegment2D.Point0, IntersectingSegment2D.Point1, 0);
+	double SegmentSlop = ComputeUnorientedSlope(Segment2D.Point1, Segment2D.Point0, IntersectingSegmentSlop);
+	if(SegmentSlop > 2)
+	{
+		SegmentSlop = 4 - SegmentSlop;
+	}
+
+	// if the segment and IntersectingSegment are parallel, segment are moved inside
+	if(SegmentSlop < 0.01)
+	{
+		double StartPointSquareDistance = SquareDistanceOfPointToSegment(Segment2D.Point0, IntersectingSegment2D.Point0, IntersectingSegment2D.Point1);
+		double EndPointSquareDistance = SquareDistanceOfPointToSegment(Segment2D.Point1, IntersectingSegment2D.Point0, IntersectingSegment2D.Point1);
+		if(StartPointSquareDistance < SquareGeometricTolerance && EndPointSquareDistance < SquareGeometricTolerance)
+		{
+			OffsetSegment(Segment, Segment2D, IntersectingSegment2D);
+			return true;
+		}
+	}
+
+	// check if the intersection is not at the extremity
+	{
+		double Coordinate = 0;
+		FindIntersectionOfSegments2D(Segment2D, IntersectingSegment2D, Coordinate);
+		if (FMath::IsNearlyZero(Coordinate))
+		{
+			// can add a test to offset the outside node and not the node a 
+			OffsetNode((FLoopNode&)Segment.GetFirstNode(), IntersectingSegment2D);
+			return true;
+		}
+		else if (FMath::IsNearlyEqual(Coordinate, 1))
+		{
+			OffsetNode((FLoopNode&) Segment.GetSecondNode(), IntersectingSegment2D);
+			return true;
+		}
+	}
+
+	double OrientedSlop = ComputeOrientedSlope(IntersectingSegment2D.Point0, Segment.GetFirstNode().Get2DPoint(EGridSpace::UniformScaled, Grid), IntersectingSegmentSlop);
+	if (OrientedSlop >= 0)
+	{
+		Node = (FLoopNode*)&Segment.GetSecondNode();
+		PreviousNode = (FLoopNode*)&Segment.GetFirstNode();
+		NextNode = (FLoopNode*)&Node->GetNextNode();
+	}
+	else
+	{
+		Node = (FLoopNode*)&Segment.GetFirstNode();
+		PreviousNode = (FLoopNode*)&Segment.GetSecondNode();
+		NextNode = (FLoopNode*)&Node->GetPreviousNode();
+	}
+
+#ifdef DEBUG_TWO_CONSECUTIVE_INTERSECTING
+	if (bDisplay)
+	{
+		F3DDebugSession _(TEXT("Intersecting Segments"));
+		Display(EGridSpace::UniformScaled, Segment, 0, EVisuProperty::BlueCurve);
+		Display(EGridSpace::UniformScaled, *Node, *NextNode, 0, EVisuProperty::BlueCurve);
+		Display(EGridSpace::UniformScaled, IntersectingSegment, 0, EVisuProperty::RedCurve);
+		Display(EGridSpace::UniformScaled, *Node, 0, EVisuProperty::RedPoint);
+		Wait(false);
+	}
+#endif
+
+	TSegment<FPoint2D> NextSegment2D(Node->Get2DPoint(EGridSpace::UniformScaled, Grid), NextNode->Get2DPoint(EGridSpace::UniformScaled, Grid));
+	if (!FastIntersectSegments2D(NextSegment2D, IntersectingSegment2D))
+	{
+		return false;
+	}
+
+	TSegment<FPoint2D>  PreviousSegment2D(Node->Get2DPoint(EGridSpace::UniformScaled, Grid), PreviousNode->Get2DPoint(EGridSpace::UniformScaled, Grid));
+
+	FPoint2D Intersection1 = FindIntersectionOfSegments2D(PreviousSegment2D, IntersectingSegment2D);
+	FPoint2D Intersection2 = FindIntersectionOfSegments2D(NextSegment2D, IntersectingSegment2D);
+
+	double Coordinate;
+	FPoint2D ProjectedPoint = ProjectPointOnSegment(Node->Get2DPoint(EGridSpace::UniformScaled, Grid), Intersection1, Intersection2, Coordinate);
+
+#ifdef DEBUG_TWO_CONSECUTIVE_INTERSECTING
+	if (bDisplay)
+	{
+		F3DDebugSession _(TEXT("ProjectedPoint"));
+		DisplayPoint(ProjectedPoint, EVisuProperty::RedPoint);
+	}
+#endif
+
+	FPoint2D SegmentTangent = IntersectingSegment2D.Point1 - IntersectingSegment2D.Point0;
+	SegmentTangent.Normalize();
+	FPoint2D MoveDirection = SegmentTangent.GetPerpendicularVector();
+
+	MoveDirection *= GeometricTolerance;
+
+	FPoint2D NewCoordinate = ProjectedPoint + MoveDirection;
+
+#ifdef DEBUG_TWO_CONSECUTIVE_INTERSECTING
+	if (bDisplay)
+	{
+		F3DDebugSession _(TEXT("NewCoordinate"));
+		DisplayPoint(NewCoordinate, EVisuProperty::BluePoint);
+	}
+#endif
+
+	if ((PreviousNode->Get2DPoint(EGridSpace::UniformScaled, Grid).SquareDistance(NewCoordinate) < SquareGeometricTolerance2) ||
+		(NextNode->Get2DPoint(EGridSpace::UniformScaled, Grid).SquareDistance(NewCoordinate) < SquareGeometricTolerance2))
+	{
+		RemoveNodeOfLoop(*Node);
+	}
+	else
+	{
+		Node->Set2DPoint(EGridSpace::UniformScaled, Grid, NewCoordinate);
+		FIsoSegment* NextSegment = Node->GetSegmentConnectedTo(NextNode);
+		LoopSegmentsIntersectionTool.Update(&Segment);
+		LoopSegmentsIntersectionTool.Update(NextSegment);
+
+#ifdef DEBUG_TWO_CONSECUTIVE_INTERSECTING
+		if (bDisplay)
+		{
+			F3DDebugSession _(TEXT("New position"));
+			DisplayPoint(Intersection1, EVisuProperty::RedPoint);
+			DisplayPoint(Intersection2, EVisuProperty::RedPoint);
+			Display(EGridSpace::UniformScaled, *Node, EVisuProperty::RedPoint);
+			Display(EGridSpace::UniformScaled, Segment, 0, EVisuProperty::GreenCurve/*, true*/);
+			Display(EGridSpace::UniformScaled, *Node, *NextNode, 0, EVisuProperty::GreenCurve/*, true*/);
+			Wait(false);
+		}
+#endif
+	}
+
+	return true;
+}
+
+//#define DEBUG_FIND_LOOP_INTERSECTION_AND_FIX_IT
+void FIsoTriangulator::FixIntersectionBetweenLoops()
+{
+	const double MaxGap = Grid.GetMinElementSize();
+
+	ensureCADKernel(2 < LoopSegments.Num());
+
+#ifdef DEBUG_FIND_LOOP_INTERSECTION_AND_FIX_IT
+	int32 Iteration = 0;
+	F3DDebugSession _(bDisplay, TEXT("FixIntersectionBetweenLoops"));
+#endif
+
+	TSet<uint32> IntersectionAlreadyProceed;
+
+	LoopSegmentsIntersectionTool.Empty(LoopSegments.Num());
+	int32 Index = 0;
+	for (; Index < LoopSegments.Num(); ++Index)
+	{
+		if(((FLoopNode&) LoopSegments[Index]->GetFirstNode()).GetLoopIndex() != 0)
+		{
+			break;
+		}
+		LoopSegmentsIntersectionTool.AddSegment(*LoopSegments[Index]);
+	}
+
+	for (; Index < LoopSegments.Num(); )
+	{
+		ensureCADKernel(LoopSegments[Index]);
+
+		FIsoSegment& Segment = *LoopSegments[Index];
+		ensureCADKernel(!Segment.IsDelete());
+
+#ifdef DEBUG_FIND_LOOP_INTERSECTION_AND_FIX_IT
+		if (bDisplay)
+		{
+			LoopSegmentsIntersectionTool.Display(TEXT("IntersectionTool"));
+			F3DDebugSession _(*FString::Printf(TEXT("Segment to proceed %d %d"), Index, Iteration++));
+			Display(EGridSpace::UniformScaled, Segment, 0, EVisuProperty::BlueCurve/*, true*/);
+		}
+#endif
+
+		if (const FIsoSegment* IntersectingSegment = LoopSegmentsIntersectionTool.DoesIntersect(Segment))
+		{
+#ifdef DEBUG_FIND_LOOP_INTERSECTION_AND_FIX_IT
+			if (bDisplay)
+			{
+				LoopSegmentsIntersectionTool.Display(TEXT("IntersectionTool"));
+				{
+					F3DDebugSession _(*FString::Printf(TEXT("Segment to proceed %d %d"), Index, Iteration++));
+					Display(EGridSpace::UniformScaled, Segment, 0, EVisuProperty::BlueCurve/*, true*/);
+				}
+				{
+					F3DDebugSession _(TEXT("Intersecting Segments"));
+					Display(EGridSpace::UniformScaled, *IntersectingSegment, 0, EVisuProperty::RedCurve/*, true*/);
+				}
+				Wait(true);
+			}
+#endif
+
+  			uint32 IntersectionHash = GetTypeHash(*IntersectingSegment, Segment);
+			bool bNotProceed = !IntersectionAlreadyProceed.Find(IntersectionHash);
+			IntersectionAlreadyProceed.Add(IntersectionHash);
+
+			bool bIsFixed = true;
+			bool bIsSameLoop = ((FLoopNode&)Segment.GetFirstNode()).GetLoopIndex() == ((FLoopNode&)IntersectingSegment->GetFirstNode()).GetLoopIndex();
+
+			// Swapping segments is possible only with outer loop
+			if (bNotProceed)
+			{
+				if (!TryToRemoveIntersectionOfTwoConsecutiveIntersectingSegments(*IntersectingSegment, Segment))
+				{
+					if (!TryToRemoveIntersectionOfTwoConsecutiveIntersectingSegments(Segment, const_cast<FIsoSegment&>(*IntersectingSegment)))
+					{
+						if (bIsSameLoop)
+						{
+							// check if it's realy used 
+							bIsFixed = TryToRemoveSelfIntersectionByMovingTheClosedOusidePoint(*IntersectingSegment, Segment);
+						}
+						else
+						{
+							RemoveIntersectionByMovingOutsideSegmentNodeInside(*IntersectingSegment, Segment);
+						}
+					}
+				}
+
+				if (bIsFixed)
+				{
+					// segment is moved, the previous segment is retested. Thanks to that, only one loop is enough.
+					if (Index > 1)
+					{
+						LoopSegmentsIntersectionTool.RemoveLast();
+					}
+				}
+			}
+			else if (bIsSameLoop)
+			{
+				bIsFixed = TryToRemoveIntersectionBySwappingSegments(const_cast<FIsoSegment&>(*IntersectingSegment), Segment);
+				if(!bIsFixed)
+				{
+					ensureCADKernel(false);
+				}
+			}
+			else
+			{
+				// if no more append remove it
+				ensureCADKernel(false);
+				LoopSegmentsIntersectionTool.AddSegment(Segment);
+			}
+
+#ifdef DEBUG_FIND_LOOP_INTERSECTION_AND_FIX_IT
+			if (false)
+			{
+				DisplayLoops(TEXT("After fix"));
+				Wait(false);
+			}
+#endif
+		}
+		else
+		{
+			LoopSegmentsIntersectionTool.AddSegment(Segment);
+		}
+
+		if (!Segment.IsDelete())
+		{
+			RemovePickOfLoop(Segment);
+		}
+
+		Index = LoopSegmentsIntersectionTool.Count();
+	}
+
+#ifdef DEBUG_FIND_LOOP_INTERSECTION_AND_FIX_IT
+	if (bDisplay)
+	{
+		DisplayLoops(TEXT("After fix"));
+		Wait(false);
+	}
+#endif
+
+	LoopSegmentsIntersectionTool.Sort();
+}
+
+//#define DEBUG_SELF_CLOSED_OUSIDE_POINT
+bool FIsoTriangulator::TryToRemoveSelfIntersectionByMovingTheClosedOusidePoint(const FIsoSegment& Segment0, const FIsoSegment& Segment1)
+{
+	const FIsoNode* Nodes[2][2] = { {nullptr, nullptr}, {nullptr, nullptr} };
+	Nodes[0][0] = &Segment0.GetFirstNode();
+	Nodes[0][1] = &Segment0.GetSecondNode();
+	Nodes[1][0] = &Segment1.GetFirstNode();
+	Nodes[1][1] = &Segment1.GetSecondNode();
+
+	FPoint2D Points[2][2];
+	Points[0][0] = Nodes[0][0]->Get2DPoint(EGridSpace::UniformScaled, Grid);
+	Points[0][1] = Nodes[0][1]->Get2DPoint(EGridSpace::UniformScaled, Grid);
+	Points[1][0] = Nodes[1][0]->Get2DPoint(EGridSpace::UniformScaled, Grid);
+	Points[1][1] = Nodes[1][1]->Get2DPoint(EGridSpace::UniformScaled, Grid);
+
+	int32 ProjectedNodeIndex[2][2] = { { 0, 0 }, { 0, 0 } };
+	FPoint2D ProjectedPoints[2][2];
+	double Distance[2][2] = { { HUGE_VALUE, HUGE_VALUE }, { HUGE_VALUE, HUGE_VALUE } };
+	FPoint2D ProjectedPointsOut[2][2];
+
+	//int32 OtherSegmentIndex = 0;
+	int32 SegmentIndex = 1;
+	int32 OtherSegmentIndex = 0;
+
+	TFunction<void(int32)> ProjectPoint = [&](int32 OtherSegmentNodeIndex)
+	{
+		double Coordinate;
+		ProjectedPoints[OtherSegmentIndex][OtherSegmentNodeIndex] = ProjectPointOnSegment(Points[OtherSegmentIndex][OtherSegmentNodeIndex], Points[SegmentIndex][0], Points[SegmentIndex][1], Coordinate, false);
+		bool bProjectedPoint1IsInside = (Coordinate >= -SMALL_NUMBER && Coordinate <= 1 + SMALL_NUMBER);
+		if (bProjectedPoint1IsInside)
+		{
+			Distance[OtherSegmentIndex][OtherSegmentNodeIndex] = ProjectedPoints[OtherSegmentIndex][OtherSegmentNodeIndex].SquareDistance(Points[OtherSegmentIndex][OtherSegmentNodeIndex]);
+		}
+
+#ifdef DEBUG_SELF_CLOSED_OUSIDE_POINT
+		if (bDisplay)
+		{
+			{
+				F3DDebugSession _(TEXT("Segs"));
+				DisplaySegment(Points[SegmentIndex][0], Points[SegmentIndex][1], 0, EVisuProperty::BlueCurve);
+				DisplayPoint(Points[OtherSegmentIndex][OtherSegmentNodeIndex], EVisuProperty::RedPoint);
+				DisplayPoint(ProjectedPoints[OtherSegmentIndex][OtherSegmentNodeIndex], EVisuProperty::RedPoint);
+				Wait(false);
+			}
+		}
+#endif
+	};
+
+	SegmentIndex = 0;
+	OtherSegmentIndex = 1;
+	ProjectPoint(0);
+	ProjectPoint(1);
+	SegmentIndex = 1;
+	OtherSegmentIndex = 0;
+	ProjectPoint(0);
+	ProjectPoint(1);
+
+	TFunction<void(int32, int32)> MoveNode = [&](int32 MoveSegmentIndex, int32 MovePointIndex)
+	{
+		FPoint2D MoveDirection = ProjectedPoints[MoveSegmentIndex][MovePointIndex] - Points[MoveSegmentIndex][MovePointIndex];
+		MoveDirection.Normalize();
+
+		if (MoveDirection.SquareLength() < 0.5)
+		{
+			MoveDirection = Points[MoveSegmentIndex][MovePointIndex == 0 ? 1 : 0] - ProjectedPoints[MoveSegmentIndex][MovePointIndex];
+			MoveDirection.Normalize();
+		}
+
+		MoveDirection *= GeometricTolerance;
+		FPoint2D NewCoordinate = ProjectedPoints[MoveSegmentIndex][MovePointIndex] + MoveDirection;
+
+		FLoopNode& Node = (FLoopNode&) const_cast<FIsoNode&>(*Nodes[MoveSegmentIndex][MovePointIndex]);
+		FLoopNode& PreviousNode = Node.GetPreviousNode();
+		FLoopNode& NextNode = Node.GetNextNode();
+
+		if ((PreviousNode.Get2DPoint(EGridSpace::UniformScaled, Grid).SquareDistance(NewCoordinate) < SquareGeometricTolerance2) ||
+			(NextNode.Get2DPoint(EGridSpace::UniformScaled, Grid).SquareDistance(NewCoordinate) < SquareGeometricTolerance2))
+		{
+			RemoveNodeOfLoop(Node);
+			return;
+		}
+
+		const_cast<FIsoNode&>((*Nodes[MoveSegmentIndex][MovePointIndex])).Set2DPoint(EGridSpace::UniformScaled, Grid, NewCoordinate);
+
+#ifdef DEBUG_SELF_CLOSED_OUSIDE_POINT
+		if (bDisplay)
+		{
+			F3DDebugSession _(TEXT("New Segs"));
+			DisplayPoint(NewCoordinate, EVisuProperty::RedPoint);
+			Display(EGridSpace::UniformScaled, Segment1, 0, EVisuProperty::RedCurve);
+			Display(EGridSpace::UniformScaled, Segment0, 0, EVisuProperty::BlueCurve);
+			Wait(false);
+		}
+#endif
+	};
+
+
+	int32 MoveSegment = 0;
+	int32 MovePoint = 0;
+	double MinDistance = Distance[0][0];
+
+	TFunction<void(int32, int32)> FindNodeToMove = [&](int32 MoveSegmentIndex, int32 MovePointIndex)
+	{
+		if (Distance[MoveSegmentIndex][MovePointIndex] < MinDistance)
+		{
+			MinDistance = Distance[MoveSegmentIndex][MovePointIndex];
+			MoveSegment = MoveSegmentIndex;
+			MovePoint = MovePointIndex;
+		}
+	};
+
+	FindNodeToMove(0, 1);
+	FindNodeToMove(1, 0);
+	FindNodeToMove(1, 1);
+
+	MoveNode(MoveSegment, MovePoint);
+
+	return true;
+}
+
+bool FIsoTriangulator::TryToRemoveIntersectionByMovingTheClosedOusidePoint(const FIsoSegment& Segment0, const FIsoSegment& Segment1)
+{
+	// For Inner/outer segment intersecting outer loop (i.e. IntersectingSegment is Outer), the segment has to be moved inside IntersectingSegment
+	// For Inner segment intersecting its loop, the segment has to be moved inside IntersectingSegment but the orientation of the loop is counterclockwise => need to swap orientation
+	// For Inner segment intersecting another inner loop, the segment has to be moved outside the loop but the orientation of the loop is counterclockwise => don't need to swap orientation
+
+	// As the iteration of the segment start from outer loop's segments, Segment can not be from outer loop if IntersectingSegment is not from outer loop
+	int32 SegmentLoopIndex[2];
+	SegmentLoopIndex[0] = Segment0.GetFirstNode().IsALoopNode() ? ((FLoopNode&)Segment0.GetFirstNode()).GetLoopIndex() : 0;
+	SegmentLoopIndex[1] = Segment1.GetFirstNode().IsALoopNode() ? ((FLoopNode&)Segment1.GetFirstNode()).GetLoopIndex() : 0;
+	bool bIsOutLoopSegment = SegmentLoopIndex[0] == 0;
+	bool bIsSameLoop = SegmentLoopIndex[0] == SegmentLoopIndex[1];
+
+	bool bToInside = bIsOutLoopSegment || bIsSameLoop;
+
+	const FIsoNode* Nodes[2][2] = { {nullptr, nullptr}, {nullptr, nullptr} };
+	Nodes[0][0] = &Segment0.GetFirstNode();
+	Nodes[0][1] = &Segment0.GetSecondNode();
+	Nodes[1][0] = &Segment1.GetFirstNode();
+	Nodes[1][1] = &Segment1.GetSecondNode();
+
+	FPoint2D Points[2][2];
+	Points[0][0] = Nodes[0][0]->Get2DPoint(EGridSpace::UniformScaled, Grid);
+	Points[0][1] = Nodes[0][1]->Get2DPoint(EGridSpace::UniformScaled, Grid);
+	Points[1][0] = Nodes[1][0]->Get2DPoint(EGridSpace::UniformScaled, Grid);
+	Points[1][1] = Nodes[1][1]->Get2DPoint(EGridSpace::UniformScaled, Grid);
+
+	int32 ProjectedNodeIndex[2] = { 0, 0 };
+	FPoint2D ProjectedPoints[2];
+	double Distance[2] = { HUGE_VALUE, HUGE_VALUE };
+	int32 OtherSegmentIndex = 0;
+	int32 SegmentIndex = 1;
+
+	TFunction<void(int32)> ProjectInnerPoint = [&](int32 OtherSegmentNodeIndex)
+	{
+		double Coordinate;
+		ProjectedPoints[OtherSegmentIndex] = ProjectPointOnSegment(Points[OtherSegmentIndex][OtherSegmentNodeIndex], Points[SegmentIndex][0], Points[SegmentIndex][1], Coordinate, false);
+		bool bProjectedPoint1IsInside = (Coordinate >= -SMALL_NUMBER && Coordinate <= 1 + SMALL_NUMBER);
+		if (bProjectedPoint1IsInside)
+		{
+			Distance[OtherSegmentIndex] = ProjectedPoints[OtherSegmentIndex].Distance(Points[OtherSegmentIndex][OtherSegmentNodeIndex]);
+		}
+
+#ifdef DEBUG_CLOSED_OUSIDE_POINT
+		if (bDisplay)
+		{
+			{
+				F3DDebugSession _(TEXT("Segs"));
+				DisplaySegment(Points[SegmentIndex][0], Points[SegmentIndex][1], 0, EVisuProperty::BlueCurve);
+				DisplayPoint(Points[OtherSegmentIndex][OtherSegmentNodeIndex], EVisuProperty::RedPoint);
+				DisplayPoint(ProjectedPoints[OtherSegmentIndex], EVisuProperty::RedPoint);
+				Wait(false);
+			}
+		}
+#endif
+	};
+
+	TFunction<void()> FindInnerNodeAndProjectIt = [&]()
+	{
+		OtherSegmentIndex = SegmentIndex ? 0 : 1;
+
+		double OrientedSlop = ComputeOrientedSlope(Points[SegmentIndex][0], Points[SegmentIndex][1], Points[OtherSegmentIndex][0]);
+		ProjectedNodeIndex[OtherSegmentIndex] = ((SegmentLoopIndex[0] == 0) == (OrientedSlop < 0)) ? 0 : 1;
+		ProjectInnerPoint(ProjectedNodeIndex[OtherSegmentIndex]);
+	};
+
+	SegmentIndex = 0;
+	FindInnerNodeAndProjectIt();
+	SegmentIndex = 1;
+	FindInnerNodeAndProjectIt();
+
+	if (Distance[0] == HUGE_VALUE && Distance[1] == HUGE_VALUE)
+	{
+		return false;
+	}
+
+	TFunction<void(int32)> MoveNode = [&](int32 MovePointIndex)
+	{
+		FPoint2D MoveDirection = ProjectedPoints[MovePointIndex] - Points[MovePointIndex][ProjectedNodeIndex[MovePointIndex]];
+		MoveDirection.Normalize();
+
+		if (MoveDirection.SquareLength() < 0.5)
+		{
+			MoveDirection = Points[MovePointIndex][ProjectedNodeIndex[MovePointIndex] == 0 ? 1 : 0] - ProjectedPoints[MovePointIndex];
+			MoveDirection.Normalize();
+		}
+
+		MoveDirection *= GeometricTolerance;
+		FPoint2D NewCoordinate = ProjectedPoints[MovePointIndex] + MoveDirection;
+
+		FLoopNode& Node = (FLoopNode&) const_cast<FIsoNode&>(*Nodes[MovePointIndex][ProjectedNodeIndex[MovePointIndex]]);
+		FLoopNode& PreviousNode = Node.GetPreviousNode();
+		FLoopNode& NextNode = Node.GetNextNode();
+
+		if ((PreviousNode.Get2DPoint(EGridSpace::UniformScaled, Grid).SquareDistance(NewCoordinate) < SquareGeometricTolerance2) ||
+			(NextNode.Get2DPoint(EGridSpace::UniformScaled, Grid).SquareDistance(NewCoordinate) < SquareGeometricTolerance2))
+		{
+			RemoveNodeOfLoop(Node);
+			return;
+		}
+
+		const_cast<FIsoNode&>((*Nodes[MovePointIndex][ProjectedNodeIndex[MovePointIndex]])).Set2DPoint(EGridSpace::UniformScaled, Grid, NewCoordinate);
+
+#ifdef DEBUG_CLOSED_OUSIDE_POINT
+		if (bDisplay)
+		{
+			F3DDebugSession _(TEXT("New Segs"));
+			DisplayPoint(NewCoordinate, EVisuProperty::RedPoint);
+			Display(EGridSpace::UniformScaled, Segment1, 0, EVisuProperty::RedCurve);
+			Display(EGridSpace::UniformScaled, Segment0, 0, EVisuProperty::BlueCurve);
+			Wait(false);
+		}
+#endif
+	};
+
+	int32 MovePoint = (Distance[0] < Distance[1]) ? 0 : 1;
+	MoveNode(MovePoint);
+	return true;
+}
+
+bool FIsoTriangulator::RemoveNodeOfLoop(FLoopNode& NodeToRemove)
+{
+	if (NodeToRemove.GetConnectedSegments().Num() != 2)
+	{
+		return false;
+	}
+
+	FLoopNode& PreviousNode = NodeToRemove.GetPreviousNode();
+	FLoopNode& NextNode = NodeToRemove.GetNextNode();
+
+	FIsoSegment* Segment = PreviousNode.GetSegmentConnectedTo(&NodeToRemove);
+	if (Segment == nullptr)
+	{
+		return false;
+	}
+
+	FIsoSegment* SegmentToDelete = NextNode.GetSegmentConnectedTo(&NodeToRemove);
+	if (SegmentToDelete == nullptr)
+	{
+		return false;
+	}
+
+	NextNode.DisconnectSegment(*SegmentToDelete);
+	NextNode.ConnectSegment(*Segment);
+	Segment->SetSecondNode(NextNode);
+
+	if (&NextNode.GetNextNode() == &NextNode.GetPreviousNode())
+	{
+		NextNode.DisconnectSegment(*Segment);
+		PreviousNode.DisconnectSegment(*Segment);
+		LoopSegments.RemoveSingle(Segment);
+		IsoSegmentFactory.DeleteEntity(Segment);
+
+		FIsoSegment* ThirdSegment = PreviousNode.GetSegmentConnectedTo(&NextNode);
+		ensureCADKernel(ThirdSegment);
+
+		NextNode.DisconnectSegment(*ThirdSegment);
+		PreviousNode.DisconnectSegment(*ThirdSegment);
+
+		LoopSegments.RemoveSingle(ThirdSegment);
+		IsoSegmentFactory.DeleteEntity(ThirdSegment);
+
+		NextNode.Delete();
+		PreviousNode.Delete();
+
+		LoopSegmentsIntersectionTool.Remove(Segment);
+		LoopSegmentsIntersectionTool.Remove(ThirdSegment);
+	}
+
+	LoopSegments.RemoveSingle(SegmentToDelete);
+	IsoSegmentFactory.DeleteEntity(SegmentToDelete);
+	NodeToRemove.Delete();
+
+	LoopSegmentsIntersectionTool.Remove(SegmentToDelete);
+	if (!Segment->IsDelete())
+	{
+		LoopSegmentsIntersectionTool.Update(Segment);
+	}
+
+	return true;
+}
+
+//#define DEBUG_REMOVE_PICK_OF_LOOP
+void FIsoTriangulator::RemovePickOfLoop(FIsoSegment& Segment)
+{
+	if (Segment.GetType() != ESegmentType::Loop)
+	{
+		return;
+	}
+
+	FLoopNode* Node0 = &(FLoopNode&)Segment.GetFirstNode();
+	FLoopNode* Node1 = &(FLoopNode&)Segment.GetSecondNode();
+	FLoopNode* PreviousNode = &Node0->GetPreviousNode();
+	FLoopNode* NextNode = &Node1->GetNextNode();
+
+	const FPoint2D* PreviousPoint = &PreviousNode->Get2DPoint(EGridSpace::UniformScaled, Grid);
+	const FPoint2D* Point0 = &Node0->Get2DPoint(EGridSpace::UniformScaled, Grid);
+	const FPoint2D* Point1 = &Node1->Get2DPoint(EGridSpace::UniformScaled, Grid);
+	const FPoint2D* NextPoint = &NextNode->Get2DPoint(EGridSpace::UniformScaled, Grid);
+
+	bool bNodeRemoved = true;
+	bool bPickRemoved = false;
+	while (LoopSegments.Num() >= 3)
+	{
+		if (CheckAndRemovePick(*PreviousPoint, *Point0, *Point1, *Node0))
+		{
+			if (PreviousNode->IsDelete())
+			{
+				Wait();
+			}
+
+			Point0 = PreviousPoint;
+			Node0 = PreviousNode;
+			PreviousNode = &Node0->GetPreviousNode();
+			PreviousPoint = &PreviousNode->Get2DPoint(EGridSpace::UniformScaled, Grid);
+			bPickRemoved = true;
+			continue;
+		}
+
+		if (CheckAndRemovePick(*Point0, *Point1, *NextPoint, *Node1))
+		{
+			if (NextNode->IsDelete())
+			{
+				Wait();
+			}
+
+			Point1 = NextPoint;
+			Node1 = NextNode;
+			NextNode = &Node1->GetNextNode();
+			NextPoint = &NextNode->Get2DPoint(EGridSpace::UniformScaled, Grid);
+			bPickRemoved = true;
+			continue;
+		}
+
+		break;
+	}
+
+#ifdef DEBUG_REMOVE_PICK_OF_LOOP
+	if (bDisplay && bPickRemoved)
+	{
+		DisplayLoops(TEXT("After pick removed"));
+		Wait(false);
+	}
+#endif
+}
+
+void FIsoTriangulator::SwapLoopOrientation(int32 FirstSegmentIndex, int32 LastSegmentIndex)
+{
+	TArray<FIsoSegment*> Segments;
+	Segments.Reserve(LastSegmentIndex - FirstSegmentIndex);
+	for (int32 Index = FirstSegmentIndex; Index < LastSegmentIndex; ++Index)
+	{
+		LoopSegments[Index]->SwapOrientation();
+		Segments.Add(LoopSegments[Index]);
+	}
+
+	for (int32 ReverseIndex = Segments.Num() - 1, Index = FirstSegmentIndex; Index < LastSegmentIndex; ++Index, --ReverseIndex)
+	{
+		LoopSegments[Index] = Segments[ReverseIndex];
+	}
+};
+
+bool FIsoTriangulator::TryToRemoveIntersectionBySwappingSegments(FIsoSegment& IntersectingSegment, FIsoSegment& Segment)
+{
+	if (LoopSegmentsIntersectionTool.DoesIntersect(IntersectingSegment.GetFirstNode(), Segment.GetFirstNode()))
+	{
+		return false;
+	}
+
+	if (LoopSegmentsIntersectionTool.DoesIntersect(IntersectingSegment.GetSecondNode(), Segment.GetSecondNode()))
+	{
+		return false;
+	}
+
+	int32 StartSegmentIndex = LoopSegments.IndexOfByKey(&IntersectingSegment);
+	int32 EndSegmentIndex = LoopSegments.IndexOfByKey(&Segment);
+
+	IntersectingSegment.GetSecondNode().DisconnectSegment(IntersectingSegment);
+	Segment.GetFirstNode().DisconnectSegment(Segment);
+
+	FIsoNode& Node = IntersectingSegment.GetSecondNode();
+	IntersectingSegment.SetSecondNode(Segment.GetFirstNode());
+	Segment.SetFirstNode(Node);
+
+	Segment.GetFirstNode().ConnectSegment(Segment);
+	IntersectingSegment.GetSecondNode().ConnectSegment(IntersectingSegment);
+
+#ifdef DEBUG_BY_SWAPPING_SEGMENTS
+	if (bDisplay)
+	{
+		F3DDebugSession _(TEXT("New Segments"));
+		Display(EGridSpace::UniformScaled, Segment, StartSegmentIndex, EVisuProperty::RedCurve/*, true*/);
+		Display(EGridSpace::UniformScaled, IntersectingSegment, EndSegmentIndex, EVisuProperty::RedCurve/*, true*/);
+		//Wait();
+	}
+#endif
+
+	SwapLoopOrientation(StartSegmentIndex + 1, EndSegmentIndex);
+
+	LoopSegmentsIntersectionTool.Update(&IntersectingSegment);
+	LoopSegmentsIntersectionTool.Update(&Segment);
+
+	return true;
+}
+
+
+} //namespace CADKernel
\ No newline at end of file
diff --git a/Engine/Source/Runtime/Datasmith/CADKernel/Private/CADKernel/Mesh/Meshers/ParametricMesher.cpp b/Engine/Source/Runtime/Datasmith/CADKernel/Private/CADKernel/Mesh/Meshers/ParametricMesher.cpp
--- a/Engine/Source/Runtime/Datasmith/CADKernel/Private/CADKernel/Mesh/Meshers/ParametricMesher.cpp
+++ b/Engine/Source/Runtime/Datasmith/CADKernel/Private/CADKernel/Mesh/Meshers/ParametricMesher.cpp
@@ -150,14 +150,12 @@ void FParametricMesher::MeshEntities()
 	//      Mesh surfaces 
 	// ============================================================================================================
 
-	FMessage::Printf(Log, TEXT("  Mesh Surfaces\n"));
-
 	FTimePoint MeshStartTime = FChrono::Now();
 	MeshSurfaceByFront(QuadTrimmedSurfaceSet);
 	Chronos.GlobalMeshDuration = FChrono::Elapse(MeshStartTime);
 	Chronos.GlobalDuration = FChrono::Elapse(StartTime);
 
-	Chronos.PrintTimeElapse();
+	//Chronos.PrintTimeElapse();
 }
 
 void FParametricMesher::ApplyFaceCriteria(TSharedRef<FTopologicalFace> Face)
@@ -177,7 +175,7 @@ void FParametricMesher::ApplyFaceCriteria(TSharedRef<FTopologicalFace> Face)
 void FParametricMesher::ApplyEdgeCriteria(FTopologicalEdge& Edge)
 {
 	FTopologicalEdge& ActiveEdge = *Edge.GetLinkActiveEdge();
-	ensureCADKernel(!ActiveEdge.IsApplyCriteria());
+	ensureCADKernel(Edge.IsVirtuallyMeshed() || !ActiveEdge.IsApplyCriteria());
 
 	Edge.ComputeCrossingPointCoordinates();
 	Edge.InitDeltaUs();
@@ -208,6 +206,10 @@ void FParametricMesher::ApplyEdgeCriteria(FTopologicalEdge& Edge)
 
 void FParametricMesher::Mesh(TSharedRef<FTopologicalFace> Face)
 {
+#ifdef DEBUG_ONLY_SURFACE_TO_DEBUG
+	bDisplay = (Face->GetId() == FaceToDebug);
+#endif
+
 	ensureCADKernel(!Face->IsDeleted());
 	ensureCADKernel(!Face->IsMeshed());
 
@@ -266,6 +268,8 @@ void FParametricMesher::Mesh(TSharedRef<FTopologicalFace> Face)
 	Chronos.GlobalMeshDuration += Duration;
 #endif
 
+	FChrono::PrintClockElapse(EVerboseLevel::Debug, TEXT("   "), TEXT("Meshing"), Duration);
+
 	if (BoolDisplayDebugMeshStep)
 	{
 		Close3DDebugSession();
@@ -406,7 +410,17 @@ void FParametricMesher::Mesh(FTopologicalEdge& InEdge, FTopologicalFace& Face)
 		FTopologicalEdge& ActiveEdge = *InEdge.GetLinkActiveEntity();
 		if (ActiveEdge.IsMeshed())
 		{
-			return;
+			if (ActiveEdge.GetMesh()->GetNodeCount() > 0)
+			{
+				return;
+			}
+
+			// In some case the 2d curve is a smooth curve and the 3d curve is a line and vice versa
+			// In the particular case where the both case are opposed, we can have the 2d line sampled with 4 points, and the 2d curve sampled with 2 points (because in 3d, the 2d curve is a 3d line)
+			// In this case, the loop is flat i.e. in 2d the meshes of the 2d line and 2d curve are coincident
+			// So the grid is degenerated and the surface is not meshed
+			// to avoid this case, the Edge is virtually meshed i.e. the nodes inside the edge has the is of the mesh of the vertices.
+			InEdge.SetAsVirtuallyMeshed();
 		}
 
 		if (ActiveEdge.IsThinPeak())
@@ -431,12 +445,10 @@ void FParametricMesher::Mesh(FTopologicalEdge& InEdge, FTopologicalFace& Face)
 	ApplyEdgeCriteria(InEdge);
 
 #ifdef DEBUG_MESH_EDGE
+	if(bDisplay)
 	{
-		F3DDebugSession _(FString::Printf(TEXT("EdgePointsOnDomain %d"), ActiveEdge.GetId()));
-		for (const FPoint& Point : EdgeCrossingPoints2D)
-		{
-			DisplayPoint(Point);
-		}
+		F3DDebugSession _(FString::Printf(TEXT("EdgePointsOnDomain %d"), InEdge.GetId()));
+		Display2D(InEdge);
 		Wait();
 	}
 #endif
@@ -495,14 +507,15 @@ void FParametricMesher::Mesh(FTopologicalEdge& InEdge, FTopologicalFace& Face)
 	}
 
 #ifdef DEBUG_MESH_EDGE
+	if (bDisplay)
 	{
 		F3DDebugSession _(FString::Printf(TEXT("Edge %d"), InEdge.GetId()));
-		for (int32 Index = 0; Index < ActiveEdge.GetImposedCuttingPoints().Num(); ++Index)
+		for (int32 Index = 0; Index < InEdge.GetImposedCuttingPoints().Num(); ++Index)
 		{
-			const FImposedCuttingPoint& CuttingPoint = ActiveEdge.GetImposedCuttingPoints()[Index];
+			const FImposedCuttingPoint& CuttingPoint = InEdge.GetImposedCuttingPoints()[Index];
 			ImposedIsoCuttingPoints.Emplace(CuttingPoint.Coordinate, ECoordinateType::ImposedCoordinate, CuttingPoint.OppositNodeIndex, MinDeltaU * AThird);
 			FCurvePoint Point;
-			ActiveEdge.EvaluatePoint(CuttingPoint.Coordinate, 0, Point);
+			InEdge.EvaluatePoint(CuttingPoint.Coordinate, 0, Point);
 			DisplayPoint(Point.Point);
 		}
 	}
@@ -651,8 +664,13 @@ void FParametricMesher::Mesh(FTopologicalEdge& InEdge, FTopologicalFace& Face)
 	}
 #endif
 
-	if (InEdge.IsDegenerated())
+	if (InEdge.IsDegenerated() || InEdge.IsVirtuallyMeshed())
 	{
+		if (ImposedIsoCuttingPoints.Num() == 2)
+		{
+			ImposedIsoCuttingPoints.EmplaceAt(1, (ImposedIsoCuttingPoints[0].Coordinate + ImposedIsoCuttingPoints[1].Coordinate) * 0.5, ECoordinateType::OtherCoordinate);
+		}
+
 		for (FCuttingPoint CuttingPoint : ImposedIsoCuttingPoints)
 		{
 			FinalEdgeCuttingPointCoordinates.Emplace(CuttingPoint.Coordinate, ECoordinateType::OtherCoordinate);
@@ -731,6 +749,7 @@ void FParametricMesher::GenerateEdgeElements(FTopologicalEdge& Edge)
 	}
 
 #ifdef DEBUG_MESH_EDGE
+	if (bDisplay)
 	{
 		F3DDebugSession _(FString::Printf(TEXT("Edge Mesh %d"), Edge.GetId()));
 		TArray<FPoint2D> Mesh2D;
@@ -1053,7 +1072,7 @@ void FParametricMesher::MeshSurfaceByFront(TArray<FCostToFace>& QuadTrimmedSurfa
 		const TSharedPtr<FTopologicalLoop> Loop = Face->GetLoops()[0];
 		for (const FOrientedEdge& OrientedEdge : Loop->GetEdges())
 		{
-			const TSharedPtr<FTopologicalEdge>& Edge = OrientedEdge.Entity;
+			const TSharedPtr<FTopologicalEdge> Edge = OrientedEdge.Entity;
 			Edge->SetMarker1(); // tmp for debug
 			for (FTopologicalEdge* NextEdge : Edge->GetTwinsEntities())
 			{
diff --git a/Engine/Source/Runtime/Datasmith/CADKernel/Private/CADKernel/Mesh/Structure/Grid.cpp b/Engine/Source/Runtime/Datasmith/CADKernel/Private/CADKernel/Mesh/Structure/Grid.cpp
--- a/Engine/Source/Runtime/Datasmith/CADKernel/Private/CADKernel/Mesh/Structure/Grid.cpp
+++ b/Engine/Source/Runtime/Datasmith/CADKernel/Private/CADKernel/Mesh/Structure/Grid.cpp
@@ -19,13 +19,14 @@ FGrid::FGrid(TSharedRef<FTopologicalFace>& InFace, TSharedRef<FModelMesh>& InMes
 	: Face(InFace)
 	, FaceTolerance(InFace->GetIsoTolerances())
 	, Tolerance3D(InFace->GetCarrierSurface()->Get3DTolerance())
-	, MinimumElementSize(Tolerance3D * 10.)
+	, MinimumElementSize(Tolerance3D * 2.)
 	, MeshModel(InMeshModel)
 	, ThinZoneFinder(*this)
 	, CuttingCoordinates(Face->GetCuttingPointCoordinates())
 {
 #ifdef DEBUG_ONLY_SURFACE_TO_DEBUG
 	bDisplay = (Face->GetId() == FaceToDebug);
+	Open3DDebugSession(bDisplay, FString::Printf(TEXT("Grid %d"), Face->GetId()));
 #endif
 }
 
@@ -943,6 +944,7 @@ void SlightlyDisplacedPolyline(TArray<FPoint2D>& D2Points, const FSurfacicBounda
 	D2Points.Last() += Normal;
 }
 
+//#define DEBUG_GET_BOUNDARY_MESH
 bool FGrid::GetMeshOfLoops()
 {
 	int32 ThinZoneNum = 0;
@@ -1038,10 +1040,32 @@ bool FGrid::GetMeshOfLoops()
 			}
 
 			TArray<int32> EdgeVerticesIndex;
-			if (ActiveEdge->IsDegenerated())
+			if (Edge->IsDegenerated())
 			{
 				EdgeVerticesIndex.Init(ActiveEdge->GetStartVertex()->GetLinkActiveEntity()->GetOrCreateMesh(MeshModel)->GetMesh(), CuttingPolyline.Coordinates.Num());
 			}
+			else if (Edge->IsVirtuallyMeshed())
+			{
+				// @see FParametricMesher::Mesh(FTopologicalEdge& InEdge, FTopologicalFace& Face)
+				int32 NodeCount = CuttingPolyline.Coordinates.Num();
+				EdgeVerticesIndex.Reserve(NodeCount);
+				int32 MiddleNodeIndex = NodeCount / 2;
+				int32 Index = 0;
+				{
+					int32 StartVertexMeshId = Edge->GetStartVertex()->GetLinkActiveEntity()->GetOrCreateMesh(MeshModel)->GetMesh();
+					for (; Index < MiddleNodeIndex; ++Index)
+					{
+						EdgeVerticesIndex.Add(StartVertexMeshId);
+					}
+				}
+				{
+					int32 EndVertexMeshId = Edge->GetEndVertex()->GetLinkActiveEntity()->GetOrCreateMesh(MeshModel)->GetMesh();
+					for (; Index < NodeCount; ++Index)
+					{
+						EdgeVerticesIndex.Add(EndVertexMeshId);
+					}
+				}
+			}
 			else
 			{
 				EdgeVerticesIndex = ((TSharedRef<FEdgeMesh>)ActiveEdge->GetOrCreateMesh(MeshModel))->EdgeVerticesIndex;
@@ -1049,7 +1073,7 @@ bool FGrid::GetMeshOfLoops()
 
 #ifdef DEBUG_GET_BOUNDARY_MESH
 			{
-				F3DDebugSession _(FString::Printf(TEXT("Edge cutting points on surface %d"), ActiveEdge->GetId()));
+				F3DDebugSession _(FString::Printf(TEXT("Edge %d cutting points on surface"), ActiveEdge->GetId()));
 				for (const FPoint2D& Point2D : CuttingPolyline.Points2D)
 				{
 					DisplayPoint(Point2D);
@@ -1798,27 +1822,29 @@ void FGrid::DisplayGridPoints(EGridSpace DisplaySpace) const
 		return;
 	}
 	int32 NbNum = 0;
-	Open3DDebugSession(TEXT("FGrid::FindInnerDomainPoints Inside Point"));
-	for (int32 Index = 0; Index < CuttingSize; ++Index)
 	{
-		if (IsInsideFace[Index])
+		F3DDebugSession _(TEXT("FGrid::FindInnerDomainPoints Inside Point"));
+		for (int32 Index = 0; Index < CuttingSize; ++Index)
 		{
-			DisplayPoint(Points2D[DisplaySpace][Index], Index);
-			NbNum++;
+			if (IsInsideFace[Index])
+			{
+				DisplayPoint(Points2D[DisplaySpace][Index], Index);
+				NbNum++;
+			}
 		}
 	}
-	Close3DDebugSession();
 	ensureCADKernel(NbNum == CountOfInnerNodes);
 
-	Open3DDebugSession(TEXT("FGrid::FindInnerDomainPoints Outside Point"));
-	for (int32 Index = 0; Index < CuttingSize; ++Index)
 	{
-		if (!IsInsideFace[Index])
+		F3DDebugSession _(TEXT("FGrid::FindInnerDomainPoints Outside Point"));
+		for (int32 Index = 0; Index < CuttingSize; ++Index)
 		{
-			DisplayPoint(Points2D[DisplaySpace][Index], EVisuProperty::OrangePoint, Index);
+			if (!IsInsideFace[Index])
+			{
+				DisplayPoint(Points2D[DisplaySpace][Index], EVisuProperty::OrangePoint, Index);
+			}
 		}
 	}
-	Close3DDebugSession();
 }
 
 void FGrid::DisplayGridNormal() const
@@ -1828,36 +1854,37 @@ void FGrid::DisplayGridNormal() const
 		return;
 	}
 
-	Open3DDebugSession(TEXT("FGrid::Inner Normal"));
 	double NormalLength = FSystem::Get().GetVisu()->GetParameters()->NormalLength;
-
-	for (int32 Index = 0; Index < CuttingSize; ++Index)
 	{
-		F3DDebugSegment GraphicSegment(Index);
-		FVector Normal = Normals[Index];
-		Normal.Normalize();
-		Normal *= NormalLength;
-		DrawSegment(Points3D[Index], Points3D[Index] + Normal, EVisuProperty::GreenCurve);
-	}
+		F3DDebugSession _(TEXT("FGrid::Inner Normal"));
 
-	Close3DDebugSession();
-
-	Open3DDebugSession(TEXT("FGrid::Loop Normal"));
-	for (int32 LoopIndex = 0; LoopIndex < FaceLoops3D.Num(); ++LoopIndex)
-	{
-		const TArray<FPoint>& LoopPoints = FaceLoops3D[LoopIndex];
-		const TArray<FVector>& LoopNormals = NormalsOfFaceLoops[LoopIndex];
-
-		for (int32 Index = 0; Index < LoopPoints.Num(); ++Index)
+		for (int32 Index = 0; Index < CuttingSize; ++Index)
 		{
 			F3DDebugSegment GraphicSegment(Index);
-			FVector Normal = LoopNormals[Index];
+			FVector Normal = Normals[Index];
 			Normal.Normalize();
 			Normal *= NormalLength;
-			DrawSegment(LoopPoints[Index], LoopPoints[Index] + Normal, EVisuProperty::YellowCurve);
+			DrawSegment(Points3D[Index], Points3D[Index] + Normal, EVisuProperty::GreenCurve);
+		}
+	}
+
+	{
+		F3DDebugSession _(TEXT("FGrid::Loop Normal"));
+		for (int32 LoopIndex = 0; LoopIndex < FaceLoops3D.Num(); ++LoopIndex)
+		{
+			const TArray<FPoint>& LoopPoints = FaceLoops3D[LoopIndex];
+			const TArray<FVector>& LoopNormals = NormalsOfFaceLoops[LoopIndex];
+
+			for (int32 Index = 0; Index < LoopPoints.Num(); ++Index)
+			{
+				F3DDebugSegment GraphicSegment(Index);
+				FVector Normal = LoopNormals[Index];
+				Normal.Normalize();
+				Normal *= NormalLength;
+				DrawSegment(LoopPoints[Index], LoopPoints[Index] + Normal, EVisuProperty::YellowCurve);
+			}
 		}
 	}
-	Close3DDebugSession();
 }
 
 
diff --git a/Engine/Source/Runtime/Datasmith/CADKernel/Private/CADKernel/Mesh/Structure/Grid.h b/Engine/Source/Runtime/Datasmith/CADKernel/Private/CADKernel/Mesh/Structure/Grid.h
--- a/Engine/Source/Runtime/Datasmith/CADKernel/Private/CADKernel/Mesh/Structure/Grid.h
+++ b/Engine/Source/Runtime/Datasmith/CADKernel/Private/CADKernel/Mesh/Structure/Grid.h
@@ -156,7 +156,14 @@ namespace CADKernel
 	public:
 		FGrid(TSharedRef<FTopologicalFace>& InFace, TSharedRef<FModelMesh>& InShellMesh);
 
+#ifndef CADKERNEL_DEV
 		virtual ~FGrid() = default;
+#else
+		virtual ~FGrid()
+		{
+			Close3DDebugSession(bDisplay);
+		}
+#endif
 
 		// ======================================================================================================================================================================================================================
 		// Meshing tools =========================================================================================================================================================================================================
@@ -524,6 +531,15 @@ namespace CADKernel
 			return IndexV * CuttingCount[EIso::IsoU] + IndexU;
 		}
 
+		/**
+		 * @return the Index of the position in the arrays of a point [IndexU, IndexV] of the grid
+		 */
+		void UVIndexFromGlobalIndex(int32 GLobalIndex, int32& OutIndexU, int32& OutIndexV) const
+		{
+			OutIndexU = GLobalIndex % CuttingCount[EIso::IsoU];
+			OutIndexV = GLobalIndex / CuttingCount[EIso::IsoU];
+		}
+
 		/**
 		 * @return the maximum difference of coordinate along the specified axis of two two successive points
 		 */
@@ -550,22 +566,6 @@ namespace CADKernel
 
 		void DisplayGridNormal() const;
 
-		template<typename TPoint>
-		void DisplayPoints(FString Message, const TArray<TPoint>& Points) const
-		{
-			if (!bDisplay)
-			{
-				return;
-			}
-
-			Open3DDebugSession(Message);
-			for (int32 Index = 0; Index < Points.Num(); ++Index)
-			{
-				DisplayPoint(Points[Index], Index);
-			}
-			Close3DDebugSession();
-		}
-
 		template<typename TPoint>
 		void DisplayInnerDomainPoints(FString Message, const TArray<TPoint>& Points) const
 		{
@@ -574,7 +574,7 @@ namespace CADKernel
 				return;
 			}
 
-			Open3DDebugSession(Message);
+			F3DDebugSession _(Message);
 			for (int32 Index = 0; Index < Points.Num(); ++Index)
 			{
 				if (IsInsideFace[Index])
@@ -589,7 +589,6 @@ namespace CADKernel
 					}
 				}
 			}
-			Close3DDebugSession();
 		}
 
 		template<typename TPoint>
@@ -600,7 +599,7 @@ namespace CADKernel
 				return;
 			}
 
-			Open3DDebugSession(*Message);
+			F3DDebugSession _(*Message);
 			int32 LoopIndex = 0;
 			for (const TArray<TPoint>& Loop : Loops)
 			{
@@ -631,7 +630,6 @@ namespace CADKernel
 					Close3DDebugSession();
 				}
 			}
-			Close3DDebugSession();
 		}
 
 		void DisplayFindPointsCloseToLoop(EGridSpace DisplaySpace) const;
diff --git a/Engine/Source/Runtime/Datasmith/CADKernel/Private/CADKernel/Topo/TopologicalLoop.cpp b/Engine/Source/Runtime/Datasmith/CADKernel/Private/CADKernel/Topo/TopologicalLoop.cpp
--- a/Engine/Source/Runtime/Datasmith/CADKernel/Private/CADKernel/Topo/TopologicalLoop.cpp
+++ b/Engine/Source/Runtime/Datasmith/CADKernel/Private/CADKernel/Topo/TopologicalLoop.cpp
@@ -195,12 +195,17 @@ bool FTopologicalLoop::Get2DSamplingWithoutDegeneratedEdges(TArray<FPoint2D>& Lo
  *
  * In this case, the orientation of the next segment is compare to the bounding box
  *
- * The last very difficult case is a sharp case at the previous case i.e.
+ * The last very difficult case is a sharp case i.e.
+ * The slop is closed to 0 or 8, so it could be a pick in self intersecting. 
+ * We try to recompute the slop a the closed point
+ *     ___________                 _______
+ *    |       o---O    			  |       O  
+ *    |     /         			  |     /        
+ *    |    o               => 	  |    o         
+ *    |    |					  |    |
  *
- *     ______
- *    |   ---O
- *    |   |
- * Not yet cover as if there is 4 sharp case in a loop should be very rare and the loop should be a very bad loop...
+ * If the slop is still closed to 0 or 8, the point is "UndefinedOrientation"
+ * This case was found in the shape of char '1'
  */
 void FTopologicalLoop::Orient()
 {
@@ -273,19 +278,58 @@ void FTopologicalLoop::Orient()
 	int32 WrongOrientationCount = 0;
 	int32 GoodOrientationCount = 0;
 	int32 UndefinedOrientationCount = 0;
-	TFunction<void(const int32)> CompareOrientation = [&](const int32 Index)
+	TFunction<void(const int32)> CompareOrientation = [&](int32 Index)
 	{
-		int32 Index1 = Index + 1;
-		if (Index1 == PointCount)
+		int32 NextIndex = Index + 1;
+		if (NextIndex == PointCount)
 		{
-			Index1 = 0;
+			NextIndex = 0;
 		}
-		int32 Index2 = Index == 0 ? PointCount - 1 : Index - 1;
+		int32 PreviousIndex = Index == 0 ? PointCount - 1 : Index - 1;
 
 		// if the slop of the selected segments is not close to the BBox side (closed of 0 or 4), so the angle between the neighboring segments of the local extrema is not closed to 4 and allows to defined the orientation
 		// Pic case: the slop is compute between previous and next segment of the extrema 
-		double Slop = ComputePositiveSlope(LoopSampling[Index], LoopSampling[Index1], LoopSampling[Index2]);
-		if (Slop > 4.2)
+		//     if the slop is closed to 0 or 8, it could be pick in self intersecting. We try to recompute the slop a the closed point
+		//          if the slop is still closed to 0 or 8, the pic is not used
+		double Slop = ComputePositiveSlope(LoopSampling[Index], LoopSampling[NextIndex], LoopSampling[PreviousIndex]);
+
+		if(Slop > 7.9 || Slop < 0.1)
+		{
+			double SquareLengthBefore = LoopSampling[Index].SquareDistance(LoopSampling[PreviousIndex]);
+			double SquareLengthAfter = LoopSampling[Index].SquareDistance(LoopSampling[NextIndex]);
+			if(SquareLengthBefore < SquareLengthAfter)
+			{
+				Index = PreviousIndex;
+				PreviousIndex = Index == 0 ? PointCount - 1 : Index - 1;
+			}
+			else
+			{
+				Index = NextIndex;
+				NextIndex = Index + 1;
+				if (NextIndex == PointCount)
+				{
+					NextIndex = 0;
+				}
+			}
+
+			Slop = ComputePositiveSlope(LoopSampling[Index], LoopSampling[NextIndex], LoopSampling[PreviousIndex]);
+
+#ifdef DEBUG_ORIENT
+			{
+				F3DDebugSession _(*FString::Printf(TEXT("Fix Pic Node %f"), Slop));
+				DisplayPoint(LoopSampling[Index], EVisuProperty::BluePoint, Index);
+				DisplaySegment(LoopSampling[PreviousIndex], LoopSampling[Index], EVisuProperty::GreenCurve);
+				DisplaySegment(LoopSampling[NextIndex], LoopSampling[Index], EVisuProperty::GreenCurve);
+			}
+#endif
+		}
+
+
+		if (Slop > 7.9 || Slop < 0.1)
+		{
+			UndefinedOrientationCount++;
+		}
+		else if (Slop > 4.2)
 		{
 			WrongOrientationCount++;
 		}
@@ -322,7 +366,7 @@ void FTopologicalLoop::Orient()
 				ReferenceSlop = 4;
 			}
 
-			Slop = ComputeUnorientedSlope(LoopSampling[Index], LoopSampling[Index1], ReferenceSlop);
+			Slop = ComputeUnorientedSlope(LoopSampling[Index], LoopSampling[NextIndex], ReferenceSlop);
 			// slop should be closed to [0, 0.2] or [3.8, 4]
 			if (Slop > 3.8)
 			{
@@ -342,7 +386,7 @@ void FTopologicalLoop::Orient()
 				}
 				{
 					F3DDebugSession _(*FString::Printf(TEXT("Next")));
-					DisplaySegment(LoopSampling[Index1], LoopSampling[Index], EVisuProperty::YellowCurve);
+					DisplaySegment(LoopSampling[NextIndex], LoopSampling[Index], EVisuProperty::YellowCurve);
 				}
 				{
 					F3DDebugSession _(*FString::Printf(TEXT("Node %f"), Slop));
@@ -364,11 +408,11 @@ void FTopologicalLoop::Orient()
 			}
 			{
 				F3DDebugSession _(*FString::Printf(TEXT("Next")));
-				DisplaySegment(LoopSampling[Index1], LoopSampling[Index], EVisuProperty::YellowCurve);
+				DisplaySegment(LoopSampling[NextIndex], LoopSampling[Index], EVisuProperty::YellowCurve);
 			}
 			{
 				F3DDebugSession _(*FString::Printf(TEXT("Previous")));
-				DisplaySegment(LoopSampling[Index2], LoopSampling[Index], EVisuProperty::YellowCurve);
+				DisplaySegment(LoopSampling[PreviousIndex], LoopSampling[Index], EVisuProperty::YellowCurve);
 			}
 			{
 				F3DDebugSession _(*FString::Printf(TEXT("Node %f"), Slop));
diff --git a/Engine/Source/Runtime/Datasmith/CADKernel/Public/CADKernel/Core/HaveStates.h b/Engine/Source/Runtime/Datasmith/CADKernel/Public/CADKernel/Core/HaveStates.h
--- a/Engine/Source/Runtime/Datasmith/CADKernel/Public/CADKernel/Core/HaveStates.h
+++ b/Engine/Source/Runtime/Datasmith/CADKernel/Public/CADKernel/Core/HaveStates.h
@@ -25,7 +25,8 @@ enum class EHaveStates : uint16
 	Degenerated         = 0x0010u,  // used for FEdge, FGrid
 								    
 	IsDeleted           = 0x0020u,  // used for all class
-								    
+	IsVirtuallyMeshed   = 0x0040u,  // used for FEdge
+
 	HasMarker1          = 0x1000u,  // used for all class
 	HasMarker2          = 0x2000u,  // used for all class
 	HasMarker3          = 0x4000u,  // used for all class
diff --git a/Engine/Source/Runtime/Datasmith/CADKernel/Public/CADKernel/Geo/Curves/SegmentCurve.h b/Engine/Source/Runtime/Datasmith/CADKernel/Public/CADKernel/Geo/Curves/SegmentCurve.h
--- a/Engine/Source/Runtime/Datasmith/CADKernel/Public/CADKernel/Geo/Curves/SegmentCurve.h
+++ b/Engine/Source/Runtime/Datasmith/CADKernel/Public/CADKernel/Geo/Curves/SegmentCurve.h
@@ -66,6 +66,20 @@ namespace CADKernel
 			return EndPoint;
 		}
 
+		virtual void ExtendTo(const FPoint& DesiredPosition) override
+		{
+			double DistanceToStartPoint = DesiredPosition.SquareDistance(StartPoint);
+			double DistanceToEndPoint = DesiredPosition.SquareDistance(EndPoint);
+			if (DistanceToEndPoint < DistanceToStartPoint)
+			{
+				EndPoint = DesiredPosition;
+			}
+			else
+			{
+				StartPoint = DesiredPosition;
+			}
+		}
+
 	private:
 		template <typename CurvePointType, typename PointType>
 		void Evaluate(double Coordinate, CurvePointType& OutPoint, int32 DerivativeOrder) const
diff --git a/Engine/Source/Runtime/Datasmith/CADKernel/Public/CADKernel/Geo/Sampling/PolylineTools.h b/Engine/Source/Runtime/Datasmith/CADKernel/Public/CADKernel/Geo/Sampling/PolylineTools.h
--- a/Engine/Source/Runtime/Datasmith/CADKernel/Public/CADKernel/Geo/Sampling/PolylineTools.h
+++ b/Engine/Source/Runtime/Datasmith/CADKernel/Public/CADKernel/Geo/Sampling/PolylineTools.h
@@ -33,7 +33,7 @@ namespace CADKernel
 		inline PointType LinearInterpolation(const TArray<PointType>& Array, const int32 Index, const double Coordinate)
 		{
 			ensureCADKernel(Index + 1 < Array.Num());
-			return Array[Index] + (Array[Index + 1] - Array[Index]) * Coordinate;
+			return PointOnSegment(Array[Index], Array[Index + 1], Coordinate);
 		}
 
 		inline double SectionCoordinate(const TArray<double>& Array, const int32 Index, const double Coordinate)
diff --git a/Engine/Source/Runtime/Datasmith/CADKernel/Public/CADKernel/Math/Geometry.h b/Engine/Source/Runtime/Datasmith/CADKernel/Public/CADKernel/Math/Geometry.h
--- a/Engine/Source/Runtime/Datasmith/CADKernel/Public/CADKernel/Math/Geometry.h
+++ b/Engine/Source/Runtime/Datasmith/CADKernel/Public/CADKernel/Math/Geometry.h
@@ -11,538 +11,597 @@
 
 namespace CADKernel
 {
-	enum EPolygonSide : uint8
-	{
-		Side01 = 0,
-		Side12,
-		Side20,
-		Side23,
-		Side30,
-	};
-	
-	/**
-	 * https://en.wikipedia.org/wiki/Circumscribed_circle#Cartesian_coordinates_2
-	 * With A = (0, 0)
-	 */
-	CADKERNEL_API inline FPoint2D ComputeCircumCircleCenter(const FPoint2D& InPoint0, const FPoint2D& InPoint1, const FPoint2D& InPoint2)
+enum EPolygonSide : uint8
+{
+	Side01 = 0,
+	Side12,
+	Side20,
+	Side23,
+	Side30,
+};
+
+/**
+ * https://en.wikipedia.org/wiki/Circumscribed_circle#Cartesian_coordinates_2
+ * With A = (0, 0)
+ */
+CADKERNEL_API inline FPoint2D ComputeCircumCircleCenter(const FPoint2D& InPoint0, const FPoint2D& InPoint1, const FPoint2D& InPoint2)
+{
+	FPoint2D Segment_P0_P1 = InPoint1 - InPoint0;
+	FPoint2D Segment_P0_P2 = InPoint2 - InPoint0;
+
+	// D = 2(BuCv - BvCu)
+	double D = 2 * Segment_P0_P1 ^ Segment_P0_P2;
+	if (FMath::IsNearlyZero(D, SMALL_NUMBER_SQUARE))
 	{
-		FPoint2D Segment_P0_P1 = InPoint1 - InPoint0;
-		FPoint2D Segment_P0_P2 = InPoint2 - InPoint0;
+		return FPoint2D::ZeroPoint;
+	}
 
-		// D = 2(BuCv - BvCu)
-		double D = 2 * Segment_P0_P1 ^ Segment_P0_P2;
-		if (FMath::IsNearlyZero(D, SMALL_NUMBER_SQUARE))
-		{
-			return FPoint2D::ZeroPoint;
-		}
+	// CenterU  = 1/D * (Cv.|B|.|B| - By.|C|.|C|) = 1/D * CBv ^ SquareNorms 
+	// CenterV  = 1/D * (Bu.|B|.|B| - Cu.|C|.|C|) = -1/D * SquareNorms ^ CBu
+	double SquareNormB = Segment_P0_P1.SquareLength();
+	double SquareNormC = Segment_P0_P2.SquareLength();
+	double CenterU = (SquareNormB * Segment_P0_P2.V - SquareNormC * Segment_P0_P1.V) / D;
+	double CenterV = (SquareNormC * Segment_P0_P1.U - SquareNormB * Segment_P0_P2.U) / D;
 
-		// CenterU  = 1/D * (Cv.|B|.|B| - By.|C|.|C|) = 1/D * CBv ^ SquareNorms 
-		// CenterV  = 1/D * (Bu.|B|.|B| - Cu.|C|.|C|) = -1/D * SquareNorms ^ CBu
-		double SquareNormB = Segment_P0_P1.SquareLength();
-		double SquareNormC = Segment_P0_P2.SquareLength();
-		double CenterU = (SquareNormB * Segment_P0_P2.V - SquareNormC * Segment_P0_P1.V) / D;
-		double CenterV = (SquareNormC * Segment_P0_P1.U - SquareNormB * Segment_P0_P2.U) / D;
+	return FPoint2D(CenterU, CenterV) + InPoint0;
+}
 
-		return FPoint2D(CenterU, CenterV) + InPoint0;
-	}
+CADKERNEL_API inline FPoint ComputeCircumCircleCenter(const FPoint& Point0, const FPoint& Point1, const FPoint& Point2)
+{
+	FPoint AxisZ = (Point1 - Point0) ^ (Point2 - Point0);
+	AxisZ.Normalize();
 
-	CADKERNEL_API inline FPoint ComputeCircumCircleCenter(const FPoint& Point0, const FPoint& Point1, const FPoint& Point2)
-	{
-		FPoint AxisZ = (Point1 - Point0) ^ (Point2 - Point0);
-		AxisZ.Normalize();
+	FPoint AxisX = Point1 - Point0;
+	AxisX.Normalize();
+	FPoint AxisY = AxisZ ^ AxisX;
 
-		FPoint AxisX = Point1 - Point0;
-		AxisX.Normalize();
-		FPoint AxisY = AxisZ ^ AxisX;
+	FMatrixH Matrix(Point0, AxisX, AxisY, AxisZ);
+	FMatrixH MatrixInverse = Matrix;
+	MatrixInverse.Inverse();
 
-		FMatrixH Matrix(Point0, AxisX, AxisY, AxisZ);
-		FMatrixH MatrixInverse = Matrix;
-		MatrixInverse.Inverse();
+	FPoint2D D2Point1 = MatrixInverse * Point1;
+	FPoint2D D2Point2 = MatrixInverse * Point2;
 
-		FPoint2D D2Point1 = MatrixInverse * Point1;
-		FPoint2D D2Point2 = MatrixInverse * Point2;
+	double D = 2 * D2Point1 ^ D2Point2;
+	if (FMath::IsNearlyZero(D, SMALL_NUMBER_SQUARE))
+	{
+		return FPoint::ZeroPoint;
+	}
 
-		double D = 2 * D2Point1 ^ D2Point2;
-		if (FMath::IsNearlyZero(D, SMALL_NUMBER_SQUARE))
-		{
-			return FPoint::ZeroPoint;
-		}
+	double SquareNormB = D2Point1.SquareLength();
+	double SquareNormC = D2Point2.SquareLength();
+
+	double CenterU = (SquareNormB * D2Point2.V - SquareNormC * D2Point1.V) / D;
+	double CenterV = (SquareNormC * D2Point1.U - SquareNormB * D2Point2.U) / D;
 
-		double SquareNormB = D2Point1.SquareLength();
-		double SquareNormC = D2Point2.SquareLength();
+	FPoint Center2D(CenterU, CenterV, 0);
+	return Matrix * Center2D;
+}
 
-		double CenterU = (SquareNormB * D2Point2.V - SquareNormC * D2Point1.V) / D;
-		double CenterV = (SquareNormC * D2Point1.U - SquareNormB * D2Point2.U) / D;
 
-		FPoint Center2D(CenterU, CenterV, 0);
-		return Matrix * Center2D;
+CADKERNEL_API inline FPoint2D ComputeCircumCircleCenterAndSquareRadius(const FPoint2D& InPoint0, const FPoint2D& InPoint1, const FPoint2D& InPoint2, double& OutSquareRadius)
+{
+	FPoint2D Segment_P0_P1 = InPoint1 - InPoint0;
+	FPoint2D Segment_P0_P2 = InPoint2 - InPoint0;
+
+	double D = 2 * Segment_P0_P1 ^ Segment_P0_P2;
+	if (FMath::IsNearlyZero(D, SMALL_NUMBER_SQUARE))
+	{
+		OutSquareRadius = 0;
+		return FPoint2D::ZeroPoint;
 	}
 
+	double SquareNormB = Segment_P0_P1.SquareLength();
+	double SquareNormC = Segment_P0_P2.SquareLength();
+	double CenterU = (SquareNormB * Segment_P0_P2.V - SquareNormC * Segment_P0_P1.V) / D;
+	double CenterV = (SquareNormC * Segment_P0_P1.U - SquareNormB * Segment_P0_P2.U) / D;
+
+	FPoint2D Center(CenterU, CenterV);
+	OutSquareRadius = Center.SquareLength();
 
-	CADKERNEL_API inline FPoint2D ComputeCircumCircleCenterAndSquareRadius(const FPoint2D& InPoint0, const FPoint2D& InPoint1, const FPoint2D& InPoint2, double& OutSquareRadius)
+	return Center + InPoint0;
+}
+
+template<class PointType>
+struct CADKERNEL_API TSegment
+{
+	const PointType& Point0;
+	const PointType& Point1;
+
+	TSegment(const PointType& InPoint0, const PointType& InPoint1)
+		: Point0(InPoint0)
+		, Point1(InPoint1)
 	{
-		FPoint2D Segment_P0_P1 = InPoint1 - InPoint0;
-		FPoint2D Segment_P0_P2 = InPoint2 - InPoint0;
+	}
 
-		double D = 2 * Segment_P0_P1 ^ Segment_P0_P2;
-		if (FMath::IsNearlyZero(D, SMALL_NUMBER_SQUARE))
+	constexpr const PointType& operator[](int32 Index) const
+	{
+		ensureCADKernel(Index < 2);
+		switch (Index)
 		{
-			OutSquareRadius = 0;
-			return FPoint2D::ZeroPoint;
+		case 0:
+			return Point0;
+		case 1:
+			return Point1;
+		default:
+			return PointType::ZeroPoint;
 		}
+	}
 
-		double SquareNormB = Segment_P0_P1.SquareLength();
-		double SquareNormC = Segment_P0_P2.SquareLength();
-		double CenterU = (SquareNormB * Segment_P0_P2.V - SquareNormC * Segment_P0_P1.V) / D;
-		double CenterV = (SquareNormC * Segment_P0_P1.U - SquareNormB * Segment_P0_P2.U) / D;
+	double SquareLength() const
+	{
+		return Point0.SquareDistance(Point1);
+	}
 
-		FPoint2D Center(CenterU, CenterV);
-		OutSquareRadius = Center.SquareLength();
+};
 
-		return Center + InPoint0;
+template<class PointType>
+struct CADKERNEL_API TTriangle
+{
+	const PointType& Point0;
+	const PointType& Point1;
+	const PointType& Point2;
+
+	TTriangle(const PointType& InPoint0, const PointType& InPoint1, const PointType& InPoint2)
+		: Point0(InPoint0)
+		, Point1(InPoint1)
+		, Point2(InPoint2)
+	{
 	}
 
-	template<class PointType>
-	struct CADKERNEL_API TSegment
+	constexpr const PointType& operator[](int32 Index) const
 	{
-		const PointType& Point0;
-		const PointType& Point1;
-
-		TSegment(const PointType& InPoint0, const PointType& InPoint1)
-			: Point0(InPoint0)
-			, Point1(InPoint1)
+		ensureCADKernel(Index < 3);
+		switch (Index)
 		{
+		case 0:
+			return Point0;
+		case 1:
+			return Point1;
+		case 2:
+			return Point2;
+		default:
+			return PointType::ZeroPoint;
 		}
+	}
 
-		constexpr const PointType& operator[](int32 Index) const
+	virtual inline PointType ProjectPoint(const PointType& InPoint, FPoint2D& OutCoordinate)
+	{
+		PointType Segment01 = Point1 - Point0;
+		PointType Segment02 = Point2 - Point0;
+		double SquareLength01 = Segment01.SquareLength();
+		double SquareLength02 = Segment02.SquareLength();
+		double Seg01Seg02 = Segment01 * Segment02;
+		double Det = SquareLength01 * SquareLength02 - FMath::Square(Seg01Seg02);
+
+		int32 SideIndex;
+		// If the 3 points are aligned
+		if (FMath::IsNearlyZero(Det))
 		{
-			ensureCADKernel(Index < 2);
-			switch (Index)
+			double MaxLength = SquareLength01;
+			SideIndex = Side01;
+			if (SquareLength02 > MaxLength)
 			{
-			case 0:
-				return Point0;
-			case 1:
-				return Point1;
-			default:
-				return PointType::ZeroPoint;
+				MaxLength = SquareLength02;
+				SideIndex = Side20;
 			}
-		}
-	};
-	
-	template<class PointType>
-	struct CADKERNEL_API TTriangle
-	{
-		const PointType& Point0;
-		const PointType& Point1;
-		const PointType& Point2;
-
-		TTriangle(const PointType& InPoint0, const PointType& InPoint1, const PointType& InPoint2)
-			: Point0(InPoint0)
-			, Point1(InPoint1)
-			, Point2(InPoint2)
-		{
-		}
 
-		constexpr const PointType& operator[](int32 Index) const 
-		{
-			ensureCADKernel(Index < 3);
-			switch (Index)
+			PointType Segment12 = Point2 - Point1;
+			if (Segment12.SquareLength() > MaxLength)
 			{
-				case 0:
-					return Point0;
-				case 1:
-					return Point1;
-				case 2:
-					return Point2;
-				default:
-					return PointType::ZeroPoint;
+				SideIndex = Side12;
 			}
 		}
-
-		virtual inline PointType ProjectPoint(const PointType& InPoint, FPoint2D& OutCoordinate)
+		else
 		{
-			PointType Segment01 = Point1 - Point0;
-			PointType Segment02 = Point2 - Point0;
-			double SquareLength01 = Segment01.SquareLength();
-			double SquareLength02 = Segment02.SquareLength();
-			double Seg01Seg02 = Segment01 * Segment02;
-			double Det = SquareLength01 * SquareLength02 - FMath::Square(Seg01Seg02);
-
-			int32 SideIndex;
-			// If the 3 points are aligned
-			if (FMath::IsNearlyZero(Det))
+			// Resolve
+			PointType Segment1Point = InPoint - Point0;
+			double Segment1PointSegment01 = Segment1Point * Segment01;
+			double Segment1PointSegment02 = Segment1Point * Segment02;
+
+			OutCoordinate.U = ((Segment1PointSegment01 * SquareLength02) - (Segment1PointSegment02 * Seg01Seg02)) / Det;
+			OutCoordinate.V = ((Segment1PointSegment02 * SquareLength01) - (Segment1PointSegment02 * Seg01Seg02)) / Det;
+
+			// tester la solution pour choisir parmi 4 possibilites
+			if (OutCoordinate.U < 0.0)
 			{
-				double MaxLength = SquareLength01;
-				SideIndex = Side01;
-				if (SquareLength02 > MaxLength)
-				{
-					MaxLength = SquareLength02;
-					SideIndex = Side20;
-				}
-
-				PointType Segment12 = Point2 - Point1;
-				if (Segment12.SquareLength() > MaxLength)
-				{
-					SideIndex = Side12;
-				}
+				// the project point is on the segment 02
+				SideIndex = Side20;
 			}
-			else
+			else if (OutCoordinate.V < 0.0)
 			{
-				// Resolve
-				PointType Segment1Point = InPoint - Point0;
-				double Segment1PointSegment01 = Segment1Point * Segment01;
-				double Segment1PointSegment02 = Segment1Point * Segment02;
-
-				OutCoordinate.U = ((Segment1PointSegment01 * SquareLength02) - (Segment1PointSegment02 * Seg01Seg02)) / Det;
-				OutCoordinate.V = ((Segment1PointSegment02 * SquareLength01) - (Segment1PointSegment02 * Seg01Seg02)) / Det;
-
-				// tester la solution pour choisir parmi 4 possibilites
-				if (OutCoordinate.U < 0.0)
-				{
-					// the project point is on the segment 02
-					SideIndex = Side20;
-				}
-				else if (OutCoordinate.V < 0.0)
-				{
-					// the project point is on the segment 01
-					SideIndex = Side01;
-				}
-				else if ((OutCoordinate.U + OutCoordinate.V) > 1.0)
-				{
-					// the project point is on the segment 12
-					SideIndex = Side12;
-				}
-				else {
-					// the project point is inside the Segment
-					Segment01 = Segment01 * OutCoordinate.U;
-					Segment02 = Segment02 * OutCoordinate.V;
-					PointType ProjectedPoint = Segment01 + Segment02;
-					ProjectedPoint = ProjectedPoint + Point0;
-					return ProjectedPoint;
-				}
+				// the project point is on the segment 01
+				SideIndex = Side01;
 			}
-
-			// projects the point on the nearest side
-			PointType ProjectedPoint;
-			double SegmentCoordinate;
-			switch (SideIndex)
+			else if ((OutCoordinate.U + OutCoordinate.V) > 1.0)
 			{
-			case Side01:
-				ProjectedPoint = ProjectPointOnSegment(InPoint, Point0, Point1, SegmentCoordinate);
-				OutCoordinate.U = SegmentCoordinate;
-				OutCoordinate.V = 0.0;
-				break;
-			case Side20:
-				ProjectedPoint = ProjectPointOnSegment(InPoint, Point0, Point2, SegmentCoordinate);
-				OutCoordinate.U = 0.0;
-				OutCoordinate.V = SegmentCoordinate;
-				break;
-			case Side12:
-				ProjectedPoint = ProjectPointOnSegment(InPoint, Point1, Point2, SegmentCoordinate);
-				OutCoordinate.U = 1.0 - SegmentCoordinate;
-				OutCoordinate.V = SegmentCoordinate;
-				break;
+				// the project point is on the segment 12
+				SideIndex = Side12;
+			}
+			else {
+				// the project point is inside the Segment
+				Segment01 = Segment01 * OutCoordinate.U;
+				Segment02 = Segment02 * OutCoordinate.V;
+				PointType ProjectedPoint = Segment01 + Segment02;
+				ProjectedPoint = ProjectedPoint + Point0;
+				return ProjectedPoint;
 			}
-			return ProjectedPoint;
 		}
 
-		virtual PointType CircumCircleCenter() const
+		// projects the point on the nearest side
+		PointType ProjectedPoint;
+		double SegmentCoordinate;
+		switch (SideIndex)
 		{
-			return ComputeCircumCircleCenter(Point0, Point1, Point2);
+		case Side01:
+			ProjectedPoint = ProjectPointOnSegment(InPoint, Point0, Point1, SegmentCoordinate);
+			OutCoordinate.U = SegmentCoordinate;
+			OutCoordinate.V = 0.0;
+			break;
+		case Side20:
+			ProjectedPoint = ProjectPointOnSegment(InPoint, Point0, Point2, SegmentCoordinate);
+			OutCoordinate.U = 0.0;
+			OutCoordinate.V = SegmentCoordinate;
+			break;
+		case Side12:
+			ProjectedPoint = ProjectPointOnSegment(InPoint, Point1, Point2, SegmentCoordinate);
+			OutCoordinate.U = 1.0 - SegmentCoordinate;
+			OutCoordinate.V = SegmentCoordinate;
+			break;
 		}
-	};
+		return ProjectedPoint;
+	}
 
-	struct CADKERNEL_API FTriangle : public TTriangle<FPoint>
+	virtual PointType CircumCircleCenter() const
 	{
-		FTriangle(const FPoint& InPoint0, const FPoint& InPoint1, const FPoint& InPoint2)
-			: TTriangle<FPoint>(InPoint0, InPoint1, InPoint2)
-		{
-		}
+		return ComputeCircumCircleCenter(Point0, Point1, Point2);
+	}
+};
 
-		virtual FPoint ComputeNormal() const
-		{
-			FPoint Normal = (Point1 - Point0) ^ (Point2 - Point0);
-			Normal.Normalize();
-			return Normal;
-		}
-	};
+struct CADKERNEL_API FTriangle : public TTriangle<FPoint>
+{
+	FTriangle(const FPoint& InPoint0, const FPoint& InPoint1, const FPoint& InPoint2)
+		: TTriangle<FPoint>(InPoint0, InPoint1, InPoint2)
+	{
+	}
 
-	struct CADKERNEL_API FTriangle2D : public TTriangle<FPoint2D>
+	virtual FPoint ComputeNormal() const
 	{
-		FTriangle2D(const FPoint2D& InPoint0, const FPoint2D& InPoint1, const FPoint2D& InPoint2)
-			: TTriangle<FPoint2D>(InPoint0, InPoint1, InPoint2)
-		{
-		}
+		FPoint Normal = (Point1 - Point0) ^ (Point2 - Point0);
+		Normal.Normalize();
+		return Normal;
+	}
+};
+
+struct CADKERNEL_API FTriangle2D : public TTriangle<FPoint2D>
+{
+	FTriangle2D(const FPoint2D& InPoint0, const FPoint2D& InPoint1, const FPoint2D& InPoint2)
+		: TTriangle<FPoint2D>(InPoint0, InPoint1, InPoint2)
+	{
+	}
+
+	FPoint2D CircumCircleCenterWithSquareRadius(double& SquareRadius) const
+	{
+		return ComputeCircumCircleCenterAndSquareRadius(this->Point0, this->Point1, this->Point2, SquareRadius);
+	}
+};
 
-		FPoint2D CircumCircleCenterWithSquareRadius(double& SquareRadius) const
-		{
-			return ComputeCircumCircleCenterAndSquareRadius(this->Point0, this->Point1, this->Point2, SquareRadius);
-		}
-	};
 
+template<class PointType>
+struct CADKERNEL_API TQuadrangle
+{
+	const PointType& Point0;
+	const PointType& Point1;
+	const PointType& Point2;
+	const PointType& Point3;
+
+	TQuadrangle(const PointType& InPoint0, const PointType& InPoint1, const PointType& InPoint2, const PointType& InPoint3)
+		: Point0(InPoint0)
+		, Point1(InPoint1)
+		, Point2(InPoint2)
+		, Point3(InPoint3)
+	{
+	}
 
-	template<class PointType>
-	struct CADKERNEL_API TQuadrangle
+	constexpr const PointType& operator[](int32 Index) const
 	{
-		const PointType& Point0;
-		const PointType& Point1;
-		const PointType& Point2;
-		const PointType& Point3;
-
-		TQuadrangle(const PointType& InPoint0, const PointType& InPoint1, const PointType& InPoint2, const PointType& InPoint3)
-			: Point0(InPoint0)
-			, Point1(InPoint1)
-			, Point2(InPoint2)
-			, Point3(InPoint3)
+		ensureCADKernel(Index < 4);
+		switch (Index)
 		{
+		case 0:
+			return Point0;
+		case 1:
+			return Point1;
+		case 2:
+			return Point2;
+		case 3:
+			return Point3;
+		default:
+			return PointType::ZeroPoint;
 		}
+	}
+
+	inline FPoint ComputeNormal(const TQuadrangle<FPoint>& InQuadrangle) const
+	{
+		FPoint Normal = (InQuadrangle[1] - InQuadrangle[0]) ^ (InQuadrangle[2] - InQuadrangle[0]);
+		Normal += (InQuadrangle[2] - InQuadrangle[0]) ^ (InQuadrangle[3] - InQuadrangle[0]);
+		Normal += (InQuadrangle[1] - InQuadrangle[0]) ^ (InQuadrangle[3] - InQuadrangle[0]);
+		Normal += (InQuadrangle[1] - InQuadrangle[3]) ^ (InQuadrangle[2] - InQuadrangle[3]);
+		Normal.Normalize();
+		return Normal;
+	}
+
+	inline PointType  ProjectPoint(const PointType& InPoint, FPoint2D& OutCoordinate)
+	{
+		FPoint2D Coordinate013;
+		TTriangle<PointType> Triangle013(Point0, Point1, Point3);
+		PointType Projection013 = Triangle013.ProjectPoint(InPoint, Coordinate013);
 
-		constexpr const PointType& operator[](int32 Index) const
+		FPoint2D Coordinate231;
+		TTriangle<PointType> Triangle231(Point2, Point3, Point1);
+		PointType Projection231 = Triangle231.ProjectPoint(InPoint, Coordinate231);
+
+		double DistanceTo013 = Projection013.Distance(InPoint);
+		double DistanceTo231 = Projection231.Distance(InPoint);
+		if (DistanceTo013 < DistanceTo231)
 		{
-			ensureCADKernel(Index < 4);
-			switch (Index)
-			{
-			case 0:
-				return Point0;
-			case 1:
-				return Point1;
-			case 2:
-				return Point2;
-			case 3:
-				return Point3;
-			default:
-				return PointType::ZeroPoint;
-			}
+			OutCoordinate = Coordinate013;
+			return Projection013;
 		}
-
-		inline FPoint ComputeNormal(const TQuadrangle<FPoint>& InQuadrangle) const
+		else
 		{
-			FPoint Normal = (InQuadrangle[1] - InQuadrangle[0]) ^ (InQuadrangle[2] - InQuadrangle[0]);
-			Normal += (InQuadrangle[2] - InQuadrangle[0]) ^ (InQuadrangle[3] - InQuadrangle[0]);
-			Normal += (InQuadrangle[1] - InQuadrangle[0]) ^ (InQuadrangle[3] - InQuadrangle[0]);
-			Normal += (InQuadrangle[1] - InQuadrangle[3]) ^ (InQuadrangle[2] - InQuadrangle[3]);
-			Normal.Normalize();
-			return Normal;
+			OutCoordinate = { 1.0 , 1.0 };
+			OutCoordinate -= Coordinate231;
+			return Projection231;
 		}
+	}
 
-		inline PointType  ProjectPoint(const PointType& InPoint, FPoint2D& OutCoordinate)
-		{
-			FPoint2D Coordinate013;
-			TTriangle<PointType> Triangle013(Point0, Point1, Point3);
-			PointType Projection013 = Triangle013.ProjectPoint(InPoint, Coordinate013);
+};
 
-			FPoint2D Coordinate231;
-			TTriangle<PointType> Triangle231(Point2, Point3, Point1);
-			PointType Projection231 = Triangle231.ProjectPoint(InPoint, Coordinate231);
+inline FPoint ProjectPointOnPlane(const FPoint& Point, const FPoint& Origin, const FPoint& InNormal, double& OutDistance)
+{
+	FPoint Normal = InNormal;
+	ensureCADKernel(!FMath::IsNearlyZero(Normal.Length()));
+	Normal.Normalize();
+
+	FPoint OP = Point - Origin;
+	OutDistance = OP * Normal;
+
+	return Point - (Normal * OutDistance);
+}
+
+/**
+ * @return InSegmentA + (InSegmentB - InSegmentA) * InCoordinate
+ */
+template<class PointType>
+inline PointType PointOnSegment(const PointType& InSegmentA, const PointType& InSegmentB, double InCoordinate)
+{
+	PointType Segment = InSegmentB - InSegmentA;
+	return InSegmentA + Segment * InCoordinate;
+}
+
+/**
+ * @return the distance between the point and the segment. If the projection of the point on the segment
+ * is not inside it, return the distance of the point to nearest the segment extremity
+ */
+template<class PointType>
+inline double DistanceOfPointToSegment(const PointType& Point, const PointType& SegmentPoint1, const PointType& SegmentPoint2)
+{
+	double Coordinate;
+	return ProjectPointOnSegment(Point, SegmentPoint1, SegmentPoint2, Coordinate, /*bRestrictCoodinateToInside*/ true).Distance(Point);
+}
+
+/**
+ * @return the distance between the point and the segment. If the projection of the point on the segment
+ * is not inside it, return the distance of the point to nearest the segment extremity
+ */
+template<class PointType>
+inline double SquareDistanceOfPointToSegment(const PointType& Point, const PointType& SegmentPoint1, const PointType& SegmentPoint2)
+{
+	double Coordinate;
+	return ProjectPointOnSegment(Point, SegmentPoint1, SegmentPoint2, Coordinate, /*bRestrictCoodinateToInside*/ true).SquareDistance(Point);
+}
+
+/**
+ * @return the distance between the point and the line i.e. the distance between the point and its projection on the line
+ */
+template<class PointType>
+inline double DistanceOfPointToLine(const PointType& Point, const PointType& LinePoint1, const PointType& LineDirection)
+{
+	double Coordinate;
+	return ProjectPointOnSegment(Point, LinePoint1, LinePoint1 + LineDirection, Coordinate, /*bRestrictCoodinateToInside*/ false).Distance(Point);
+}
+
+CADKERNEL_API double ComputeCurvature(const FPoint& Gradient, const FPoint& Laplacian);
+CADKERNEL_API double ComputeCurvature(const FPoint& normal, const FPoint& Gradient, const FPoint& Laplacian);
+
+/**
+ * @param OutCoordinate: the coordinate of the projected point in the segment AB (coodinate of A = 0 and of B = 1)
+ * @return Projected point
+ */
+template<class PointType>
+inline PointType ProjectPointOnSegment(const PointType& Point, const PointType& InSegmentA, const PointType& InSegmentB, double& OutCoordinate, bool bRestrictCoodinateToInside = true)
+{
+	PointType Segment = InSegmentB - InSegmentA;
 
-			double DistanceTo013 = Projection013.Distance(InPoint);
-			double DistanceTo231 = Projection231.Distance(InPoint);
-			if (DistanceTo013 < DistanceTo231)
+	double SquareLength = Segment * Segment;
+
+	if (SquareLength <= 0.0)
+	{
+		OutCoordinate = 0.0;
+		return InSegmentA;
+	}
+	else
+	{
+		PointType APoint = Point - InSegmentA;
+		OutCoordinate = (APoint * Segment) / SquareLength;
+
+		if (bRestrictCoodinateToInside)
+		{
+			if (OutCoordinate < 0.0)
 			{
-				OutCoordinate = Coordinate013;
-				return Projection013;
+				OutCoordinate = 0.0;
+				return InSegmentA;
 			}
-			else
+
+			if (OutCoordinate > 1.0)
 			{
-				OutCoordinate = { 1.0 , 1.0 };
-				OutCoordinate -= Coordinate231;
-				return Projection231;
+				OutCoordinate = 1.0;
+				return InSegmentB;
 			}
 		}
 
-	};
-
-	inline FPoint ProjectPointOnPlane(const FPoint& Point, const FPoint& Origin, const FPoint& InNormal, double& OutDistance)
-	{
-		FPoint Normal = InNormal;
-		ensureCADKernel(!FMath::IsNearlyZero(Normal.Length()));
-		Normal.Normalize();
-
-		FPoint OP = Point - Origin;
-		OutDistance = OP * Normal;
-
-		return Point - (Normal * OutDistance);
+		PointType ProjectedPoint = Segment * OutCoordinate;
+		ProjectedPoint += InSegmentA;
+		return ProjectedPoint;
 	}
+}
 
-	/**
-	 * @return the distance between the point and the segment. If the projection of the point on the segment
-	 * is not inside it, return the distance of the point to nearest the segment extremity
-	 */
-	template<class PointType>
-	inline double DistanceOfPointToSegment(const PointType& Point, const PointType& SegmentPoint1, const PointType& SegmentPoint2)
-	{
-		double Coordinate;
-		return ProjectPointOnSegment(Point, SegmentPoint1, SegmentPoint2, Coordinate, /*bRestrictCoodinateToInside*/ true).Distance(Point);
-	}
+/**
+ * @return Coordinate of the projected point in the segment AB (coordinate of A = 0 and of B = 1)
+ */
+template<class PointType>
+inline double CoordinateOfProjectedPointOnSegment(const PointType& Point, const PointType& InSegmentA, const PointType& InSegmentB, bool bRestrictCoodinateToInside = true)
+{
+	PointType Segment = InSegmentB - InSegmentA;
 
-	/**
-	 * @return the distance between the point and the segment. If the projection of the point on the segment
-	 * is not inside it, return the distance of the point to nearest the segment extremity
-	 */
-	template<class PointType>
-	inline double SquareDistanceOfPointToSegment(const PointType& Point, const PointType& SegmentPoint1, const PointType& SegmentPoint2)
-	{
-		double Coordinate;
-		return ProjectPointOnSegment(Point, SegmentPoint1, SegmentPoint2, Coordinate, /*bRestrictCoodinateToInside*/ true).SquareDistance(Point);
-	}
+	double SquareLength = Segment * Segment;
 
-	/**
-	 * @return the distance between the point and the line i.e. the distance between the point and its projection on the line 
-	 */
-	template<class PointType>
-	inline double DistanceOfPointToLine(const PointType& Point, const PointType& LinePoint1, const PointType& LineDirection)
+	if (SquareLength <= 0.0)
 	{
-		double Coordinate;
-		return ProjectPointOnSegment(Point, LinePoint1, LinePoint1 + LineDirection, Coordinate, /*bRestrictCoodinateToInside*/ false).Distance(Point);
+		return 0.0;
 	}
-
-	CADKERNEL_API double ComputeCurvature(const FPoint& Gradient, const FPoint& Laplacian);
-	CADKERNEL_API double ComputeCurvature(const FPoint& normal, const FPoint& Gradient, const FPoint& Laplacian);
-
-	/**
-	 * @param OutCoordinate: the coordinate of the projected point in the segment AB (coodinate of A = 0 and of B = 1)
-	 * @return Projected point
-	 */
-	template<class PointType>
-	inline PointType ProjectPointOnSegment(const PointType& Point, const PointType& InSegmentA, const PointType& InSegmentB, double& OutCoordinate, bool bRestrictCoodinateToInside = true)
+	else
 	{
-		PointType Segment = InSegmentB - InSegmentA;
-
-		double SquareLength = Segment * Segment;
+		PointType APoint = Point - InSegmentA;
+		double Coordinate = (APoint * Segment) / SquareLength;
 
-		if (SquareLength <= 0.0)
-		{
-			OutCoordinate = 0.0;
-			return InSegmentA;
-		}
-		else
+		if (bRestrictCoodinateToInside)
 		{
-			PointType APoint = Point - InSegmentA;
-			OutCoordinate = (APoint * Segment) / SquareLength;
-
-			if (bRestrictCoodinateToInside)
+			if (Coordinate < 0.0)
 			{
-				if (OutCoordinate < 0.0)
-				{
-					OutCoordinate = 0.0;
-					return InSegmentA;
-				}
-
-				if (OutCoordinate > 1.0)
-				{
-					OutCoordinate = 1.0;
-					return InSegmentB;
-				}
+				return 0.0;
 			}
 
-			PointType ProjectedPoint = Segment * OutCoordinate;
-			ProjectedPoint += InSegmentA;
-			return ProjectedPoint;
+			if (Coordinate > 1.0)
+			{
+				return 1.0;
+			}
 		}
+
+		return Coordinate;
 	}
+}
 
-	/**
-	 * @return Coordinate of the projected point in the segment AB (coordinate of A = 0 and of B = 1)
-	 */
-	template<class PointType>
-	inline double CoordinateOfProjectedPointOnSegment(const PointType& Point, const PointType& InSegmentA, const PointType& InSegmentB, bool bRestrictCoodinateToInside = true)
-	{
-		PointType Segment = InSegmentB - InSegmentA;
+CADKERNEL_API void FindLoopIntersectionsWithIso(const EIso Iso, const double IsoParameter, const TArray<TArray<FPoint2D>>& Loops, TArray<double>& OutIntersections);
 
-		double SquareLength = Segment * Segment;
+/**
+ * Similar as IntersectSegments2D but do not check intersection if both segment are carried by the same line.
+ * Must be done before (with BBox comparison for example)
+ * This method is 50% faster than IntersectSegments2D even if
+ */
+inline bool FastIntersectSegments2D(const TSegment<FPoint2D>& SegmentAB, const TSegment<FPoint2D>& SegmentCD)
+{
+	constexpr const double Min = -SMALL_NUMBER;
+	constexpr const double Max = 1. + SMALL_NUMBER;
 
-		if (SquareLength <= 0.0)
-		{
-			return 0.0;
-		}
-		else
-		{
-			PointType APoint = Point - InSegmentA;
-			double Coordinate = (APoint * Segment) / SquareLength;
+	FPoint2D AB = SegmentAB[1] - SegmentAB[0];
+	FPoint2D CD = SegmentCD[1] - SegmentCD[0];
+	FPoint2D CA = SegmentAB[0] - SegmentCD[0];
 
-			if (bRestrictCoodinateToInside)
-			{
-				if (Coordinate < 0.0)
-				{
-					return 0.0;
-				}
-
-				if (Coordinate > 1.0)
-				{
-					return 1.0;
-				}
-			}
+	double ParallelCoef = CD ^ AB;
+	double ABIntersectionCoordinate = (CA ^ CD) / ParallelCoef;
+	double CDIntersectionCoordinate = (CA ^ AB) / ParallelCoef;
 
-			return Coordinate;
+	if (FMath::IsNearlyZero(ParallelCoef))
+	{
+		ParallelCoef = CA ^ AB;
+		if (!FMath::IsNearlyZero(ParallelCoef))
+		{
+			return false;
 		}
+		return true;
 	}
 
-	CADKERNEL_API void FindLoopIntersectionsWithIso(const EIso Iso, const double IsoParameter, const TArray<TArray<FPoint2D>>& Loops, TArray<double>& OutIntersections);
+	return (ABIntersectionCoordinate <= Max && ABIntersectionCoordinate >= Min && CDIntersectionCoordinate <= Max && CDIntersectionCoordinate >= Min);
+}
 
-	/**
-	 * Similar as IntersectSegments2D but do not check intersection if both segment are carried by the same line.
-	 * Must be done before (with BBox comparison for example)
-	 * This method is 50% faster than IntersectSegments2D even if
-	 */
-	inline bool FastIntersectSegments2D(const TSegment<FPoint2D>& SegmentAB, const TSegment<FPoint2D>& SegmentCD)
-	{
-		constexpr const double Min = -SMALL_NUMBER;
-		constexpr const double Max = 1. + SMALL_NUMBER;
+/**
+ * The segments must intersect because no check is done
+ */
+inline FPoint2D FindIntersectionOfSegments2D(const TSegment<FPoint2D>& SegmentAB, const TSegment<FPoint2D>& SegmentCD, double& OutABIntersectionCoordinate)
+{
+	constexpr const double Min = -SMALL_NUMBER;
+	constexpr const double Max = 1. + SMALL_NUMBER;
 
-		FPoint2D AB = SegmentAB[1] - SegmentAB[0];
-		FPoint2D CD = SegmentCD[1] - SegmentCD[0];
-		FPoint2D CA = SegmentAB[0] - SegmentCD[0];
+	FPoint2D AB = SegmentAB[1] - SegmentAB[0];
+	FPoint2D DC = SegmentCD[0] - SegmentCD[1];
+	FPoint2D AC = SegmentCD[0] - SegmentAB[0];
 
-		double ParallelCoef = CD ^ AB;
-		double ABIntersectionCoordinate = (CA ^ CD) / ParallelCoef;
-		double CDIntersectionCoordinate = (CA ^ AB) / ParallelCoef;
+	double ParallelCoef = DC ^ AB;
+	if (FMath::IsNearlyZero(ParallelCoef))
+	{
+		double SquareAB = AB * AB;
+		double CCoordinate = (AB * AC) / SquareAB;
 
-		if (FMath::IsNearlyZero(ParallelCoef))
+		FPoint2D AD = SegmentCD[1] - SegmentAB[0];
+		double DCoordinate = (AB * AD) / SquareAB;
+
+		if (CCoordinate >= -KINDA_SMALL_NUMBER && CCoordinate <= 1 + KINDA_SMALL_NUMBER)
 		{
-			ParallelCoef = CA ^ AB;
-			if (!FMath::IsNearlyZero(ParallelCoef))
+			if (DCoordinate >= -KINDA_SMALL_NUMBER && DCoordinate <= 1 + KINDA_SMALL_NUMBER)
 			{
-				return false;
+				OutABIntersectionCoordinate = (DCoordinate + CCoordinate) * 0.5;
+				return SegmentCD[0].Middle(SegmentCD[1]);
 			}
-			return true;
-		}
 
-		return (ABIntersectionCoordinate <= Max && ABIntersectionCoordinate >= Min && CDIntersectionCoordinate <= Max && CDIntersectionCoordinate >= Min);
+			CCoordinate = FMath::Clamp(CCoordinate, 0., 1.);
+			OutABIntersectionCoordinate = CCoordinate;
+			return SegmentCD[0];
+		}
+		else if (DCoordinate >= -KINDA_SMALL_NUMBER && DCoordinate <= 1 + KINDA_SMALL_NUMBER)
+		{
+			DCoordinate = FMath::Clamp(DCoordinate, 0., 1.);
+			OutABIntersectionCoordinate = DCoordinate;
+			return SegmentCD[1];
+		}
+		else
+		{
+			OutABIntersectionCoordinate = 0.5;
+			return SegmentAB[0].Middle(SegmentAB[1]);
+		}
 	}
 
-	/**
-	 * The segments must intersect because no check is done
-	 */
-	inline FPoint2D FindIntersectionOfSegments2D(const TSegment<FPoint2D>& SegmentAB, const TSegment<FPoint2D>& SegmentCD)
-	{
-		constexpr const double Min = -SMALL_NUMBER;
-		constexpr const double Max = 1. + SMALL_NUMBER;
+	OutABIntersectionCoordinate = (DC ^ AC) / ParallelCoef;
+	OutABIntersectionCoordinate = FMath::Clamp(OutABIntersectionCoordinate, 0., 1.);
 
-		FPoint2D AB = SegmentAB[1] - SegmentAB[0];
-		FPoint2D CD = SegmentCD[1] - SegmentCD[0];
-		FPoint2D CA = SegmentAB[0] - SegmentCD[0];
+	return SegmentAB[0] + AB * OutABIntersectionCoordinate;
+}
 
-		double ParallelCoef = CD ^ AB;
-		double ABIntersectionCoordinate = (CA ^ CD) / ParallelCoef;
-
-		return SegmentAB[0] + AB * ABIntersectionCoordinate;
-	}
+/**
+ * The segments must intersect because no check is done
+ */
+inline FPoint2D FindIntersectionOfSegments2D(const TSegment<FPoint2D>& SegmentAB, const TSegment<FPoint2D>& SegmentCD)
+{
+	double ABIntersectionCoordinate;
+	return FindIntersectionOfSegments2D(SegmentAB, SegmentCD, ABIntersectionCoordinate);
+}
 
-	/**
-	 * Similar as FastIntersectSegments2D but check intersection if both segment are carried by the same line.
-	 * This method is 50% slower than FastIntersectSegments2D even if the segments tested are never carried by the same line
-	 */
-	CADKERNEL_API bool IntersectSegments2D(const TSegment<FPoint2D>&SegmentAB, const TSegment<FPoint2D>&SegmentCD);
+/**
+ * Similar as FastIntersectSegments2D but check intersection if both segment are carried by the same line.
+ * This method is 50% slower than FastIntersectSegments2D even if the segments tested are never carried by the same line
+ */
+CADKERNEL_API bool IntersectSegments2D(const TSegment<FPoint2D>& SegmentAB, const TSegment<FPoint2D>& SegmentCD);
 
-	inline double ComputeCosinus(FVector Vector, FVector OtherVector)
-	{
-		Vector.Normalize();
-		OtherVector.Normalize();
+inline double ComputeCosinus(FVector Vector, FVector OtherVector)
+{
+	Vector.Normalize();
+	OtherVector.Normalize();
 
-		double Cosinus = Vector.X * OtherVector.X + Vector.Y * OtherVector.Y + Vector.Z * OtherVector.Z;
+	double Cosinus = Vector.X * OtherVector.X + Vector.Y * OtherVector.Y + Vector.Z * OtherVector.Z;
 
-		return FMath::Max(-1.0, FMath::Min(Cosinus, 1.0));
-	}
+	return FMath::Max(-1.0, FMath::Min(Cosinus, 1.0));
+}
 
 
 } // namespace CADKernel
diff --git a/Engine/Source/Runtime/Datasmith/CADKernel/Public/CADKernel/Mesh/Structure/EdgeMesh.h b/Engine/Source/Runtime/Datasmith/CADKernel/Public/CADKernel/Mesh/Structure/EdgeMesh.h
--- a/Engine/Source/Runtime/Datasmith/CADKernel/Public/CADKernel/Mesh/Structure/EdgeMesh.h
+++ b/Engine/Source/Runtime/Datasmith/CADKernel/Public/CADKernel/Mesh/Structure/EdgeMesh.h
@@ -34,6 +34,11 @@ namespace CADKernel
 			}
 			EdgeVerticesIndex.Add(EndVertexNodeIndex);
 		}
+
+		int32 GetNodeCount() const
+		{
+			return EdgeVerticesIndex.Num();
+		}
 	};
 }
 
diff --git a/Engine/Source/Runtime/Datasmith/CADKernel/Public/CADKernel/Topo/TopologicalEdge.h b/Engine/Source/Runtime/Datasmith/CADKernel/Public/CADKernel/Topo/TopologicalEdge.h
--- a/Engine/Source/Runtime/Datasmith/CADKernel/Public/CADKernel/Topo/TopologicalEdge.h
+++ b/Engine/Source/Runtime/Datasmith/CADKernel/Public/CADKernel/Topo/TopologicalEdge.h
@@ -673,6 +673,21 @@ namespace CADKernel
 			States &= ~EHaveStates::ThinPeak;
 		}
 
+		bool IsVirtuallyMeshed() const
+		{
+			return ((States & EHaveStates::IsVirtuallyMeshed) == EHaveStates::IsVirtuallyMeshed);
+		}
+
+		virtual void SetAsVirtuallyMeshed() const
+		{
+			States |= EHaveStates::IsVirtuallyMeshed;
+		}
+
+		virtual void ResetVirtuallyMeshed() const
+		{
+			States &= ~EHaveStates::IsVirtuallyMeshed;
+		}
+		
 		/**
 		 * @return true if the edge is adjacent to only one surface (its carrier surface)
 		 */
diff --git a/Engine/Source/Runtime/Datasmith/CADKernel/Public/CADKernel/UI/Display.h b/Engine/Source/Runtime/Datasmith/CADKernel/Public/CADKernel/UI/Display.h
--- a/Engine/Source/Runtime/Datasmith/CADKernel/Public/CADKernel/UI/Display.h
+++ b/Engine/Source/Runtime/Datasmith/CADKernel/Public/CADKernel/UI/Display.h
@@ -40,9 +40,17 @@ namespace CADKernel
 #ifdef CADKERNEL_DEV
 	void Wait(bool bMakeWait = true);
 	void Open3DDebugSession(FString name, const TArray<FIdent>& idList = TArray<FIdent>());
+	inline void Open3DDebugSession(bool bIsDisplayed, FString name, const TArray<FIdent>& idList = TArray<FIdent>()) 
+	{
+		if (bIsDisplayed)
+		{
+			Open3DDebugSession(name, idList);
+		}
+	};
 	void Close3DDebugSession(bool bIsDisplayed = true);
 #else
 	inline void Wait(bool bMakeWait = true) {};
+	inline void Open3DDebugSession(bool bIsDisplayed, FString name, const TArray<FIdent>& idList = TArray<FIdent>()) {};
 	inline void Open3DDebugSession(FString name, const TArray<FIdent>& idList = TArray<FIdent>()) {};
 	inline void Close3DDebugSession(bool bIsDisplayed = true) {};
 #endif
@@ -136,6 +144,22 @@ namespace CADKernel
 		DrawPoint(Point, Property);
 	}
 
+	template<typename TPoint>
+	void DisplayPoints(FString Message, const TArray<TPoint>& Points, EVisuProperty Property = EVisuProperty::BluePoint, bool bDisplay = true)
+	{
+		if (!bDisplay)
+		{
+			return;
+		}
+
+		Open3DDebugSession(Message);
+		for (int32 Index = 0; Index < Points.Num(); ++Index)
+		{
+			DisplayPoint(Points[Index], Property);
+		}
+		Close3DDebugSession();
+	}
+
 	CADKERNEL_API void DisplayProductTree(const FEntity& RootId);
 	CADKERNEL_API void DisplayProductTree(const FModel& Model);
 	CADKERNEL_API void DisplayProductTree(const FBody& Body);
