Updater CADKernel

diff --git a/Engine/Source/Runtime/Datasmith/CADKernel/Private/CADKernel/Geo/Sampler/SamplerAbstract.h b/Engine/Source/Runtime/Datasmith/CADKernel/Private/CADKernel/Geo/Sampler/SamplerAbstract.h
--- a/Engine/Source/Runtime/Datasmith/CADKernel/Private/CADKernel/Geo/Sampler/SamplerAbstract.h
+++ b/Engine/Source/Runtime/Datasmith/CADKernel/Private/CADKernel/Geo/Sampler/SamplerAbstract.h
@@ -108,7 +108,7 @@ namespace CADKernel
 				NotDerivableCoordinates.Insert(Boundary.Min, 0);
 				NotDerivableCoordinates.Add(Boundary.Max);
 
-				ComplementatyPointCount = NotDerivableCoordinates.Num() < 5 ? 5 : 1;
+				ComplementatyPointCount = NotDerivableCoordinates.Num() < 10 ? 10 : 1;
 
 				NextCoordinates.Empty(NotDerivableCoordinates.Num() * (ComplementatyPointCount + 1));
 				NextCoordinates.Add(Boundary.Min);
@@ -152,7 +152,8 @@ namespace CADKernel
 		{
 
 			// < 100000 : check to avoid that the process loops endlessly
-			while (int32 CandidatePointsCount = CandidatePoints.Size() && Sampling.Coordinates.Num() < 100000)
+			int32 CandidatePointsCount = CandidatePoints.Size();
+			while (CandidatePointsCount && Sampling.Coordinates.Num() < 100000)
 			{
 				NextCoordinates.Empty();
 
@@ -194,7 +195,7 @@ namespace CADKernel
 							AddIntermediateCoordinates(Sampling.Coordinates[StartSamplingSegmentIndex], Sampling.Coordinates[EndSamplingSegmentIndex], 1);
 							StartSamplingSegmentIndex = EndSamplingSegmentIndex;
 							break;
-					}
+						}
 
 						int32 NeededPointNum = CheckSamplingError(FirstCandidateIndex, LastCandidateIndex);
 						if (NeededPointNum == 0 || (NeededPointNum == 1 && ((LastCandidateIndex + 1) == CandidatePointsCount || (((LastCandidateIndex + 1) < CandidatePointsCount) && TmpPolylineCoordinates[LastCandidateIndex + 1] > Sampling.Coordinates[EndSamplingSegmentIndex]))))
@@ -235,6 +236,8 @@ namespace CADKernel
 			CandidatePoints.SwapCoordinates(NextCoordinates);
 			EvaluatesNewCandidatePoints();
 
+			CandidatePointsCount = CandidatePoints.Size();
+
 #ifdef DEBUG_CURVE_SAMPLING
 			DisplaySampling(CurveIndex == CurveToDisplay, 0);
 #endif
diff --git a/Engine/Source/Runtime/Datasmith/CADKernel/Private/CADKernel/Geo/Sampler/SamplerOnParam.h b/Engine/Source/Runtime/Datasmith/CADKernel/Private/CADKernel/Geo/Sampler/SamplerOnParam.h
--- a/Engine/Source/Runtime/Datasmith/CADKernel/Private/CADKernel/Geo/Sampler/SamplerOnParam.h
+++ b/Engine/Source/Runtime/Datasmith/CADKernel/Private/CADKernel/Geo/Sampler/SamplerOnParam.h
@@ -14,8 +14,6 @@
 #include "CADKernel/Math/Point.h"
 #include "CADKernel/Utils/ArrayUtils.h"
 
-
-//#define DEBUG_CURVE_SAMPLING
 namespace CADKernel
 {
 
