Fix Horde Issue 124362: Compile warnings

diff --git a/Engine/Source/Runtime/Datasmith/CADKernel/Private/CADKernel/Mesh/Meshers/IsoTriangulator/IntersectionSegmentTool.cpp b/Engine/Source/Runtime/Datasmith/CADKernel/Private/CADKernel/Mesh/Meshers/IsoTriangulator/IntersectionSegmentTool.cpp
--- a/Engine/Source/Runtime/Datasmith/CADKernel/Private/CADKernel/Mesh/Meshers/IsoTriangulator/IntersectionSegmentTool.cpp
+++ b/Engine/Source/Runtime/Datasmith/CADKernel/Private/CADKernel/Mesh/Meshers/IsoTriangulator/IntersectionSegmentTool.cpp
@@ -31,6 +31,11 @@ bool FIntersectionSegmentTool::FindIntersections(const FIsoNode& StartNode, cons
 
 	for (const FSegment4IntersectionTools& Segment : Segments)
 	{
+		if (!Segment.IsoSegment)
+		{
+			continue;
+		}
+
 		if (&Segment.IsoSegment->GetFirstNode() == &StartNode || &Segment.IsoSegment->GetSecondNode() == &StartNode)
 		{
 			continue;
@@ -60,6 +65,11 @@ const FIsoSegment* FIntersectionSegmentTool::DoesIntersect(const FIsoNode& Start
 
 	for (const FSegment4IntersectionTools& Segment : Segments)
 	{
+		if (!Segment.IsoSegment)
+		{
+			continue;
+		}
+
 		if (Segment.IsoSegment != nullptr && Segment.IsoSegment->IsDelete())
 		{
 			continue;
@@ -99,6 +109,11 @@ bool FIntersectionSegmentTool::DoesIntersect(const FIsoNode& StartNode, const FP
 
 	for (const FSegment4IntersectionTools& Segment : Segments)
 	{
+		if (!Segment.IsoSegment)
+		{
+			continue;
+		}
+
 		if (bSegmentsAreSorted && DMax < Segment.AxisMin)
 		{
 			break;
diff --git a/Engine/Source/Runtime/Datasmith/CADKernel/Private/CADKernel/Mesh/Meshers/IsoTriangulatorRemoveIntersection.cpp b/Engine/Source/Runtime/Datasmith/CADKernel/Private/CADKernel/Mesh/Meshers/IsoTriangulatorRemoveIntersection.cpp
--- a/Engine/Source/Runtime/Datasmith/CADKernel/Private/CADKernel/Mesh/Meshers/IsoTriangulatorRemoveIntersection.cpp
+++ b/Engine/Source/Runtime/Datasmith/CADKernel/Private/CADKernel/Mesh/Meshers/IsoTriangulatorRemoveIntersection.cpp
@@ -199,8 +199,11 @@ void FIsoTriangulator::FindLoopIntersections(const TArray<FLoopNode*>& NodesOfLo
 	GetSegmentToNodeMethod GetFirst = bForward ? GetFirstNode : GetSecondNode;
 	GetSegmentToNodeMethod GetSecond = bForward ? GetSecondNode : GetFirstNode;
 
-	FLoopNode *const* StartNodePtr = NodesOfLoop.FindByPredicate([](const FLoopNode* Node) { return !Node->IsDelete(); });
-	ensureCADKernel(*StartNodePtr);
+	FLoopNode *const*const StartNodePtr = NodesOfLoop.FindByPredicate([](const FLoopNode* Node) { return !Node->IsDelete(); });
+	if(!*StartNodePtr)
+	{
+		return;
+	}
 
 	FLoopNode* StartNode = *StartNodePtr;
 
@@ -214,7 +217,10 @@ void FIsoTriangulator::FindLoopIntersections(const TArray<FLoopNode*>& NodesOfLo
 	TFunction<void()> FindSegmentIntersection = [&]()
 	{
 		FIsoSegment* Segment = Node->GetSegmentConnectedTo(NextNode);
-		ensureCADKernel(Segment);
+		if(Segment)
+		{
+			return;
+		}
 
 		if (LoopSegmentsIntersectionTool.FindIntersections(*Node, *NextNode, IntersectedSegments))
 		{
@@ -2016,7 +2022,10 @@ bool FIsoTriangulator::RemoveNodeOfLoop(FLoopNode& NodeToRemove)
 		IsoSegmentFactory.DeleteEntity(Segment);
 
 		FIsoSegment* ThirdSegment = PreviousNode.GetSegmentConnectedTo(&NextNode);
-		ensureCADKernel(ThirdSegment);
+		if(!ThirdSegment)
+		{
+			return false;
+		}
 
 		NextNode.DisconnectSegment(*ThirdSegment);
 		PreviousNode.DisconnectSegment(*ThirdSegment);
