Fixed Static Analysis Warning
Assumed this code should be checking for NULL ptr and returning.

diff --git a/Engine/Source/Runtime/Datasmith/CADKernel/Private/CADKernel/Mesh/Meshers/IsoTriangulatorRemoveIntersection.cpp b/Engine/Source/Runtime/Datasmith/CADKernel/Private/CADKernel/Mesh/Meshers/IsoTriangulatorRemoveIntersection.cpp
--- a/Engine/Source/Runtime/Datasmith/CADKernel/Private/CADKernel/Mesh/Meshers/IsoTriangulatorRemoveIntersection.cpp
+++ b/Engine/Source/Runtime/Datasmith/CADKernel/Private/CADKernel/Mesh/Meshers/IsoTriangulatorRemoveIntersection.cpp
@@ -217,7 +217,7 @@ void FIsoTriangulator::FindLoopIntersections(const TArray<FLoopNode*>& NodesOfLo
 	TFunction<void()> FindSegmentIntersection = [&]()
 	{
 		FIsoSegment* Segment = Node->GetSegmentConnectedTo(NextNode);
-		if(Segment)
+		if(Segment==nullptr)
 		{
 			return;
 		}
