Fix Horde Static Analysis Win64 (PVS-Studio)

diff --git a/Engine/Source/Runtime/Datasmith/CADKernel/Private/CADKernel/Mesh/Meshers/IsoTriangulatorRemoveIntersection.cpp b/Engine/Source/Runtime/Datasmith/CADKernel/Private/CADKernel/Mesh/Meshers/IsoTriangulatorRemoveIntersection.cpp
--- a/Engine/Source/Runtime/Datasmith/CADKernel/Private/CADKernel/Mesh/Meshers/IsoTriangulatorRemoveIntersection.cpp
+++ b/Engine/Source/Runtime/Datasmith/CADKernel/Private/CADKernel/Mesh/Meshers/IsoTriangulatorRemoveIntersection.cpp
@@ -199,8 +199,8 @@ void FIsoTriangulator::FindLoopIntersections(const TArray<FLoopNode*>& NodesOfLo
 	GetSegmentToNodeMethod GetFirst = bForward ? GetFirstNode : GetSecondNode;
 	GetSegmentToNodeMethod GetSecond = bForward ? GetSecondNode : GetFirstNode;
 
-	FLoopNode *const*const StartNodePtr = NodesOfLoop.FindByPredicate([](const FLoopNode* Node) { return !Node->IsDelete(); });
-	if(!*StartNodePtr)
+	FLoopNode* const* const StartNodePtr = NodesOfLoop.FindByPredicate([](const FLoopNode* Node) { return !Node->IsDelete(); });
+	if (*StartNodePtr == nullptr)
 	{
 		return;
 	}
@@ -217,7 +217,7 @@ void FIsoTriangulator::FindLoopIntersections(const TArray<FLoopNode*>& NodesOfLo
 	TFunction<void()> FindSegmentIntersection = [&]()
 	{
 		FIsoSegment* Segment = Node->GetSegmentConnectedTo(NextNode);
-		if(Segment==nullptr)
+		if (Segment == nullptr)
 		{
 			return;
 		}
@@ -227,6 +227,11 @@ void FIsoTriangulator::FindLoopIntersections(const TArray<FLoopNode*>& NodesOfLo
 			for (const FIsoSegment* IntersectedSegment : IntersectedSegments)
 			{
 				const FLoopNode* IntersectedSegmentFirstNode = GetFirst(IntersectedSegment);
+				if (IntersectedSegmentFirstNode == nullptr)
+				{
+					continue;
+				}
+
 				int32 IntersectionIndex = 0;
 				FLoopNode* TmpNode = StartNode;
 				while (TmpNode != IntersectedSegmentFirstNode)
@@ -265,13 +270,16 @@ void FIsoTriangulator::FindLoopIntersections(const TArray<FLoopNode*>& NodesOfLo
 					Display(EGridSpace::UniformScaled, *Node, *NextNode, 0, EVisuProperty::BluePoint);
 					Display(EGridSpace::UniformScaled, *IntersectedSegment, 0, EVisuProperty::BluePoint);
 					DisplayPoint(IntersectionPoint, EVisuProperty::RedPoint);
-				}
-#endif
 			}
+#endif
 		}
+	}
 
-		LoopSegmentsIntersectionTool.AddSegment(*Segment);
-	};
+		if (Segment)
+		{
+			LoopSegmentsIntersectionTool.AddSegment(*Segment);
+		}
+};
 
 	LoopSegmentsIntersectionTool.Empty(NodesOfLoop.Num());
 
@@ -290,11 +298,24 @@ void FIsoTriangulator::FindLoopIntersections(const TArray<FLoopNode*>& NodesOfLo
 	}
 #endif
 
-	LoopSegmentsIntersectionTool.AddSegment(*StartNode->GetSegmentConnectedTo(NextNode));
+	FIsoSegment* StartToEndSegment = StartNode->GetSegmentConnectedTo(NextNode);
+	if (StartToEndSegment)
+	{
+		LoopSegmentsIntersectionTool.AddSegment(*StartToEndSegment);
+	}
 
 	for (Node = NextNode; Node != StartNode; Node = NextNode, ++SegmentIndex)
 	{
+		if (Node == nullptr || Node->IsDelete())
+		{
+			return;
+		}
+
 		NextNode = GetNext(Node);
+		if (NextNode == nullptr || NextNode->IsDelete())
+		{
+			return;
+		}
 		FindSegmentIntersection();
 	}
 
@@ -410,13 +431,13 @@ bool FIsoTriangulator::RemoveLoopIntersections(const TArray<FLoopNode*>& NodesOf
 	}
 
 	return true;
-}
+	}
 
 void FIsoTriangulator::RemoveLoopPicks(TArray<FLoopNode*>& NodesOfLoop, TArray<TPair<double, double>>& Intersections)
 {
 	for (FLoopNode* Node : NodesOfLoop)
 	{
-		if (Node->IsDelete())
+		if (Node == nullptr || Node->IsDelete())
 		{
 			continue;
 		}
@@ -458,7 +479,7 @@ bool FIsoTriangulator::RemovePickToOutside(const TArray<FLoopNode*>& NodesOfLoop
 		return Intersection;
 	};
 
-	FLoopNode* TmpNode = GetNodeAt(NodesOfLoop, NextIndex(NodeCount, (int32) Intersection.Value));
+	FLoopNode* TmpNode = GetNodeAt(NodesOfLoop, NextIndex(NodeCount, (int32)Intersection.Value));
 	FLoopNode* StartNode = GetPrevious(TmpNode);
 	FPoint2D FirstIntersection = IntersectingPoint(Intersection.Value, StartNode, TmpNode);
 
@@ -489,7 +510,7 @@ bool FIsoTriangulator::RemovePickToOutside(const TArray<FLoopNode*>& NodesOfLoop
 	MoveDirection *= GeometricTolerance;
 	MiddlePoint += MoveDirection;
 
-	if(!EndNode->IsDelete())
+	if (!EndNode->IsDelete())
 	{
 		EndNode->Set2DPoint(EGridSpace::UniformScaled, Grid, MiddlePoint);
 	}
@@ -539,7 +560,7 @@ void FIsoTriangulator::MoveIntersectingSectionBehindOppositeSection(IsoTriangula
 		{
 			DisplayPoint(Node->Get2DPoint(EGridSpace::UniformScaled, Grid), EVisuProperty::PurplePoint);
 		}
-	}
+}
 	{
 		F3DDebugSession _(TEXT("OppositeSection"));
 		for (FLoopNode* Node = FirstNodeOppositeSection; Node != GetNext(LastNodeOppositeSection); Node = GetNext(Node))
@@ -724,7 +745,7 @@ bool FIsoTriangulator::RemoveIntersectionsOfSubLoop(const TArray<FLoopNode*>& No
 					{
 						LoopSegmentsIntersectionTool.AddSegment(*Node->GetSegmentConnectedTo(GetNext(Node)));
 					}
-				}
+					}
 
 				for (FLoopNode* Node = OppositeSection.Key; Node != OppositeSection.Value; Node = GetNext(Node))
 				{
@@ -753,9 +774,9 @@ bool FIsoTriangulator::RemoveIntersectionsOfSubLoop(const TArray<FLoopNode*>& No
 				DisplayLoops(TEXT("RemoveLoopIntersections after remove pick"), false, true);
 				Wait();
 #endif
-			}
+				}
 			--Index;
-		}
+			}
 		else
 		{
 			if (!RemoveUniqueIntersection(NodesOfLoop, LoopIntersections[IntersectionIndex], bForward))
@@ -763,10 +784,10 @@ bool FIsoTriangulator::RemoveIntersectionsOfSubLoop(const TArray<FLoopNode*>& No
 				return false;
 			}
 		}
-	}
+		}
 
 	return true;
-}
+	}
 
 bool FIsoTriangulator::IsSubLoopBiggerThanMainLoop(const TArray<FLoopNode*>& NodesOfLoop, const TPair<double, double>& Intersection, bool bForward)
 {
@@ -828,7 +849,7 @@ bool FIsoTriangulator::RemoveUniqueIntersection(const TArray<FLoopNode*>& NodesO
 	// TODO
 	/*
 	if(bIntersectionKeyIsExtremity && bIntersectionValueIsExtremity)
-	{   //       _   _ 
+	{   //       _   _
 		// Cas:   \./
 		//       _/ \_
 
@@ -890,7 +911,7 @@ bool FIsoTriangulator::SpreadCoincidentNodes(const TArray<FLoopNode*>& NodesOfLo
 	// 	   Use case:
 	//	      FilesToProcess.append([1, 0, 0, 1, r"D:/Data/Cad Files/SolidWorks/p014 - Unreal Sport Bike/Headset front left brake.SLDPRT"])
 	//        selection = [1967]
-	
+
 	//using namespace IsoTriangulatorImpl;
 	//GetNextNodeMethod GetNext = bForward ? GetNextNodeImpl : GetPreviousNodeImpl;
 
@@ -968,7 +989,7 @@ void FIsoTriangulator::MoveNode(FLoopNode& NodeToMove, FPoint2D& ProjectedPoint)
 		{
 			F3DDebugSession _(TEXT("Point To Move"));
 			DisplayPoint(PointToMove, EVisuProperty::YellowPoint);
-		}
+}
 		{
 			F3DDebugSession _(TEXT("Projected Point"));
 			DisplayPoint(ProjectedPoint, EVisuProperty::GreenPoint);
@@ -1140,7 +1161,7 @@ void FIsoTriangulator::TryToSwapSegmentsOrRemoveLoop(const TArray<FLoopNode*>& N
 			Display(EGridSpace::UniformScaled, *Segment1_Node0, 0, EVisuProperty::BluePoint);
 		}
 #endif
-	}
+}
 
 	if (!bIsFixed)
 	{
@@ -1252,7 +1273,7 @@ bool FIsoTriangulator::FindLoopIntersectionAndFixIt()
 		RemoveLoopIntersections(LoopNodesFromStartNode, Intersections, bIsAnOuterLoop);
 
 #ifdef DEBUG_LOOP_INTERSECTION_AND_FIX_IT		
- 		DisplayLoop(EGridSpace::UniformScaled, TEXT("LoopIntersections: remove self intersection"), LoopNodesFromStartNode, true, EVisuProperty::YellowPoint);
+		DisplayLoop(EGridSpace::UniformScaled, TEXT("LoopIntersections: remove self intersection"), LoopNodesFromStartNode, true, EVisuProperty::YellowPoint);
 		Wait(bDisplay);
 #endif
 
@@ -1271,7 +1292,7 @@ bool FIsoTriangulator::FindLoopIntersectionAndFixIt()
 		FixLoopOrientation(LoopNodesFromStartNode);
 
 		bIsAnOuterLoop = false;
-	}
+		}
 
 	if (!CheckMainLoopConsistency())
 	{
@@ -1315,7 +1336,7 @@ bool FIsoTriangulator::FindLoopIntersectionAndFixIt()
 #endif
 
 	return true;
-}
+	}
 
 void FIsoTriangulator::FixLoopOrientation(const TArray<FLoopNode*>& NodesOfLoop)
 {
@@ -1338,7 +1359,7 @@ void FIsoTriangulator::FixLoopOrientation(const TArray<FLoopNode*>& NodesOfLoop)
 			DisplayLoops(TEXT("After orientation"), false, true, true);
 			Wait();
 #endif
-		}
+}
 	}
 }
 
@@ -1427,7 +1448,7 @@ void FIsoTriangulator::RemoveIntersectionByMovingOutsideSegmentNodeInside(const
 	{
 		RemoveNodeOfLoop(Node);
 		return;
-	}
+}
 
 	Node.Set2DPoint(EGridSpace::UniformScaled, Grid, NewCoordinate);
 
@@ -1480,17 +1501,17 @@ bool FIsoTriangulator::TryToRemoveIntersectionOfTwoConsecutiveIntersectingSegmen
 
 	double IntersectingSegmentSlop = ComputeOrientedSlope(IntersectingSegment2D.Point0, IntersectingSegment2D.Point1, 0);
 	double SegmentSlop = ComputeUnorientedSlope(Segment2D.Point1, Segment2D.Point0, IntersectingSegmentSlop);
-	if(SegmentSlop > 2)
+	if (SegmentSlop > 2)
 	{
 		SegmentSlop = 4 - SegmentSlop;
 	}
 
 	// if the segment and IntersectingSegment are parallel, segment are moved inside
-	if(SegmentSlop < 0.01)
+	if (SegmentSlop < 0.01)
 	{
 		double StartPointSquareDistance = SquareDistanceOfPointToSegment(Segment2D.Point0, IntersectingSegment2D.Point0, IntersectingSegment2D.Point1);
 		double EndPointSquareDistance = SquareDistanceOfPointToSegment(Segment2D.Point1, IntersectingSegment2D.Point0, IntersectingSegment2D.Point1);
-		if(StartPointSquareDistance < SquareGeometricTolerance && EndPointSquareDistance < SquareGeometricTolerance)
+		if (StartPointSquareDistance < SquareGeometricTolerance && EndPointSquareDistance < SquareGeometricTolerance)
 		{
 			OffsetSegment(Segment, Segment2D, IntersectingSegment2D);
 			return true;
@@ -1509,7 +1530,7 @@ bool FIsoTriangulator::TryToRemoveIntersectionOfTwoConsecutiveIntersectingSegmen
 		}
 		else if (FMath::IsNearlyEqual(Coordinate, 1))
 		{
-			OffsetNode((FLoopNode&) Segment.GetSecondNode(), IntersectingSegment2D);
+			OffsetNode((FLoopNode&)Segment.GetSecondNode(), IntersectingSegment2D);
 			return true;
 		}
 	}
@@ -1526,7 +1547,7 @@ bool FIsoTriangulator::TryToRemoveIntersectionOfTwoConsecutiveIntersectingSegmen
 		Node = (FLoopNode*)&Segment.GetFirstNode();
 		PreviousNode = (FLoopNode*)&Segment.GetSecondNode();
 		NextNode = (FLoopNode*)&Node->GetPreviousNode();
-	}
+}
 
 #ifdef DEBUG_TWO_CONSECUTIVE_INTERSECTING
 	if (bDisplay)
@@ -1605,7 +1626,7 @@ bool FIsoTriangulator::TryToRemoveIntersectionOfTwoConsecutiveIntersectingSegmen
 	}
 
 	return true;
-}
+	}
 
 //#define DEBUG_FIND_LOOP_INTERSECTION_AND_FIX_IT
 void FIsoTriangulator::FixIntersectionBetweenLoops()
@@ -1625,7 +1646,7 @@ void FIsoTriangulator::FixIntersectionBetweenLoops()
 	int32 Index = 0;
 	for (; Index < LoopSegments.Num(); ++Index)
 	{
-		if(((FLoopNode&) LoopSegments[Index]->GetFirstNode()).GetLoopIndex() != 0)
+		if (((FLoopNode&)LoopSegments[Index]->GetFirstNode()).GetLoopIndex() != 0)
 		{
 			break;
 		}
@@ -1666,7 +1687,7 @@ void FIsoTriangulator::FixIntersectionBetweenLoops()
 			}
 #endif
 
-  			uint32 IntersectionHash = GetTypeHash(*IntersectingSegment, Segment);
+			uint32 IntersectionHash = GetTypeHash(*IntersectingSegment, Segment);
 			bool bNotProceed = !IntersectionAlreadyProceed.Find(IntersectionHash);
 			IntersectionAlreadyProceed.Add(IntersectionHash);
 
@@ -1704,7 +1725,7 @@ void FIsoTriangulator::FixIntersectionBetweenLoops()
 			else if (bIsSameLoop)
 			{
 				bIsFixed = TryToRemoveIntersectionBySwappingSegments(const_cast<FIsoSegment&>(*IntersectingSegment), Segment);
-				if(!bIsFixed)
+				if (!bIsFixed)
 				{
 					ensureCADKernel(false);
 				}
@@ -1723,7 +1744,7 @@ void FIsoTriangulator::FixIntersectionBetweenLoops()
 				Wait(false);
 			}
 #endif
-		}
+			}
 		else
 		{
 			LoopSegmentsIntersectionTool.AddSegment(Segment);
@@ -1794,7 +1815,7 @@ bool FIsoTriangulator::TryToRemoveSelfIntersectionByMovingTheClosedOusidePoint(c
 			}
 		}
 #endif
-	};
+};
 
 	SegmentIndex = 0;
 	OtherSegmentIndex = 1;
@@ -1923,7 +1944,7 @@ bool FIsoTriangulator::TryToRemoveIntersectionByMovingTheClosedOusidePoint(const
 			}
 		}
 #endif
-	};
+};
 
 	TFunction<void()> FindInnerNodeAndProjectIt = [&]()
 	{
@@ -2022,7 +2043,7 @@ bool FIsoTriangulator::RemoveNodeOfLoop(FLoopNode& NodeToRemove)
 		IsoSegmentFactory.DeleteEntity(Segment);
 
 		FIsoSegment* ThirdSegment = PreviousNode.GetSegmentConnectedTo(&NextNode);
-		if(!ThirdSegment)
+		if (!ThirdSegment)
 		{
 			return false;
 		}
