Datasmith UE5 merge fix: Made FDatasmithUniqueNameProviderBase copyable and movable

diff --git a/Engine/Plugins/Enterprise/DatasmithImporter/Source/DatasmithImporter/Private/DatasmithBlueprintLibrary.cpp b/Engine/Plugins/Enterprise/DatasmithImporter/Source/DatasmithImporter/Private/DatasmithBlueprintLibrary.cpp
--- a/Engine/Plugins/Enterprise/DatasmithImporter/Source/DatasmithImporter/Private/DatasmithBlueprintLibrary.cpp
+++ b/Engine/Plugins/Enterprise/DatasmithImporter/Source/DatasmithImporter/Private/DatasmithBlueprintLibrary.cpp
@@ -389,8 +389,7 @@ namespace DatasmithSceneElementUtil
 							}
 
 							// Copy AssetsContext - to assign static meshes to actors
-							//ImportContext.AssetsContext = AssetsImportContext.AssetsContext;
-							checkf(false, TEXT("Above line commented out to allow compiles during staging testing.  Needs to be fixed by Datasmith team"));
+							ImportContext.AssetsContext = AssetsImportContext.AssetsContext;
 						}
 
 						// CreateSceneAsset
diff --git a/Engine/Source/Runtime/Datasmith/DatasmithCore/Private/DatasmithUtils.cpp b/Engine/Source/Runtime/Datasmith/DatasmithCore/Private/DatasmithUtils.cpp
--- a/Engine/Source/Runtime/Datasmith/DatasmithCore/Private/DatasmithUtils.cpp
+++ b/Engine/Source/Runtime/Datasmith/DatasmithCore/Private/DatasmithUtils.cpp
@@ -10,6 +10,7 @@
 #include "DatasmithVariantElements.h"
 #include "IDatasmithSceneElements.h"
 
+#include "Algo/Sort.h"
 #include "HAL/FileManager.h"
 #include "HAL/PlatformFileManager.h"
 #include "Math/UnrealMath.h"
@@ -1508,6 +1509,54 @@ void FDatasmithSceneUtils::CleanUpScene(TSharedRef<IDatasmithScene> Scene, bool
 	}
 }
 
+FDatasmithUniqueNameProviderBase::FDatasmithUniqueNameProviderBase(const FDatasmithUniqueNameProviderBase& Other)
+{
+	FScopeLock LockOther(&Other.CriticalSection);
+
+	FrequentlyUsedNames = Other.FrequentlyUsedNames;
+}
+
+FDatasmithUniqueNameProviderBase::FDatasmithUniqueNameProviderBase(FDatasmithUniqueNameProviderBase&& Other)
+{
+	FScopeLock LockOther(&Other.CriticalSection);
+
+	FrequentlyUsedNames = MoveTemp(Other.FrequentlyUsedNames);
+}
+
+FDatasmithUniqueNameProviderBase* FDatasmithUniqueNameProviderBase::operator=(const FDatasmithUniqueNameProviderBase& Other)
+{
+	// Guarantee lock order
+	TArray< FCriticalSection*, TInlineAllocator<2> > OrderedCriticalSections;
+	OrderedCriticalSections.Add( &CriticalSection );
+	OrderedCriticalSections.Add( &Other.CriticalSection );
+
+	Algo::Sort( OrderedCriticalSections );
+
+	FScopeLock LockFirst( OrderedCriticalSections[0] );
+	FScopeLock LockSecond( OrderedCriticalSections[1] );
+
+	FrequentlyUsedNames = Other.FrequentlyUsedNames;
+
+	return this;
+}
+
+FDatasmithUniqueNameProviderBase* FDatasmithUniqueNameProviderBase::operator=(FDatasmithUniqueNameProviderBase&& Other)
+{
+	// Guarantee lock order
+	TArray< FCriticalSection*, TInlineAllocator<2> > OrderedCriticalSections;
+	OrderedCriticalSections.Add( &CriticalSection );
+	OrderedCriticalSections.Add( &Other.CriticalSection );
+
+	Algo::Sort( OrderedCriticalSections );
+
+	FScopeLock LockFirst( OrderedCriticalSections[0] );
+	FScopeLock LockSecond( OrderedCriticalSections[1] );
+
+	FrequentlyUsedNames = MoveTemp(Other.FrequentlyUsedNames);
+
+	return this;
+}
+
 FString FDatasmithUniqueNameProviderBase::GenerateUniqueName(const FString& InBaseName, int32 CharBudget)
 {
 	FScopeLock Lock( &CriticalSection );
@@ -1563,6 +1612,11 @@ FString FDatasmithUniqueNameProviderBase::GenerateUniqueName(const FString& InBa
 	return {};
 }
 
+void FDatasmithUniqueNameProviderBase::Clear()
+{
+	FScopeLock Lock( &CriticalSection );
+	FrequentlyUsedNames.Empty();
+}
 
 FTransform FDatasmithUtils::ConvertTransform(EModelCoordSystem SourceCoordSystem, const FTransform& LocalTransform)
 {
diff --git a/Engine/Source/Runtime/Datasmith/DatasmithCore/Public/DatasmithUtils.h b/Engine/Source/Runtime/Datasmith/DatasmithCore/Public/DatasmithUtils.h
--- a/Engine/Source/Runtime/Datasmith/DatasmithCore/Public/DatasmithUtils.h
+++ b/Engine/Source/Runtime/Datasmith/DatasmithCore/Public/DatasmithUtils.h
@@ -235,8 +235,15 @@ public:
 class DATASMITHCORE_API FDatasmithUniqueNameProviderBase
 {
 public:
+	FDatasmithUniqueNameProviderBase() = default;
+	FDatasmithUniqueNameProviderBase(const FDatasmithUniqueNameProviderBase& Other);
+	FDatasmithUniqueNameProviderBase(FDatasmithUniqueNameProviderBase&& Other);
+
 	virtual ~FDatasmithUniqueNameProviderBase() = default;
 
+	FDatasmithUniqueNameProviderBase* operator=(const FDatasmithUniqueNameProviderBase& Other);
+	FDatasmithUniqueNameProviderBase* operator=(FDatasmithUniqueNameProviderBase&& Other);
+
 	/**
 	 * Generates a unique name
 	 * @param BaseName Name that will be suffixed with an index to be unique
@@ -260,7 +267,7 @@ public:
 	/**
 	 * Flushes all known names
 	 */
-	virtual void Clear() { FrequentlyUsedNames.Empty(); }
+	virtual void Clear();
 
 protected:
 
@@ -274,7 +281,7 @@ protected:
 
 private:
 	TMap<FString, int32> FrequentlyUsedNames;
-	FCriticalSection CriticalSection;
+	mutable FCriticalSection CriticalSection;
 };
 
 /**
