CADKernel Update

diff --git a/Engine/Source/Runtime/Datasmith/CADKernel/Private/CADKernel/Math/SlopeUtils.h b/Engine/Source/Runtime/Datasmith/CADKernel/Private/CADKernel/Math/SlopeUtils.h
--- a/Engine/Source/Runtime/Datasmith/CADKernel/Private/CADKernel/Math/SlopeUtils.h
+++ b/Engine/Source/Runtime/Datasmith/CADKernel/Private/CADKernel/Math/SlopeUtils.h
@@ -4,10 +4,35 @@
 
 #include "CADKernel/Core/Types.h"
 #include "CADKernel/Math/Point.h"
+#include "Algo/AllOf.h"
 
 namespace CADKernel
 {
 
+	typedef TFunction<double(const FPoint2D&, const FPoint2D&, double)> SlopMethod;
+
+
+	/**
+	 * Transform a positive slope into an oriented slope [-4, 4] i.e. an equivalent angle between [-Pi, Pi]
+	 * @return a slop between [-4, 4]
+	 */
+	inline double TransformIntoOrientedSlope(double Slope)
+	{
+		if (Slope > 4.) Slope -= 8;
+		if (Slope < -4.) Slope += 8;
+		return Slope;
+	}
+
+	/**
+	 * Transform a slope into a positive slope [0, *] i.e. an equivalent angle between [0, 2.Pi]
+	 * @return a slop between [0, 8]
+	 */
+	inline double TransformIntoPositiveSlope(double Slope)
+	{
+		if (Slope < 0) Slope += 8;
+		return Slope;
+	}
+
 	/**
 	 * Fast angle approximation by a "slope" :
 	 * This method is used to compute an approximation of the angle between the input segment defined by two points and [0, u) axis.
@@ -103,7 +128,6 @@ namespace CADKernel
 		return Delta;
 	}
 
-
 	/**
 	 * Compute the oriented slope of a segment according to a reference slope
 	 * This method is used to compute an approximation of the angle between two segments in 2D.
@@ -126,19 +150,17 @@ namespace CADKernel
 		double ReferenceSlope = ComputeSlope(StartPoint, EndPoint1);
 		double Slope = ComputeSlope(StartPoint, EndPoint2);
 		Slope -= ReferenceSlope;
-		if (Slope < 0) Slope += 8;
-		return Slope;
+		return TransformIntoPositiveSlope(Slope);
 	}
 
-	/**
-	 * Transform a positive slope into an oriented slope [-4, 4] i.e. an equivalent angle between [-Pi, Pi]
-	 * @return a slop between [-4, 4]
-	 */
-	inline double TransformIntoOrientedSlope(double Slope)
+	inline double ClockwiseSlop(const FPoint2D& StartPoint, const FPoint2D& EndPoint, double ReferenceSlope)
 	{
-		if (Slope > 4.) Slope -= 8;
-		if (Slope < -4.) Slope += 8;
-		return Slope;
+		return 8. - ComputePositiveSlope(StartPoint, EndPoint, ReferenceSlope);
+	}
+
+	inline double CounterClockwiseSlop(const FPoint2D& StartPoint, const FPoint2D& EndPoint, double ReferenceSlope)
+	{
+		return ComputePositiveSlope(StartPoint, EndPoint, ReferenceSlope);
 	}
 
 	/**
@@ -189,7 +211,7 @@ namespace CADKernel
 	 * Return true if the segment BP is inside the sector defined the half-lines [BA) and [BC) in the counterclockwise.
 	 * Return false if ABP angle or PBC angle is too flat (smaller than FlatAngle)
 	 */
-	inline bool IsPointPBeInsideSectorABC(const FPoint2D& PointA, const FPoint2D& PointB, const FPoint2D& PointC, const FPoint2D& PointP, const double FlatAngle)
+	inline bool IsPointPInsideSectorABC(const FPoint2D& PointA, const FPoint2D& PointB, const FPoint2D& PointC, const FPoint2D& PointP, const double FlatAngle)
 	 {
 		 double SlopWithNextBoundary = ComputeSlope(PointB, PointC);
 		 double BoundaryDeltaSlope = ComputePositiveSlope(PointB, PointA, SlopWithNextBoundary);
@@ -201,5 +223,93 @@ namespace CADKernel
 		 return true;
 	 }
 
+	/**
+	 *                         P1
+	 *          inside        /
+	 *                       /   inside
+	 *                      /
+	 *    A -------------- B --------------- C
+	 *                      \
+	 *           Outside     \  Outside
+	 *                        \
+	 *                         P2
+	 *
+	 * Return true if all of the segment BPi is inside the sector defined the half-lines [BA) and [BC) in the counterclockwise.
+	 */
+	inline bool ArePointsInsideSectorABC(const FPoint2D& PointA, const FPoint2D& PointB, const FPoint2D& PointC, const TArray<const FPoint2D*>& Points, const double FlatAngle = -SMALL_NUMBER)
+	{
+		double SlopWithNextBoundary = ComputeSlope(PointB, PointC);
+		double BoundaryDeltaSlope = ComputePositiveSlope(PointB, PointA, SlopWithNextBoundary);
+
+		return Algo::AllOf(Points, [&](const FPoint2D* PointP) {
+			double DeltaU = PointB.U - PointP->U;
+			double DeltaV = PointB.V - PointP->V;
+
+			if (FMath::Abs(DeltaU) < SMALL_NUMBER && FMath::Abs(DeltaV) < SMALL_NUMBER)
+			{
+				return true;
+			}
+
+			double SegmentSlope = ComputePositiveSlope(PointB, *PointP, SlopWithNextBoundary);
+			if (SegmentSlope < FlatAngle || SegmentSlope + FlatAngle > BoundaryDeltaSlope)
+			{
+				return false;
+			}
+			return true;
+			});
+	}
+
+	inline FPoint2D SlopeToVector(const double Slop)
+	{
+		int32 SlopStep = (int32)(Slop);
+		FPoint2D Vector;
+		switch (SlopStep)
+		{
+		case 0:
+			// Delta = DeltaV / DeltaU;
+			Vector[0] = 1.;
+			Vector[1] = Slop;
+			break;
+		case 1:
+			// 2 - DeltaU / DeltaV;
+			Vector[0] = 2. - Slop;
+			Vector[1] = 1.;
+			break;
+		case 2:
+			// 2 - DeltaU / DeltaV;
+			Vector[0] = 2. - Slop;
+			Vector[1] = 1.;
+			break;
+		case 3:
+			// 4 + DeltaV / DeltaU;
+			Vector[0] = -1.;
+			Vector[1] = 4. - Slop;
+			break;
+		case 4:
+			// 4 + DeltaV / DeltaU;
+			Vector[0] = -1.;
+			Vector[1] = 4. - Slop;
+			break;
+		case 5:
+			// 6 - DeltaU / DeltaV;
+			Vector[0] = Slop - 6.;
+			Vector[1] = -1.;
+			break;
+		case 6:
+			// 6 - DeltaU / DeltaV // deltaU/deltaV <0
+			Vector[0] = Slop - 6.;
+			Vector[1] = -1;
+			break;
+		case 7:
+			// 8 + DeltaV / DeltaU; // deltaU/deltaV <0
+			Vector[0] = 1.;
+			Vector[1] = Slop - 8.;
+			break;
+		default:
+			break;
+		}
+
+		return Vector;
+	}
 
 } // namespace CADKernel	
diff --git a/Engine/Source/Runtime/Datasmith/CADKernel/Private/CADKernel/Mesh/Meshers/IsoTriangulator.cpp b/Engine/Source/Runtime/Datasmith/CADKernel/Private/CADKernel/Mesh/Meshers/IsoTriangulator.cpp
--- a/Engine/Source/Runtime/Datasmith/CADKernel/Private/CADKernel/Mesh/Meshers/IsoTriangulator.cpp
+++ b/Engine/Source/Runtime/Datasmith/CADKernel/Private/CADKernel/Mesh/Meshers/IsoTriangulator.cpp
@@ -12,6 +12,7 @@
 #include "CADKernel/Mesh/Structure/Grid.h"
 #include "CADKernel/Mesh/Structure/EdgeMesh.h"
 #include "CADKernel/Mesh/Structure/FaceMesh.h"
+#include "CADKernel/Mesh/Structure/LoopCleaner.h"
 #include "CADKernel/Topo/TopologicalEdge.h"
 #include "CADKernel/Topo/TopologicalFace.h"
 #include "CADKernel/Topo/TopologicalFace.h"
@@ -23,7 +24,7 @@ namespace CADKernel
 
 const double FIsoTriangulator::GeometricToMeshingToleranceFactor = 10.;
 
-FIsoTriangulator::FIsoTriangulator(FGrid& InGrid, TSharedRef<FFaceMesh> EntityMesh)
+FIsoTriangulator::FIsoTriangulator(FGrid& InGrid, TSharedRef<FFaceMesh> EntityMesh, FMesherReport& InMesherReport)
 	: Grid(InGrid)
 	, Mesh(EntityMesh)
 	, LoopSegmentsIntersectionTool(InGrid)
@@ -35,6 +36,7 @@ FIsoTriangulator::FIsoTriangulator(FGrid& InGrid, TSharedRef<FFaceMesh> EntityMe
 	, SquareGeometricTolerance2(2. * SquareGeometricTolerance)
 	, MeshingTolerance(GeometricTolerance* GeometricToMeshingToleranceFactor)
 	, SquareMeshingTolerance(FMath::Square(MeshingTolerance))
+	, MesherReport(InMesherReport)
 {
 	FinalInnerSegments.Reserve(3 * Grid.InnerNodesCount());
 	IndexOfLowerLeftInnerNodeSurroundingALoop.Reserve(Grid.GetLoopCount());
@@ -78,29 +80,42 @@ bool FIsoTriangulator::Triangulate()
 	BuildLoopSegments();
 
 #ifdef DEBUG_ISOTRIANGULATOR
-	Display(DisplaySpace, TEXT("FIsoTrianguler::LoopSegments"), LoopSegments, false, false, EVisuProperty::GreenCurve);
+	DisplayLoops(TEXT("FIsoTrianguler::LoopSegments"), true, false);
 #endif
 
-	if (!FindLoopIntersectionAndFixIt())
+	FLoopCleaner LoopCleaner(*this);
+	if (!LoopCleaner.Run())
 	{
+		FMessage::Printf(EVerboseLevel::Log, TEXT("The meshing of the surface %d failed due to a degenerated loop\n"), Grid.GetFace()->GetId());
 		return false;
 	}
 
 #ifdef DEBUG_ISOTRIANGULATOR
-	Display(DisplaySpace, TEXT("FIsoTrianguler::LoopSegments after fix intersection"), LoopSegments, true, false, EVisuProperty::OrangeCurve);
+	DisplayLoops(TEXT("FIsoTrianguler::LoopSegments after fix intersection"), true, false);
+	Grid.DisplayIsoSegments(TEXT("FIsoTrianguler::LoopSegments orientation"), DisplaySpace, LoopSegments, false, true, EVisuProperty::YellowCurve);
 	Wait(false);
 #endif
 
 	BuildThinZoneSegments();
 
+	//Fill Intersection tool
+	LoopSegmentsIntersectionTool.Empty(LoopSegments.Num());
+	LoopSegmentsIntersectionTool.AddSegments(LoopSegments);
+	LoopSegmentsIntersectionTool.AddSegments(ThinZoneSegments);
+	LoopSegmentsIntersectionTool.Sort();
+
 #ifdef DEBUG_ISOTRIANGULATOR
-	Display(DisplaySpace, TEXT("FIsoTrianguler::ThinZoneSegments"), ThinZoneSegments, false);
+	Grid.DisplayIsoSegments(TEXT("FIsoTrianguler::ThinZoneSegments"), DisplaySpace, ThinZoneSegments, false);
 	LoopSegmentsIntersectionTool.Display(TEXT("FIsoTrianguler::IntersectionTool Loop"));
 #endif
 
 	BuildInnerSegments();
+
 #ifdef DEBUG_ISOTRIANGULATOR
-	Display(DisplaySpace, TEXT("FIsoTrianguler::FinalInnerSegments"), FinalInnerSegments, false, false, EVisuProperty::BlueCurve);
+	{
+		F3DDebugSession A(bDisplay, TEXT("FIsoTrianguler::FinalInnerSegments"));
+		Grid.DisplayIsoSegments(DisplaySpace, FinalInnerSegments, true, false, EVisuProperty::BlueCurve);
+	}
 	InnerToOuterSegmentsIntersectionTool.Display(TEXT("FIsoTrianguler::IntersectionTool InnerToOutter"));
 	Chronos.TriangulateDuration1 = FChrono::Elapse(StartTime);
 #endif
@@ -125,12 +140,14 @@ bool FIsoTriangulator::Triangulate()
 	// 	           - find the shortest segment to connect each connected loop by Delaunay
 	// =============================================================================================================
 #ifdef DEBUG_ISOTRIANGULATOR
-	Display(DisplaySpace, TEXT("FIsoTrianguler::Final To Loops Before"), FinalToLoops, false, false, EVisuProperty::YellowCurve);
+	Grid.DisplayIsoSegments(TEXT("FIsoTrianguler::Final To Loops Before"), DisplaySpace, FinalToLoops, false, false, EVisuProperty::YellowCurve);
 	//Wait(bDisplay);
 #endif
+
 	ConnectCellLoops();
+
 #ifdef DEBUG_ISOTRIANGULATOR
-	Display(DisplaySpace, TEXT("FIsoTrianguler::Final Iso ToLink Inner To Loops"), FinalToLoops, false, false, EVisuProperty::YellowCurve);
+	Grid.DisplayIsoSegments(TEXT("FIsoTrianguler::Final Iso ToLink Inner To Loops"), DisplaySpace, FinalToLoops, false, false, EVisuProperty::YellowCurve);
 #endif
 
 
@@ -209,7 +226,7 @@ bool FIsoTriangulator::Triangulate()
 	// Triangulate between inner grid boundary and loops
 	TriangulateOverCycle(EGridSpace::Scaled);
 
-	// Finalise the mesh by the tessellation of the inner grid
+	// Finalize the mesh by the tessellation of the inner grid
 	TriangulateInnerNodes();
 
 #ifdef DEBUG_ISOTRIANGULATOR
@@ -1999,8 +2016,8 @@ void FIsoTriangulator::MeshCycle(const EGridSpace Space, const TArray<FIsoSegmen
 		if (bDisplay)
 		{
 			F3DDebugSession G(TEXT("Mesh cycle"));
-			DisplayTriangle(EGridSpace::UniformScaled, *Nodes[NodeIndices[1]], *Nodes[NodeIndices[3]], *Nodes[NodeIndices[0]]);
-			DisplayTriangle(EGridSpace::UniformScaled, *Nodes[NodeIndices[1]], *Nodes[NodeIndices[2]], *Nodes[NodeIndices[3]]);
+			Grid.DisplayTriangle(EGridSpace::UniformScaled, *Nodes[NodeIndices[1]], *Nodes[NodeIndices[3]], *Nodes[NodeIndices[0]]);
+			Grid.DisplayTriangle(EGridSpace::UniformScaled, *Nodes[NodeIndices[1]], *Nodes[NodeIndices[2]], *Nodes[NodeIndices[3]]);
 		}
 #endif 
 		Mesh->AddTriangle(Nodes[NodeIndices[1]]->GetFaceIndex(), Nodes[NodeIndices[3]]->GetFaceIndex(), Nodes[NodeIndices[0]]->GetFaceIndex());
@@ -2016,7 +2033,7 @@ void FIsoTriangulator::MeshCycle(const EGridSpace Space, const TArray<FIsoSegmen
 			if (bDisplay)
 			{
 				F3DDebugSession G(TEXT("Mesh cycle"));
-				DisplayTriangle(EGridSpace::UniformScaled, Cycle[0]->GetFirstNode(), Cycle[0]->GetSecondNode(), CycleOrientation[1] ? Cycle[1]->GetSecondNode() : Cycle[1]->GetFirstNode());
+				Grid.DisplayTriangle(EGridSpace::UniformScaled, Cycle[0]->GetFirstNode(), Cycle[0]->GetSecondNode(), CycleOrientation[1] ? Cycle[1]->GetSecondNode() : Cycle[1]->GetFirstNode());
 			}
 #endif 
 			Mesh->AddTriangle(Cycle[0]->GetFirstNode().GetFaceIndex(), Cycle[0]->GetSecondNode().GetFaceIndex(), CycleOrientation[1] ? Cycle[1]->GetSecondNode().GetFaceIndex() : Cycle[1]->GetFirstNode().GetFaceIndex());
@@ -2027,7 +2044,7 @@ void FIsoTriangulator::MeshCycle(const EGridSpace Space, const TArray<FIsoSegmen
 			if (bDisplay)
 			{
 				F3DDebugSession G(TEXT("Mesh cycle"));
-				DisplayTriangle(EGridSpace::UniformScaled, Cycle[0]->GetSecondNode(), Cycle[0]->GetFirstNode(), CycleOrientation[1] ? Cycle[1]->GetSecondNode() : Cycle[1]->GetFirstNode());
+				Grid.DisplayTriangle(EGridSpace::UniformScaled, Cycle[0]->GetSecondNode(), Cycle[0]->GetFirstNode(), CycleOrientation[1] ? Cycle[1]->GetSecondNode() : Cycle[1]->GetFirstNode());
 			}
 #endif 
 			Mesh->AddTriangle(Cycle[0]->GetSecondNode().GetFaceIndex(), Cycle[0]->GetFirstNode().GetFaceIndex(), CycleOrientation[1] ? Cycle[1]->GetSecondNode().GetFaceIndex() : Cycle[1]->GetFirstNode().GetFaceIndex());
@@ -2044,10 +2061,7 @@ void FIsoTriangulator::MeshCycle(const EGridSpace Space, const TArray<FIsoSegmen
 
 #ifdef ADD_TRIANGLE_2D
 	F3DDebugSession _(bDisplay, TEXT("Mesh cycle"));
-	if (bDisplay)
-	{
-		DisplayCycle(Cycle, TEXT("Cycle"));
-	}
+	Grid.DisplayIsoSegments(TEXT("Cycle"), EGridSpace::UniformScaled, Cycle);
 #endif
 
 	// Check if the cycle is in self intersecting and fix it. 
@@ -2185,8 +2199,8 @@ void FIsoTriangulator::MeshCycle(const EGridSpace Space, const TArray<FIsoSegmen
 	TFunction<void(FIsoSegment*, bool)> FindBestTriangle = [&](FIsoSegment* Segment, bool bOrientation)
 	{
 		using namespace IsoTriangulatorImpl;
-		GetSlopMethod GetSlopAtStartNode = ClockwiseSlop;
-		GetSlopMethod GetSlopAtEndNode = CounterClockwiseSlop;
+		SlopMethod GetSlopAtStartNode = ClockwiseSlop;
+		SlopMethod GetSlopAtEndNode = CounterClockwiseSlop;
 
 		// StartNode = A
 		FIsoNode& StartNode = bOrientation ? Segment->GetFirstNode() : Segment->GetSecondNode();
@@ -2346,7 +2360,7 @@ void FIsoTriangulator::MeshCycle(const EGridSpace Space, const TArray<FIsoSegmen
 				double PointCriteria = FMath::Abs(CotangentCriteria(StartPoint3D, EndPoint3D, NodePoint3D, NodeNormal));
 				double CosAngle = FMath::Abs(ComputeCosinus(NodeNormal, Node->GetNormal(Grid)));
 
-				// the criteria is weighted according to the cosinus of the angle between the normal of the candidate triangle and the normal at the tested point
+				// the criteria is weighted according to the cosine of the angle between the normal of the candidate triangle and the normal at the tested point
 				if (CosAngle > SMALL_NUMBER)
 				{
 					PointCriteria /= CosAngle;
@@ -2410,7 +2424,7 @@ void FIsoTriangulator::MeshCycle(const EGridSpace Space, const TArray<FIsoSegmen
 			{
 				{
 					F3DDebugSession _(FString::Printf(TEXT("Triangle")));
-					DisplayTriangle(EGridSpace::UniformScaled, EndNode, StartNode, *CandidatNode);
+					Grid.DisplayTriangle(EGridSpace::UniformScaled, EndNode, StartNode, *CandidatNode);
 				}
 			}
 #endif 
@@ -2657,16 +2671,16 @@ void FIsoTriangulator::TriangulateOverCycle(const EGridSpace Space)
 #endif
 }
 
-#ifdef FIND_CYCLE
+//#define DEBUG_FIND_CYCLE
+#ifdef DEBUG_FIND_CYCLE
 static int32 CycleId = -1;
 static int32 CycleIndex = 0;
 #endif
 
-//#define DEBUG_FIND_CYCLE
 void FIsoTriangulator::FindCycle(FIsoSegment* StartSegment, bool LeftSide, TArray<FIsoSegment*>& Cycle, TArray<bool>& CycleOrientation)
 {
 
-#ifdef FIND_CYCLE
+#ifdef DEBUG_FIND_CYCLE
 	CycleIndex++;
 	//CycleId = CycleIndex;
 	if (Grid.GetFace()->GetId() == FaceToDebug)
@@ -2698,8 +2712,8 @@ void FIsoTriangulator::FindCycle(FIsoSegment* StartSegment, bool LeftSide, TArra
 		{
 			FIsoNode* EndNode = &StartSegment->GetFirstNode();
 			F3DDebugSession _(TEXT("FirstSegment left"));
-			Display(EGridSpace::UniformScaled, *EndNode);
-			Display(EGridSpace::UniformScaled, *StartSegment);
+			Grid.DisplayIsoNode(EGridSpace::UniformScaled, *EndNode);
+			Grid.DisplayIsoSegment(EGridSpace::UniformScaled, *StartSegment);
 			//Wait();
 		}
 #endif
@@ -2713,8 +2727,8 @@ void FIsoTriangulator::FindCycle(FIsoSegment* StartSegment, bool LeftSide, TArra
 		{
 			FIsoNode* EndNode = &StartSegment->GetSecondNode();
 			F3DDebugSession _(TEXT("FirstSegment right"));
-			Display(EGridSpace::UniformScaled, *EndNode);
-			Display(EGridSpace::UniformScaled, *StartSegment);
+			Grid.DisplayIsoNode(EGridSpace::UniformScaled, *EndNode);
+			Grid.DisplayIsoSegment(EGridSpace::UniformScaled, *StartSegment);
 			//Wait();
 		}
 #endif
@@ -2726,7 +2740,7 @@ void FIsoTriangulator::FindCycle(FIsoSegment* StartSegment, bool LeftSide, TArra
 
 	for (;;)
 	{
-		Segment = FindNextSegment(EGridSpace::UniformScaled, Segment, Node, IsoTriangulatorImpl::ClockwiseSlop);
+		Segment = FindNextSegment(EGridSpace::UniformScaled, Segment, Node, ClockwiseSlop);
 		if (Segment == nullptr)
 		{
 			Cycle.Empty();
@@ -2746,7 +2760,7 @@ void FIsoTriangulator::FindCycle(FIsoSegment* StartSegment, bool LeftSide, TArra
 			if (Segment->HasCycleOnLeft())
 			{
 				F3DDebugSession _(TEXT("Segment HasCycleOnLeft"));
-				Display(EGridSpace::UniformScaled, *Segment, 0, EVisuProperty::RedPoint);
+				Grid.DisplayIsoSegment(EGridSpace::UniformScaled, *Segment, 0, EVisuProperty::RedPoint);
 				Wait();
 			}
 #endif
@@ -2758,8 +2772,8 @@ void FIsoTriangulator::FindCycle(FIsoSegment* StartSegment, bool LeftSide, TArra
 			if (CycleId == CycleIndex)
 			{
 				F3DDebugSession _(TEXT("Next"));
-				Display(EGridSpace::UniformScaled, *Node);
-				Display(EGridSpace::UniformScaled, *Segment);
+				Grid.DisplayIsoNode(EGridSpace::UniformScaled, *Node);
+				Grid.DisplayIsoSegment(EGridSpace::UniformScaled, *Segment);
 				//Wait();
 			}
 #endif
@@ -2770,7 +2784,7 @@ void FIsoTriangulator::FindCycle(FIsoSegment* StartSegment, bool LeftSide, TArra
 			if (Segment->HasCycleOnRight())
 			{
 				F3DDebugSession _(TEXT("Segment HasCycleOnRight"));
-				Display(EGridSpace::UniformScaled, *Segment, 0, EVisuProperty::RedPoint);
+				Grid.DisplayIsoSegment(EGridSpace::UniformScaled, *Segment, 0, EVisuProperty::RedPoint);
 				Wait();
 			}
 #endif
@@ -2782,8 +2796,8 @@ void FIsoTriangulator::FindCycle(FIsoSegment* StartSegment, bool LeftSide, TArra
 			if (CycleId == CycleIndex)
 			{
 				F3DDebugSession _(TEXT("Next"));
-				Display(EGridSpace::UniformScaled, *Node);
-				Display(EGridSpace::UniformScaled, *Segment);
+				Grid.DisplayIsoNode(EGridSpace::UniformScaled, *Node);
+				Grid.DisplayIsoSegment(EGridSpace::UniformScaled, *Segment);
 				//Wait();
 			}
 #endif
@@ -2806,7 +2820,7 @@ void FIsoTriangulator::FindCycle(FIsoSegment* StartSegment, bool LeftSide, TArra
 static bool bDisplayStar = false;
 #endif
 
-FIsoSegment* FIsoTriangulator::FindNextSegment(EGridSpace Space, const FIsoSegment* StartSegment, const FIsoNode* StartNode, IsoTriangulatorImpl::GetSlopMethod GetSlop) const
+FIsoSegment* FIsoTriangulator::FindNextSegment(EGridSpace Space, const FIsoSegment* StartSegment, const FIsoNode* StartNode, SlopMethod GetSlop) const
 {
 	const FPoint2D& StartPoint = StartNode->Get2DPoint(Space, Grid);
 	const FPoint2D& EndPoint = (StartNode == &StartSegment->GetFirstNode()) ? StartSegment->GetSecondNode().Get2DPoint(Space, Grid) : StartSegment->GetFirstNode().Get2DPoint(Space, Grid);
@@ -2911,8 +2925,8 @@ void FIsoTriangulator::TriangulateInnerNodes()
 			}
 
 #ifdef ADD_TRIANGLE_2D
-			DisplayTriangle(EGridSpace::Default2D, *GlobalIndexToIsoInnerNodes[Index], *GlobalIndexToIsoInnerNodes[Index + 1], *GlobalIndexToIsoInnerNodes[OppositIndex]);
-			DisplayTriangle(EGridSpace::Default2D, *GlobalIndexToIsoInnerNodes[OppositIndex], *GlobalIndexToIsoInnerNodes[OppositIndex - 1], *GlobalIndexToIsoInnerNodes[Index]);
+			Grid.DisplayTriangle(EGridSpace::Default2D, *GlobalIndexToIsoInnerNodes[Index], *GlobalIndexToIsoInnerNodes[Index + 1], *GlobalIndexToIsoInnerNodes[OppositIndex]);
+			Grid.DisplayTriangle(EGridSpace::Default2D, *GlobalIndexToIsoInnerNodes[OppositIndex], *GlobalIndexToIsoInnerNodes[OppositIndex - 1], *GlobalIndexToIsoInnerNodes[Index]);
 #endif 
 
 			Mesh->AddTriangle(GlobalIndexToIsoInnerNodes[Index]->GetFaceIndex(), GlobalIndexToIsoInnerNodes[Index + 1]->GetFaceIndex(), GlobalIndexToIsoInnerNodes[OppositIndex]->GetFaceIndex());
@@ -3145,8 +3159,8 @@ void FIsoTriangulator::TryToConnectTwoLoopsWithIsocelesTriangle(FCell& Cell, con
 
 	TFunction<FIsoNode* (FIsoSegment*)> FindBestTriangle = [&](FIsoSegment* Segment) -> FIsoNode*
 	{
-		IsoTriangulatorImpl::GetSlopMethod GetSlopAtStartNode = IsoTriangulatorImpl::ClockwiseSlop;
-		IsoTriangulatorImpl::GetSlopMethod GetSlopAtEndNode = IsoTriangulatorImpl::CounterClockwiseSlop;
+		SlopMethod GetSlopAtStartNode = ClockwiseSlop;
+		SlopMethod GetSlopAtEndNode = CounterClockwiseSlop;
 
 		// StartNode = A
 		FIsoNode& StartNode = Segment->GetSecondNode();
@@ -3731,7 +3745,7 @@ void FIsoTriangulator::FindSegmentToLinkLoopToLoop()
 
 		// Is Outside and not to flat at NodeA
 		const double FlatAngle = 1;
-		if (!IsPointPBeInsideSectorABC(NodeA->GetPreviousNode().Get2DPoint(EGridSpace::Default2D, Grid), ACoordinates, NodeA->GetNextNode().Get2DPoint(EGridSpace::Default2D, Grid), BCoordinates, FlatAngle))
+		if (!IsPointPInsideSectorABC(NodeA->GetPreviousNode().Get2DPoint(EGridSpace::Default2D, Grid), ACoordinates, NodeA->GetNextNode().Get2DPoint(EGridSpace::Default2D, Grid), BCoordinates, FlatAngle))
 		{
 			return;
 		}
@@ -3745,7 +3759,7 @@ void FIsoTriangulator::FindSegmentToLinkLoopToLoop()
 		//}
 
 		// Is Outside and not to flat at NodeB
-		if (!IsPointPBeInsideSectorABC(NodeB->GetPreviousNode().Get2DPoint(EGridSpace::Default2D, Grid), BCoordinates, NodeB->GetNextNode().Get2DPoint(EGridSpace::Default2D, Grid), ACoordinates, FlatAngle))
+		if (!IsPointPInsideSectorABC(NodeB->GetPreviousNode().Get2DPoint(EGridSpace::Default2D, Grid), BCoordinates, NodeB->GetNextNode().Get2DPoint(EGridSpace::Default2D, Grid), ACoordinates, FlatAngle))
 		{
 			return;
 		}
diff --git a/Engine/Source/Runtime/Datasmith/CADKernel/Private/CADKernel/Mesh/Meshers/IsoTriangulator.h b/Engine/Source/Runtime/Datasmith/CADKernel/Private/CADKernel/Mesh/Meshers/IsoTriangulator.h
--- a/Engine/Source/Runtime/Datasmith/CADKernel/Private/CADKernel/Mesh/Meshers/IsoTriangulator.h
+++ b/Engine/Source/Runtime/Datasmith/CADKernel/Private/CADKernel/Mesh/Meshers/IsoTriangulator.h
@@ -8,6 +8,7 @@
 #include "CADKernel/Mesh/Meshers/IsoTriangulator/IntersectionSegmentTool.h"
 #include "CADKernel/Mesh/Meshers/IsoTriangulator/IsoNode.h"
 #include "CADKernel/Mesh/Meshers/IsoTriangulator/IsoSegment.h"
+#include "CADKernel/Mesh/Meshers/MesherReport.h"
 #include "CADKernel/UI/Visu.h"
 
 #include "CADKernel/Mesh/Meshers/IsoTriangulator/DefineForDebug.h"
@@ -23,15 +24,6 @@ class FFaceMesh;
 struct FCell;
 
 
-namespace IsoTriangulatorImpl
-{
-typedef TFunction<double(const FPoint2D&, const FPoint2D&, double)> GetSlopMethod;
-typedef TFunction<FLoopNode* (FLoopNode*)> GetNextNodeMethod;
-typedef TFunction<const FLoopNode* (const FLoopNode*)> GetNextConstNodeMethod;
-typedef TFunction<const FLoopNode* (const FIsoSegment*)> GetSegmentToNodeMethod;
-typedef TPair<FLoopNode*, FLoopNode*> FLoopSection;
-}
-
 struct CADKERNEL_API FIsoTriangulatorChronos
 {
 	FDuration TriangulateDuration1 = FChrono::Init();
@@ -96,6 +88,7 @@ struct CADKERNEL_API FIsoTriangulatorChronos
 class FIsoTriangulator
 {
 	friend class FParametricMesher;
+	friend class FLoopCleaner;
 
 protected:
 
@@ -180,9 +173,11 @@ protected:
 	const double MeshingTolerance;
 	const double SquareMeshingTolerance;
 
+	FMesherReport& MesherReport;
+
 public:
 
-	FIsoTriangulator(FGrid& InGrid, TSharedRef<FFaceMesh> EntityMesh);
+	FIsoTriangulator(FGrid& InGrid, TSharedRef<FFaceMesh> EntityMesh, FMesherReport& InMesherReport);
 
 	/**
 	 * Main method
@@ -192,8 +187,6 @@ public:
 
 	void BuildNodes();
 
-	bool FindLoopIntersectionAndFixIt();
-
 	/**
 	 * Build the segments of the loops and check if each loop is self intersecting.
 	 * @return false if the loop is self intersecting
@@ -346,49 +339,6 @@ public:
 
 	bool CanCycleBeMeshed(const TArray<FIsoSegment*>& Cycle, FIntersectionSegmentTool& CycleIntersectionTool);
 
-	bool TryToRemoveSelfIntersectionByMovingTheClosedOusidePoint(const FIsoSegment& Segment0, const FIsoSegment& Segment1);
-	bool TryToRemoveIntersectionByMovingTheClosedOusidePoint(const FIsoSegment& Segment0, const FIsoSegment& Segment1);
-	void RemoveIntersectionByMovingOutsideSegmentNodeInside(const FIsoSegment& IntersectingSegment, const FIsoSegment& Segment);
-
-	bool TryToRemoveIntersectionBySwappingSegments(FIsoSegment& Segment0, FIsoSegment& Segment1);                                  // ok
-
-		/**
-		 * Segment and the segment next (or before Segment) are intersecting IntersectingSegment
-		 * The common node is moved
-		 */
-	bool TryToRemoveIntersectionOfTwoConsecutiveIntersectingSegments(const FIsoSegment& IntersectingSegment, FIsoSegment& Segment); // ok 
-
-	bool RemoveNodeOfLoop(FLoopNode& NodeToRemove);
-	void RemovePickOfLoop(FIsoSegment& Segment);
-
-	bool RemovePickRecursively(FLoopNode* Node0, FLoopNode* Node1);
-
-	bool CheckAndRemovePick(const FPoint2D& PreviousPoint, const FPoint2D& NodeToRemovePoint, const FPoint2D& NextPoint, FLoopNode& NodeToRemove)
-	{
-		double Slop = ComputeUnorientedSlope(NodeToRemovePoint, PreviousPoint, NextPoint);
-		bool bRemoveNode = false;
-		if (Slop < 0.03)
-		{
-			bRemoveNode = true;
-		}
-		else if (Slop < 0.1)
-		{
-			double SquareDistance1 = SquareDistanceOfPointToSegment(PreviousPoint, NodeToRemovePoint, NextPoint);
-			double SquareDistance2 = SquareDistanceOfPointToSegment(NextPoint, NodeToRemovePoint, PreviousPoint);
-			double MinSquareDistance = FMath::Min(SquareDistance1, SquareDistance2);
-			if (MinSquareDistance < SquareGeometricTolerance2)
-			{
-				bRemoveNode = true;
-			}
-		}
-
-		if (bRemoveNode)
-		{
-			return RemoveNodeOfLoop(NodeToRemove);
-		}
-		return false;
-	};
-
 	/**
 	 * Finalization of the mesh by the tessellation of the inner grid
 	 */
@@ -415,102 +365,13 @@ public:
 
 private:
 
-	// Methods used by FindLoopIntersectionAndFixIt Step1
-	void FindBestLoopExtremity(TArray<FLoopNode*>& BestStartNodeOfLoop);
-	void FindLoopIntersections(const TArray<FLoopNode*>& LoopNodes, bool bForward, TArray<TPair<double, double>>& OutIntersections);
-	EOrientation GetLoopOrientation(const FLoopNode* StartNode);
-	void FixLoopOrientation(const TArray<FLoopNode*>& LoopNodes);
-	void SwapLoopOrientation(int32 FirstSegmentIndex, int32 LastSegmentIndex);
-	void MoveIntersectingSectionBehindOppositeSection(IsoTriangulatorImpl::FLoopSection IntersectingSection, IsoTriangulatorImpl::FLoopSection OppositeSection, IsoTriangulatorImpl::GetNextNodeMethod GetNext, IsoTriangulatorImpl::GetNextNodeMethod GetPrevious);
-	void OffsetSegment(FIsoSegment& Segment, TSegment<FPoint2D>& Segment2D, TSegment<FPoint2D>& IntersectingSegment2D);
-	void OffsetNode(FLoopNode& Node, TSegment<FPoint2D>& IntersectingSegment2D);
-
-	/**
-	 * If the new position of the node is closed to the previous or next node position, it will be delete
-	 */
-	void MoveNode(FLoopNode& NodeToMove, FPoint2D& ProjectedPoint);
-
-	/**
-	 * @return false if the SubLoop is bigger than the other part of the loop
-	 */
-	bool RemoveLoopIntersections(const TArray<FLoopNode*>& NodesOfLoop, TArray<TPair<double, double>>& Intersections, bool bForward);
-
-	void RemoveLoopPicks(TArray<FLoopNode*>& NodesOfLoop, TArray<TPair<double, double>>& Intersections);
-
-	void RemoveThePick(const TArray<FLoopNode*>& NodesOfLoop, const TPair<double, double>& Intersection, bool bForward);
-	void SwapNodes(const TArray<FLoopNode*>& NodesOfLoop, const TPair<double, double>& Intersection, bool bForward);
-
-	/**
-	 * Case0: __     ___        __      ___
-	 *          \   /             \    /
-	 *           o o        =>     o  o
-	 *        __/   \___		__/    \___
-	 *
-	 * Case1:    ____    ___        ____    ___
-	 *           |   \  /   |       |   \  /   |
-	 *           |    oo    |   =>  |    o     |
-	 *           |___/  \___|       |          |
-	 *                              |    o     |
-	 * 								|___/  \___|
-	 */
-	bool SpreadCoincidentNodes(const TArray<FLoopNode*>& NodesOfLoop, TPair<double, double> Intersection);
-
-	/**
-	 * Case:  o_________o          o_________o
-	 *             o         =>
-	 *          __/ \___		        o
-	 *								 __/ \___
-	 */
-	bool MovePickBehind(const TArray<FLoopNode*>& NodesOfLoop, TPair<double, double> Intersection, bool bKeyIsExtremity);
-
-	/**
-	 * Two cases:
-	 *    - the segments of the intersection a closed parallel and in same orientation. in this case, the sub-loop is a long pick. The sub-Loop is delete
-	 *    - the loop is an inner loop closed to the border:
-	 *       _____________                _____________
-	 *   |  /             \			  |__/             \
-	 *    \/               |     =>                     |
-	 *    /\               |		   __               |
-	 *   |  \_____________/			  |  \_____________/
-	 *
-	 */
-	void TryToSwapSegmentsOrRemoveLoop(const TArray<FLoopNode*>& NodesOfLoop, const TPair<double, double>& Intersection, bool bForward);
-
-	bool RemovePickToOutside(const TArray<FLoopNode*>& NodesOfLoop, const TPair<double, double>& Intersection, const TPair<double, double>& NextIntersection, bool bForward);
-
-	/**
-	 * @return false if the SubLoop is bigger than the other part of the loop
-	 */
-	bool RemoveIntersectionsOfSubLoop(const TArray<FLoopNode*>& LoopNodesFromStartNode, TArray<TPair<double, double>>& Intersections, int32 IntersectionIndex, int32 IntersectionCount, bool bForward);
-	bool RemoveUniqueIntersection(const TArray<FLoopNode*>& LoopNodesFromStartNode, TPair<double, double> Intersection, bool bForward);
-	//void RemoveSubLoop(const TArray<FLoopNode*>& NodesOfLoop, const TPair<double, double>& Intersection, bool bForward);
-	void RemoveSubLoop(FLoopNode* StartNode, FLoopNode* EndNode, IsoTriangulatorImpl::GetNextNodeMethod GetNext);
-
-	/**
-	 * Check if the sub-loop is not bigger than the main loop.
-	 * If the segment count of the SubLoop is bigger than 1/3 of the segment count of the main loop
-	 * the length of each part are compute.
-	 * If the sub-loop length is bigger than the main loop^, there is a problem (problem orientation ?)
-	 * The surface meshing is canceled to avoid a bigger problem during the meshing step, that could cause a process crash.
-	 */
-	bool IsSubLoopBiggerThanMainLoop(const TArray<FLoopNode*>& LoopNodesFromStartNode, const TPair<double, double>& Intersection, bool bForward);
-
-	void RemoveIntersectionByMovingOutsideNodeInside(const FIsoSegment& IntersectingSegment, FLoopNode& NodeToMove);
-
-	bool CheckMainLoopConsistency();
-
-	// Methods used by FindLoopIntersectionAndFixIt Step2
-	void FillIntersectionToolWithOuterLoop();
-
-	void FixIntersectionBetweenLoops();
-
-	FIsoSegment* FindNextSegment(EGridSpace Space, const FIsoSegment* StartSegment, const FIsoNode* StartNode, IsoTriangulatorImpl::GetSlopMethod GetSlop) const;
-
+	FIsoSegment* FindNextSegment(EGridSpace Space, const FIsoSegment* StartSegment, const FIsoNode* StartNode, SlopMethod GetSlop) const;
 
 
 #ifdef UNUSED_TO_DELETE_
 	void CreateLoopToLoopSegment(FLoopNode* NodeA, const FPoint2D& ACoordinates, FLoopNode* NodeB, const FPoint2D& BCoordinates, TArray<FIsoSegment*>& NewSegments);
 #endif
+
 	// ==========================================================================================
 	// 	   Create segments
 	// ==========================================================================================
@@ -562,24 +423,13 @@ public:
 	void DisplayPixel(const int32 IndexU, const int32 IndexV) const;
 
 	void DisplayIsoNodes(EGridSpace Space) const;
-	void Display(EGridSpace Space, const FIsoNode& Node, FIdent Ident = 0, EVisuProperty Property = EVisuProperty::BluePoint) const;
-	void Display(FString Message, EGridSpace Space, const TArray<const FIsoNode*>& Points, EVisuProperty Property) const;
-
-	void Display(EGridSpace Space, const FIsoNode& NodeA, const FIsoNode& NodeB, FIdent Ident = 0, EVisuProperty Property = EVisuProperty::Element) const;
-	void Display(EGridSpace Space, const FIsoSegment& Segment, FIdent Ident = 0, EVisuProperty Property = EVisuProperty::Element, bool bDisplayOrientation = false) const;
-	void Display(EGridSpace Space, const TCHAR* Message, const TArray<FIsoSegment*>& Segment, bool bDisplayNode, bool bDisplayOrientation = false, EVisuProperty Property = EVisuProperty::Element) const;
-
-	void DisplayLoops(EGridSpace Space, const TCHAR* Message, const TArray<FLoopNode>& Nodes, bool bDisplayNode, EVisuProperty Property = EVisuProperty::Element) const;
-	void DisplayLoops(const TCHAR* Message, bool bOneNode = true, bool bSplitBySegment = false, bool bDisplayNext = false, bool bDisplayPrevious = false) const;
-	void DisplayLoop(EGridSpace Space, const TCHAR* Message, const TArray<FLoopNode*>& Nodes, bool bDisplayNode, EVisuProperty Property = EVisuProperty::Element) const;
 
-	void DisplayCycle(const TArray<FIsoSegment*>& Cycle, const TCHAR* Message) const;
+	void DisplayLoops(const TCHAR* Message, bool bOneNode = true, bool bSplitBySegment = false) const;
+	void DisplayLoopsByNextAndPrevious(const TCHAR* Message) const;
 
 	void DisplayCells(const TArray<FCell>& Cells) const;
 	void DisplayCell(const FCell& Cell) const;
 	void DrawCellBoundary(int32 Index, EVisuProperty Property) const;
-
-	void DisplayTriangle(EGridSpace Space, const FIsoNode& NodeA, const FIsoNode& NodeB, const FIsoNode& NodeC) const;
 #endif
 
 };
@@ -656,95 +506,6 @@ inline double EquilateralCriteria(const PointType& SegmentA, const PointType& Se
 	return Criteria1 + Criteria2 + Criteria3;
 }
 
-inline FLoopNode* GetNextNodeImpl(FLoopNode* Node)
-{
-	return &Node->GetNextNode();
-}
-
-inline FLoopNode* GetPreviousNodeImpl(FLoopNode* Node)
-{
-	return &Node->GetPreviousNode();
-}
-
-inline const FLoopNode* GetNextConstNodeImpl(const FLoopNode* Node)
-{
-	return &Node->GetNextNode();
-}
-
-inline const FLoopNode* GetPreviousConstNodeImpl(const FLoopNode* Node)
-{
-	return &Node->GetPreviousNode();
-}
-
-inline const FLoopNode* GetFirstNode(const FIsoSegment* Segment)
-{
-	return (const FLoopNode*)&Segment->GetFirstNode();
-};
-
-inline const FLoopNode* GetSecondNode(const FIsoSegment* Segment)
-{
-	return (const FLoopNode*)&Segment->GetSecondNode();
-};
-
-inline void GetLoopNodeStartingFrom(FLoopNode* StartNode, bool bForward, TArray<FLoopNode*>& Loop)
-{
-	using namespace IsoTriangulatorImpl;
-	GetNextNodeMethod GetNext = bForward ? GetNextNodeImpl : GetPreviousNodeImpl;
-	FLoopNode* Node = StartNode;
-	Loop.Add(StartNode);
-	for (Node = GetNext(Node); Node != StartNode; Node = GetNext(Node))
-	{
-		Loop.Add(Node);
-	}
-}
-
-inline FLoopNode* GetNodeAt(const TArray<FLoopNode*>& NodesOfLoop, const int32 Index)
-{
-	FLoopNode* Node = NodesOfLoop[Index];
-#ifdef DEBUG_LOOP_INTERSECTION_AND_FIX_IT		
-	Wait(Node->IsDelete());
-#endif
-	ensureCADKernel(!Node->IsDelete());
-	return Node;
-};
-
-inline int32 NextIndex(const int32 NodeCount, int32 StartIndex)
-{
-	++StartIndex;
-	return (StartIndex >= NodeCount) ? 0 : StartIndex;
-};
-
-inline void RemoveDeletedNodes(TArray<FLoopNode*>& NodesOfLoop)
-{
-	int32 Index = NodesOfLoop.IndexOfByPredicate([](FLoopNode* Node) { return Node->IsDelete(); });
-	if (Index == INDEX_NONE)
-	{
-		return;
-	}
-	int32 NewIndex = Index;
-	for (; Index < NodesOfLoop.Num(); ++Index)
-	{
-		if (!NodesOfLoop[Index]->IsDelete())
-		{
-			NodesOfLoop[NewIndex++] = NodesOfLoop[Index];
-		}
-	}
-	NodesOfLoop.SetNum(NewIndex);
-}
-
-/**
- * for findNextSegment
- */
-inline double ClockwiseSlop(const FPoint2D& StartPoint, const FPoint2D& EndPoint, double ReferenceSlope)
-{
-	return 8. - ComputePositiveSlope(StartPoint, EndPoint, ReferenceSlope);
-}
-
-inline double CounterClockwiseSlop(const FPoint2D& StartPoint, const FPoint2D& EndPoint, double ReferenceSlope)
-{
-	return ComputePositiveSlope(StartPoint, EndPoint, ReferenceSlope);
-}
-
 
 } // namespace FIsoTriangulatorImpl
 
diff --git a/Engine/Source/Runtime/Datasmith/CADKernel/Private/CADKernel/Mesh/Meshers/IsoTriangulator/DefineForDebug.h b/Engine/Source/Runtime/Datasmith/CADKernel/Private/CADKernel/Mesh/Meshers/IsoTriangulator/DefineForDebug.h
--- a/Engine/Source/Runtime/Datasmith/CADKernel/Private/CADKernel/Mesh/Meshers/IsoTriangulator/DefineForDebug.h
+++ b/Engine/Source/Runtime/Datasmith/CADKernel/Private/CADKernel/Mesh/Meshers/IsoTriangulator/DefineForDebug.h
@@ -2,19 +2,26 @@
 // Refactor in progress
 #pragma once
 
-static int32 FaceToDebug = 0;
+//#define DEBUG_CADKERNEL
 
-//#define ADD_TRIANGLE_2D
-//#define DEBUG_ISOTRIANGULATOR
-//
-//#define DEBUG_LOOP_INTERSECTION_AND_FIX_IT		
-//#define DEBUG_REMOVE_LOOP_INTERSECTIONS
-//#define DEBUG_REMOVE_UNIQUE_INTERSECTION
-//#define DEBUG_SWAP_SEGMENTS_OR_REMOVE
-//#define DEBUG_FIND_LOOP_INTERSECTION_AND_FIX_IT
-//#define DEBUG_FIND_LOOP_INTERSECTIONS		
+#ifdef DEBUG_CADKERNEL
+static int32 FaceToDebug = 70;
+
+#define ADD_TRIANGLE_2D
+#define DEBUG_ISOTRIANGULATOR
+
+// Loop cleaner
+#define DEBUG_CLEAN_LOOPS
+
+#define DEBUG_LOOP_INTERSECTION_AND_FIX_IT		
+#define DEBUG_REMOVE_LOOP_INTERSECTIONS
+#define DEBUG_REMOVE_UNIQUE_INTERSECTION
+#define DEBUG_SWAP_SEGMENTS_OR_REMOVE
+#define DEBUG_FIND_LOOP_INTERSECTIONS	
+#define DEBUG_REMOVE_INTERSECTIONS
 
 
 #define DEBUG_ONLY_SURFACE_TO_DEBUG
-static bool BoolDisplayDebugMeshStep=false;
+static bool bDisplayDebugMeshStep=false;
 
+#endif
diff --git a/Engine/Source/Runtime/Datasmith/CADKernel/Private/CADKernel/Mesh/Meshers/IsoTriangulator/IsoNode.h b/Engine/Source/Runtime/Datasmith/CADKernel/Private/CADKernel/Mesh/Meshers/IsoTriangulator/IsoNode.h
--- a/Engine/Source/Runtime/Datasmith/CADKernel/Private/CADKernel/Mesh/Meshers/IsoTriangulator/IsoNode.h
+++ b/Engine/Source/Runtime/Datasmith/CADKernel/Private/CADKernel/Mesh/Meshers/IsoTriangulator/IsoNode.h
@@ -388,7 +388,7 @@ public:
 	bool IsSegmentBeInsideFace(const FPoint2D& EndSegmentCoordinate, const FGrid& Grid, const double FlatAngle)
 	{
 		//bool bIsOuterLoop = (GetLoopIndex() == 0);
-		return (/*bIsOuterLoop  == */!IsPointPBeInsideSectorABC(GetPreviousNode().Get2DPoint(EGridSpace::UniformScaled, Grid), Get2DPoint(EGridSpace::UniformScaled, Grid), GetNextNode().Get2DPoint(EGridSpace::UniformScaled, Grid), EndSegmentCoordinate, FlatAngle));
+		return (/*bIsOuterLoop  == */!IsPointPInsideSectorABC(GetPreviousNode().Get2DPoint(EGridSpace::UniformScaled, Grid), Get2DPoint(EGridSpace::UniformScaled, Grid), GetNextNode().Get2DPoint(EGridSpace::UniformScaled, Grid), EndSegmentCoordinate, FlatAngle));
 	}
 
 	void SetAsIsoU()
diff --git a/Engine/Source/Runtime/Datasmith/CADKernel/Private/CADKernel/Mesh/Meshers/IsoTriangulator/IsoSegment.h b/Engine/Source/Runtime/Datasmith/CADKernel/Private/CADKernel/Mesh/Meshers/IsoTriangulator/IsoSegment.h
--- a/Engine/Source/Runtime/Datasmith/CADKernel/Private/CADKernel/Mesh/Meshers/IsoTriangulator/IsoSegment.h
+++ b/Engine/Source/Runtime/Datasmith/CADKernel/Private/CADKernel/Mesh/Meshers/IsoTriangulator/IsoSegment.h
@@ -120,6 +120,14 @@ public:
 
 	void Clean()
 	{
+		if (FirstNode)
+		{
+			FirstNode->DisconnectSegment(*this);
+		}
+		if (SecondNode)
+		{
+			SecondNode->DisconnectSegment(*this);
+		}
 		FirstNode = nullptr;
 		SecondNode = nullptr;
 		States = EIsoSegmentStates::Delete;
@@ -261,6 +269,19 @@ public:
 		((FLoopNode*)FirstNode)->SetNextConnectedNode((FLoopNode*)SecondNode);
 	}
 
+	void ReplaceNode(FIsoNode& OldNode, FIsoNode& NewNode)
+	{
+		if (FirstNode == &OldNode)
+		{
+			SetFirstNode(NewNode);
+		}
+		else 
+		{
+			ensureCADKernel(SecondNode == &OldNode);
+			SetSecondNode(NewNode);
+		}
+	}
+
 	void SwapOrientation()
 	{
 		FIsoNode* TempNode = SecondNode;
diff --git a/Engine/Source/Runtime/Datasmith/CADKernel/Private/CADKernel/Mesh/Meshers/IsoTriangulator/IsoTriangulatorDisplay.cpp b/Engine/Source/Runtime/Datasmith/CADKernel/Private/CADKernel/Mesh/Meshers/IsoTriangulator/IsoTriangulatorDisplay.cpp
--- a/Engine/Source/Runtime/Datasmith/CADKernel/Private/CADKernel/Mesh/Meshers/IsoTriangulator/IsoTriangulatorDisplay.cpp
+++ b/Engine/Source/Runtime/Datasmith/CADKernel/Private/CADKernel/Mesh/Meshers/IsoTriangulator/IsoTriangulatorDisplay.cpp
@@ -13,155 +13,6 @@
 namespace CADKernel
 {
 
-void FIsoTriangulator::Display(EGridSpace Space, const FIsoSegment& Segment, FIdent Ident, EVisuProperty Property, bool bDisplayOrientation) const
-{
-	DisplaySegment(Segment.GetFirstNode().GetPoint(Space, Grid), Segment.GetSecondNode().GetPoint(Space, Grid), Ident, Property, bDisplayOrientation);
-}
-
-void FIsoTriangulator::DisplayTriangle(EGridSpace Space, const FIsoNode& NodeA, const FIsoNode& NodeB, const FIsoNode& NodeC) const
-{
-	TArray<FPoint> Points;
-	Points.SetNum(3);
-	Points[0] = NodeA.GetPoint(Space, Grid);
-	Points[1] = NodeB.GetPoint(Space, Grid);
-	Points[2] = NodeC.GetPoint(Space, Grid);
-	DrawElement(2, Points, EVisuProperty::Element);
-	DisplaySegment(Points[0], Points[1]); 
-	DisplaySegment(Points[1], Points[2]);
-	DisplaySegment(Points[2], Points[0]);
-}
-
-void FIsoTriangulator::Display(EGridSpace Space, const FIsoNode& NodeA, const FIsoNode& NodeB, FIdent Ident, EVisuProperty Property) const
-{
-	return DisplaySegment(NodeA.GetPoint(Space, Grid), NodeB.GetPoint(Space, Grid), Ident, Property);
-}
-
-void FIsoTriangulator::Display(EGridSpace Space, const FIsoNode& Node, FIdent Ident, EVisuProperty Property) const
-{
-	return DisplayPoint(Node.GetPoint(Space, Grid), Property, Ident);
-}
-
-void FIsoTriangulator::Display(FString Message, EGridSpace Space, const TArray<const FIsoNode*>& Nodes, EVisuProperty Property) const
-{
-	if (!bDisplay)
-	{
-		return;
-	}
-
-	Open3DDebugSession(Message);
-	for (int32 Index = 0; Index < Nodes.Num(); ++Index)
-	{
-		Display(Space, *Nodes[Index], Nodes[Index]->GetIndex(), Property);
-	}
-	Close3DDebugSession();
-}
-
-void FIsoTriangulator::Display(EGridSpace Space, const TCHAR* Message, const TArray<FIsoSegment*>& Segments, bool bDisplayNode, bool bDisplayOrientation, EVisuProperty Property) const
-{
-	if (!bDisplay)
-	{
-		return;
-	}
-
-	F3DDebugSession _(Message);
-#ifdef DEBUG_LIMIT
-	for (FIsoSegment* Segment : Segments)
-	{
-		if (Segment->IsALimit())
-		{
-			Display(Space, *Segment, 0, EVisuProperty::Iso, bDisplayOrientation);
-		}
-	}
-#endif
-	for (FIsoSegment* Segment : Segments)
-	{
-#ifdef DEBUG_LIMIT
-		if (!Segment->IsALimit())
-#endif
-		{
-			Display(Space, *Segment, 0, Property, bDisplayOrientation);
-		}
-	}
-
-	if (bDisplayNode)
-	{
-		for (FIsoSegment* Segment : Segments)
-		{
-			Display(Space, Segment->GetFirstNode(), Segment->GetFirstNode().GetId(), (EVisuProperty)(Property - 1));
-			Display(Space, Segment->GetSecondNode(), Segment->GetSecondNode().GetId(), (EVisuProperty)(Property - 1));
-		}
-	}
-	//Wait();
-}
-
-void FIsoTriangulator::DisplayLoops(EGridSpace Space, const TCHAR* Message, const TArray<FLoopNode>& Nodes, bool bDisplayNode, EVisuProperty NodeVisuProperty) const
-{
-	if (!bDisplay)
-	{
-		return;
-	}
-
-	F3DDebugSession _(Message);
-
-	for (const FLoopNode& Node : Nodes)
-	{
-		if (!Node.IsDelete())
-		{
-			Display(Space, Node, Node.GetNextNode(), 0, NodeVisuProperty);
-		}
-	}
-
-	if (bDisplayNode)
-	{
-		for (const FLoopNode& Node : Nodes)
-		{
-			if (Node.IsDelete())
-			{
-				Display(Space, Node, Node.GetFaceIndex(), EVisuProperty::PurplePoint);
-			}
-			else
-			{
-				Display(Space, Node, Node.GetFaceIndex(), (EVisuProperty) ((int32)NodeVisuProperty + 1));
-			}
-		}
-	}
-}
-
-void FIsoTriangulator::DisplayLoop(EGridSpace Space, const TCHAR* Message, const TArray<FLoopNode*>& Nodes, bool bDisplayNode, EVisuProperty NodeVisuProperty) const
-{
-	if (!bDisplay)
-	{
-		return;
-	}
-
-	F3DDebugSession _(Message);
-
-	int32 Index = 0;
-	for (const FLoopNode* Node : Nodes)
-	{
-		if (!Node->IsDelete())
-		{
-			Display(Space, *Node, Node->GetNextNode(), Index++, (EVisuProperty)((int32)NodeVisuProperty + 1));
-		}
-	}
-
-	if (bDisplayNode)
-	{
-		for (const FLoopNode* Node : Nodes)
-		{
-			if (Node->IsDelete())
-			{
-				Display(Space, *Node, Node->GetFaceIndex(), EVisuProperty::PurplePoint);
-			}
-			else
-			{
-				Display(Space, *Node, Node->GetFaceIndex(), (EVisuProperty)((int32)NodeVisuProperty));
-			}
-		}
-	}
-}
-
-
 void FIsoTriangulator::DisplayIsoNodes(EGridSpace Space) const
 {
 	if (!bDisplay)
@@ -172,13 +23,13 @@ void FIsoTriangulator::DisplayIsoNodes(EGridSpace Space) const
 	Open3DDebugSession(TEXT("FIsoTrianguler::IsoNodes"));
 	for (const FLoopNode& Node : LoopNodes)
 	{
-		Display(Space, Node, Node.GetFaceIndex(), EVisuProperty::YellowPoint);
+		Grid.DisplayIsoNode(Space, Node, Node.GetFaceIndex(), EVisuProperty::YellowPoint);
 	}
 	Close3DDebugSession();
 	Open3DDebugSession(TEXT("FIsoTrianguler::IsoNodes Inner"));
 	for (const FIsoNode& Node : InnerNodes)
 	{
-		Display(Space, Node, Node.GetFaceIndex());
+		Grid.DisplayIsoNode(Space, Node, Node.GetFaceIndex());
 	}
 	Close3DDebugSession();
 }
@@ -212,16 +63,7 @@ void FIsoTriangulator::DisplayPixels(TArray<uint8>& Pixel) const
 	Close3DDebugSession();
 }
 
-void FIsoTriangulator::DisplayCycle(const TArray<FIsoSegment*>& Cycle, const TCHAR* Message) const
-{
-	F3DDebugSession _(Message);
-	for (const FIsoSegment* Segment : Cycle)
-	{
-		Display(EGridSpace::UniformScaled, *Segment);
-	}
-}
-
-void FIsoTriangulator::DisplayLoops(const TCHAR* Message, bool bOneNode, bool bSplitBySegment, bool bDisplayNext, bool bDisplayPrevious) const
+void FIsoTriangulator::DisplayLoops(const TCHAR* Message, bool bOneNode, bool bSplitBySegment) const
 {
 	if (bDisplay)
 	{
@@ -234,9 +76,9 @@ void FIsoTriangulator::DisplayLoops(const TCHAR* Message, bool bOneNode, bool bS
 				{
 					if (Node.IsDelete())
 					{
-						Display(EGridSpace::UniformScaled, Node, Index++, EVisuProperty::PurplePoint);
+						Grid.DisplayIsoNode(EGridSpace::UniformScaled, Node, Index++, EVisuProperty::PurplePoint);
 					}
-					Display(EGridSpace::UniformScaled, Node, Index++, EVisuProperty::BluePoint);
+					Grid.DisplayIsoNode(EGridSpace::UniformScaled, Node, Index++, EVisuProperty::BluePoint);
 				}
 			}
 
@@ -246,29 +88,38 @@ void FIsoTriangulator::DisplayLoops(const TCHAR* Message, bool bOneNode, bool bS
 				for (FIsoSegment* Segment : LoopSegments)
 				{
 					F3DDebugSession D(bSplitBySegment, TEXT("Segment"));
-					Display(EGridSpace::UniformScaled, Segment->GetFirstNode(), Segment->GetSecondNode(), Index++, EVisuProperty::BlueCurve);
+					Grid.DisplayIsoSegment(EGridSpace::UniformScaled, Segment->GetFirstNode(), Segment->GetSecondNode(), Index++, EVisuProperty::BlueCurve);
 				}
 			}
 
-			if(bDisplayNext)
+		}
+		Wait(false);
+	}
+}
+
+void FIsoTriangulator::DisplayLoopsByNextAndPrevious(const TCHAR* Message) const
+{
+	if (bDisplay)
+	{
+		F3DDebugSession _(Message);
+		{
 			{
 				F3DDebugSession A(TEXT("Node Next"));
 				int32 Index = 0;
 				for (FIsoSegment* Segment : LoopSegments)
 				{
-					F3DDebugSession B(bSplitBySegment, TEXT("Segment"));
-					Display(EGridSpace::UniformScaled, Segment->GetFirstNode(), ((const FLoopNode&) Segment->GetFirstNode()).GetNextNode(), Index++, EVisuProperty::BlueCurve);
+					F3DDebugSession B(TEXT("Segment"));
+					Grid.DisplayIsoSegment(EGridSpace::UniformScaled, Segment->GetFirstNode(), ((const FLoopNode&)Segment->GetFirstNode()).GetNextNode(), Index++, EVisuProperty::BlueCurve);
 				}
 			}
 
-			if(bDisplayPrevious)
 			{
 				F3DDebugSession C(TEXT("Node Forward"));
 				int32 Index = 0;
 				for (FIsoSegment* Segment : LoopSegments)
 				{
-					F3DDebugSession B(bSplitBySegment, TEXT("Segment"));
-					Display(EGridSpace::UniformScaled, Segment->GetSecondNode(), ((const FLoopNode&)Segment->GetSecondNode()).GetPreviousNode(), Index++, EVisuProperty::BlueCurve);
+					F3DDebugSession B(TEXT("Segment"));
+					Grid.DisplayIsoSegment(EGridSpace::UniformScaled, Segment->GetSecondNode(), ((const FLoopNode&)Segment->GetSecondNode()).GetPreviousNode(), Index++, EVisuProperty::BlueCurve);
 				}
 			}
 
