CADKernel Update

diff --git a/Engine/Source/Runtime/Datasmith/CADKernel/Private/CADKernel/Mesh/Meshers/IsoTriangulatorRemoveIntersection.cpp b/Engine/Source/Runtime/Datasmith/CADKernel/Private/CADKernel/Mesh/Meshers/IsoTriangulatorRemoveIntersection.cpp
--- a/Engine/Source/Runtime/Datasmith/CADKernel/Private/CADKernel/Mesh/Meshers/IsoTriangulatorRemoveIntersection.cpp
+++ b/Engine/Source/Runtime/Datasmith/CADKernel/Private/CADKernel/Mesh/Meshers/IsoTriangulatorRemoveIntersection.cpp
@@ -6,6 +6,7 @@
 #include "CADKernel/Math/Point.h"
 #include "CADKernel/Mesh/Meshers/IsoTriangulator/IsoNode.h"
 #include "CADKernel/Mesh/Meshers/IsoTriangulator/IsoSegment.h"
+#include "CADKernel/Topo/TopologicalFace.h"
 #include "Algo/AllOf.h"
 
 namespace CADKernel
@@ -13,438 +14,8 @@ namespace CADKernel
 
 //#define DEBUG_FIND_LOOP_INTERSECTION_AND_FIX_IT
 
-void FIsoTriangulator::FindBestLoopExtremity(TArray<FLoopNode*>& BestStartNodeOfLoop)
-{
-	double UMin = HUGE_VAL;
-	double UMax = -HUGE_VAL;
-	double VMin = HUGE_VAL;
-	double VMax = -HUGE_VAL;
-
-	FLoopNode* ExtremityNodes[4] = { nullptr, nullptr, nullptr, nullptr };
-
-	FLoopNode* BestNode = nullptr;
-	double OptimalSlop = 9.;
-	int32 LoopIndex = 0;
-
-	//F3DDebugSession _(TEXT("FindBestLoopExtremity"));
-
-	TFunction<void(FLoopNode*)> CompareWithSlopAt = [&](FLoopNode* Node)
-	{
-		FLoopNode& PreviousNode = Node->GetPreviousNode();
-		FLoopNode& NextNode = Node->GetNextNode();
-		double Slop = ComputePositiveSlope(Node->Get2DPoint(EGridSpace::UniformScaled, Grid), PreviousNode.Get2DPoint(EGridSpace::UniformScaled, Grid), NextNode.Get2DPoint(EGridSpace::UniformScaled, Grid));
-
-		//F3DDebugSession A(FString::Printf(TEXT("Node Slop %f"), Slop));
-		//Display(EGridSpace::UniformScaled, *Node, 0, EVisuProperty::RedPoint);
-
-		if ((Slop > OptimalSlop) == (LoopIndex == 0))
-		{
-			OptimalSlop = Slop;
-			BestNode = Node;
-		}
-	};
-
-	TFunction<void()> FindLoopExtremity = [&]()
-	{
-		BestNode = nullptr;
-		OptimalSlop = (LoopIndex == 0) ? -1 : 9.;
-
-		for (FLoopNode* Node : ExtremityNodes)
-		{
-			CompareWithSlopAt(Node);
-		}
-		BestStartNodeOfLoop.Add(BestNode);
-
-		// init for next loop
-		UMin = HUGE_VAL;
-		UMax = -HUGE_VAL;
-		VMin = HUGE_VAL;
-		VMax = -HUGE_VAL;
-
-		for (FLoopNode*& Node : ExtremityNodes)
-		{
-			Node = nullptr;
-		}
-	};
-
-	for (FLoopNode& Node : LoopNodes)
-	{
-		if (Node.GetLoopIndex() != LoopIndex)
-		{
-			FindLoopExtremity();
-			LoopIndex = Node.GetLoopIndex();
-		}
-
-		const FPoint2D& Point = Node.Get2DPoint(EGridSpace::UniformScaled, Grid);
-
-		if (Point.U > UMax)
-		{
-			UMax = Point.U;
-			ExtremityNodes[0] = &Node;
-		}
-		if (Point.U < UMin)
-		{
-			UMin = Point.U;
-			ExtremityNodes[1] = &Node;
-		}
-
-		if (Point.V > VMax)
-		{
-			VMax = Point.V;
-			ExtremityNodes[2] = &Node;
-		}
-		if (Point.V < VMin)
-		{
-			VMin = Point.V;
-			ExtremityNodes[3] = &Node;
-		}
-	}
-	FindLoopExtremity();
-}
-
-EOrientation FIsoTriangulator::GetLoopOrientation(const FLoopNode* StartNode)
-{
-	using namespace IsoTriangulatorImpl;
-	double UMin = HUGE_VAL;
-	double UMax = -HUGE_VAL;
-	double VMin = HUGE_VAL;
-	double VMax = -HUGE_VAL;
-
-	const FLoopNode* ExtremityNodes[4] = { nullptr, nullptr, nullptr, nullptr };
-
-	const FLoopNode* BestNode = nullptr;
-
-	int32 LoopIndex = StartNode->GetLoopIndex();
-	double OptimalSlop = (LoopIndex == 0) ? -1 : 9.;
-
-	TFunction<void(const FLoopNode*)> CompareWithSlopAt = [&](const FLoopNode* Node)
-	{
-		FLoopNode& PreviousNode = Node->GetPreviousNode();
-		FLoopNode& NextNode = Node->GetNextNode();
-		double Slop = ComputePositiveSlope(Node->Get2DPoint(EGridSpace::UniformScaled, Grid), PreviousNode.Get2DPoint(EGridSpace::UniformScaled, Grid), NextNode.Get2DPoint(EGridSpace::UniformScaled, Grid));
-
-		if ((Slop > OptimalSlop) == (LoopIndex == 0))
-		{
-			OptimalSlop = Slop;
-			BestNode = Node;
-		}
-	};
-
-	TFunction<void()> FindLoopExtremity = [&]()
-	{
-		BestNode = nullptr;
-
-		for (const FLoopNode* Node : ExtremityNodes)
-		{
-			CompareWithSlopAt(Node);
-		}
-
-		// init for next loop
-		UMin = HUGE_VAL;
-		UMax = -HUGE_VAL;
-		VMin = HUGE_VAL;
-		VMax = -HUGE_VAL;
-
-		for (const FLoopNode*& Node : ExtremityNodes)
-		{
-			Node = nullptr;
-		}
-	};
-
-	const FLoopNode* Node = GetNextConstNodeImpl(StartNode);
-	for (; Node != StartNode; Node = GetNextConstNodeImpl(Node))
-	{
-		const FPoint2D& Point = Node->Get2DPoint(EGridSpace::UniformScaled, Grid);
-
-		if (Point.U > UMax)
-		{
-			UMax = Point.U;
-			ExtremityNodes[0] = Node;
-		}
-		if (Point.U < UMin)
-		{
-			UMin = Point.U;
-			ExtremityNodes[1] = Node;
-		}
-
-		if (Point.V > VMax)
-		{
-			VMax = Point.V;
-			ExtremityNodes[2] = Node;
-		}
-		if (Point.V < VMin)
-		{
-			VMin = Point.V;
-			ExtremityNodes[3] = Node;
-		}
-	}
-	FindLoopExtremity();
-
-	if (LoopIndex == 0)
-	{
-		return OptimalSlop > 4 ? EOrientation::Front : EOrientation::Back;
-	}
-	else
-	{
-		return OptimalSlop < 4 ? EOrientation::Front : EOrientation::Back;
-	}
-}
-
-//#define DEBUG_FIND_LOOP_INTERSECTIONS
-void FIsoTriangulator::FindLoopIntersections(const TArray<FLoopNode*>& NodesOfLoop, bool bForward, TArray<TPair<double, double>>& OutIntersections)
-{
-	using namespace IsoTriangulatorImpl;
-
-	GetNextNodeMethod GetNext = bForward ? GetNextNodeImpl : GetPreviousNodeImpl;
-	GetSegmentToNodeMethod GetFirst = bForward ? GetFirstNode : GetSecondNode;
-	GetSegmentToNodeMethod GetSecond = bForward ? GetSecondNode : GetFirstNode;
-
-	FLoopNode* const* const StartNodePtr = NodesOfLoop.FindByPredicate([](const FLoopNode* Node) { return !Node->IsDelete(); });
-	if (*StartNodePtr == nullptr)
-	{
-		return;
-	}
-
-	FLoopNode* StartNode = *StartNodePtr;
-
-	TArray<const FIsoSegment*> IntersectedSegments;
-
-	int32 SegmentCount = 1;
-	int32 SegmentIndex = 1;
-	FLoopNode* Node = nullptr;
-	FLoopNode* NextNode = nullptr;
-
-	TFunction<void()> FindSegmentIntersection = [&]()
-	{
-		FIsoSegment* Segment = Node->GetSegmentConnectedTo(NextNode);
-		if (Segment == nullptr)
-		{
-			return;
-		}
-
-		if (LoopSegmentsIntersectionTool.FindIntersections(*Node, *NextNode, IntersectedSegments))
-		{
-			for (const FIsoSegment* IntersectedSegment : IntersectedSegments)
-			{
-				const FLoopNode* IntersectedSegmentFirstNode = GetFirst(IntersectedSegment);
-				if (IntersectedSegmentFirstNode == nullptr)
-				{
-					continue;
-				}
-
-				int32 IntersectionIndex = 0;
-				FLoopNode* TmpNode = StartNode;
-				while (TmpNode != IntersectedSegmentFirstNode)
-				{
-					++IntersectionIndex;
-					TmpNode = GetNext(TmpNode);
-				}
-
-				TSegment<FPoint2D> SegmentPoints(Node->Get2DPoint(EGridSpace::UniformScaled, Grid), NextNode->Get2DPoint(EGridSpace::UniformScaled, Grid));
-				TSegment<FPoint2D> IntersectedSegmentPoints(GetFirst(IntersectedSegment)->Get2DPoint(EGridSpace::UniformScaled, Grid), GetSecond(IntersectedSegment)->Get2DPoint(EGridSpace::UniformScaled, Grid));
-
-				double SegmentIntersectionCoordinate;
-				FPoint2D IntersectionPoint = FindIntersectionOfSegments2D(SegmentPoints, IntersectedSegmentPoints, SegmentIntersectionCoordinate);
-
-				double SquareLength = IntersectedSegmentPoints.SquareLength();
-				double IntersectedSegmentIntersectionCoordinate = (SquareLength > SMALL_NUMBER_SQUARE) ? FMath::Sqrt(IntersectedSegmentPoints[0].SquareDistance(IntersectionPoint) / IntersectedSegmentPoints.SquareLength()) : 0;
-				IntersectedSegmentIntersectionCoordinate += IntersectionIndex;
-
-				SegmentIntersectionCoordinate += SegmentIndex;
-				if (FMath::IsNearlyEqual(SegmentIntersectionCoordinate, SegmentCount))
-				{
-					SegmentIntersectionCoordinate = IntersectedSegmentIntersectionCoordinate;
-					IntersectedSegmentIntersectionCoordinate = 0;
-				}
-
-				if (OutIntersections.Num() && FMath::IsNearlyEqual(OutIntersections.Last().Key, IntersectedSegmentIntersectionCoordinate) && FMath::IsNearlyEqual(OutIntersections.Last().Value, SegmentIntersectionCoordinate))
-				{
-					continue;
-				}
-				OutIntersections.Emplace(IntersectedSegmentIntersectionCoordinate, SegmentIntersectionCoordinate);
-
-#ifdef DEBUG_FIND_LOOP_INTERSECTIONS		
-				if (bDisplay)
-				{
-					F3DDebugSession _(*FString::Printf(TEXT("Intersection %f %f"), IntersectedSegmentIntersectionCoordinate, SegmentIntersectionCoordinate));
-					Display(EGridSpace::UniformScaled, *Node, *NextNode, 0, EVisuProperty::BluePoint);
-					Display(EGridSpace::UniformScaled, *IntersectedSegment, 0, EVisuProperty::BluePoint);
-					DisplayPoint(IntersectionPoint, EVisuProperty::RedPoint);
-			}
-#endif
-		}
-	}
-
-		if (Segment)
-		{
-			LoopSegmentsIntersectionTool.AddSegment(*Segment);
-		}
-};
-
-	LoopSegmentsIntersectionTool.Empty(NodesOfLoop.Num());
-
-	Node = GetNext(StartNode);
-	for (; Node != StartNode; Node = GetNext(Node), ++SegmentCount);
 
-	Node = StartNode;
-	NextNode = GetNext(Node);
-
-#ifdef DEBUG_FIND_LOOP_INTERSECTIONS		
-	if (bDisplay)
-	{
-		F3DDebugSession _(TEXT("Start Node"));
-		Display(EGridSpace::UniformScaled, *Node, *NextNode, 0, EVisuProperty::BlueCurve);
-		Display(EGridSpace::UniformScaled, *Node, 0, EVisuProperty::RedPoint);
-	}
-#endif
-
-	FIsoSegment* StartToEndSegment = StartNode->GetSegmentConnectedTo(NextNode);
-	if (StartToEndSegment)
-	{
-		LoopSegmentsIntersectionTool.AddSegment(*StartToEndSegment);
-	}
-
-	for (Node = NextNode; Node != StartNode; Node = NextNode, ++SegmentIndex)
-	{
-		if (Node == nullptr || Node->IsDelete())
-		{
-			return;
-		}
-
-		NextNode = GetNext(Node);
-		if (NextNode == nullptr || NextNode->IsDelete())
-		{
-			return;
-		}
-		FindSegmentIntersection();
-	}
-
-}
-
-bool FIsoTriangulator::RemoveLoopIntersections(const TArray<FLoopNode*>& NodesOfLoop, TArray<TPair<double, double>>& Intersections, bool bForward)
-{
-	using namespace IsoTriangulatorImpl;
-	GetNextNodeMethod GetNext = bForward ? GetNextNodeImpl : GetPreviousNodeImpl;
-	GetNextNodeMethod GetPrevious = bForward ? GetPreviousNodeImpl : GetNextNodeImpl;
-	const int32 NodeCount = NodesOfLoop.Num();
-
-#ifdef DEBUG_REMOVE_LOOP_INTERSECTIONS		
-	TFunction<void(const TPair<double, double>&)> DisplayIntersection = [&](const TPair<double, double>& Intersection)
-	{
-		if (bDisplay)
-		{
-			FLoopNode* Segment0End = GetNodeAt(NodesOfLoop, NextIndex(NodeCount, (int32)Intersection.Key));
-			FLoopNode* Segment1Start = GetNodeAt(NodesOfLoop, (int32)Intersection.Value);
-
-			Display(EGridSpace::UniformScaled, *Segment0End, *GetPrevious(Segment0End), 0, EVisuProperty::RedCurve);
-			Display(EGridSpace::UniformScaled, *Segment1Start, *GetNext(Segment1Start), 0, EVisuProperty::RedCurve);
-			Display(EGridSpace::UniformScaled, *GetPrevious(Segment0End), 0, EVisuProperty::RedPoint);
-			Display(EGridSpace::UniformScaled, *Segment1Start, 0, EVisuProperty::RedPoint);
-		}
-	};
-#endif
-
-	Algo::Sort(Intersections, [&](const TPair<double, double>& Intersection1, const TPair<double, double>& Intersection2)
-		{
-			return (Intersection1.Key < Intersection2.Key);
-		});
-
-	// Nodes are not yet deleted in the sub loop
-	for (int32 IntersectionIndex = 0; IntersectionIndex < Intersections.Num(); )
-	{
-#ifdef DEBUG_REMOVE_LOOP_INTERSECTIONS		
-		if (bDisplay)
-		{
-			DisplayLoops(TEXT("RemoveLoopIntersections"));
-		}
-#endif
-
-		const TPair<double, double>& Intersection = Intersections[IntersectionIndex];
-
-#ifdef DEBUG_REMOVE_LOOP_INTERSECTIONS		
-		F3DDebugSession _(bDisplay, TEXT("Intersected Segments"));
-		DisplayIntersection(Intersection);
-#endif
-
-		bool bIntersectionForward = true;
-
-		int32 NextIntersectionIndex = IntersectionIndex + 1;
-		int32 IntersectionCount = 1;
-		for (; NextIntersectionIndex < Intersections.Num(); ++NextIntersectionIndex)
-		{
-			if (Intersections[NextIntersectionIndex].Value > Intersection.Value)
-			{
-				break;
-			}
-			++IntersectionCount;
-
-#ifdef DEBUG_REMOVE_LOOP_INTERSECTIONS		
-			DisplayIntersection(Intersections[NextIntersectionIndex]);
-#endif
-		}
-
-		if (IntersectionCount == 1)
-		{
-			NextIntersectionIndex = IntersectionIndex + 1;
-			for (; NextIntersectionIndex < Intersections.Num(); ++NextIntersectionIndex)
-			{
-				if (Intersections[NextIntersectionIndex].Key > Intersection.Value)
-				{
-					break;
-				}
-
-				bIntersectionForward = false;
-				++IntersectionCount;
-#ifdef DEBUG_REMOVE_LOOP_INTERSECTIONS		
-				DisplayIntersection(Intersections[NextIntersectionIndex]);
-#endif
-				break;
-			}
-		}
-
-		// Stating from this point 
-		// the process must not delete node after the first intersection
-		if (IntersectionCount == 1)
-		{
-			if (!RemoveUniqueIntersection(NodesOfLoop, Intersections[IntersectionIndex], bForward))
-			{
-				return false;
-			}
-			IntersectionIndex++;
-		}
-		else if (!bIntersectionForward)
-		{
-			if (!RemovePickToOutside(NodesOfLoop, Intersections[IntersectionIndex], Intersections[IntersectionIndex + 1], bForward))
-			{
-				return false;
-			}
-			IntersectionIndex += 2;
-		}
-		else
-		{
-			if (!RemoveIntersectionsOfSubLoop(NodesOfLoop, Intersections, IntersectionIndex, IntersectionCount, bForward))
-			{
-				return false;
-			}
-			IntersectionIndex += IntersectionCount;
-		}
-	}
-
-	return true;
-	}
-
-void FIsoTriangulator::RemoveLoopPicks(TArray<FLoopNode*>& NodesOfLoop, TArray<TPair<double, double>>& Intersections)
-{
-	for (FLoopNode* Node : NodesOfLoop)
-	{
-		if (Node == nullptr || Node->IsDelete())
-		{
-			continue;
-		}
-		RemovePickRecursively(Node, &Node->GetNextNode());
-	}
-	IsoTriangulatorImpl::RemoveDeletedNodes(NodesOfLoop);
-}
+#ifdef UNUSED
 
 bool FIsoTriangulator::RemovePickToOutside(const TArray<FLoopNode*>& NodesOfLoop, const TPair<double, double>& Intersection, const TPair<double, double>& NextIntersection, bool bForward)
 {
@@ -480,10 +51,15 @@ bool FIsoTriangulator::RemovePickToOutside(const TArray<FLoopNode*>& NodesOfLoop
 	};
 
 	FLoopNode* TmpNode = GetNodeAt(NodesOfLoop, NextIndex(NodeCount, (int32)Intersection.Value));
+	FLoopNode* EndNode = GetNodeAt(NodesOfLoop, (int32)NextIntersection.Value);
+	if (TmpNode == nullptr || TmpNode->IsDelete() || EndNode == nullptr || EndNode->IsDelete())
+	{
+		return false;
+	}
+
 	FLoopNode* StartNode = GetPrevious(TmpNode);
 	FPoint2D FirstIntersection = IntersectingPoint(Intersection.Value, StartNode, TmpNode);
 
-	FLoopNode* EndNode = GetNodeAt(NodesOfLoop, (int32)NextIntersection.Value);
 	TmpNode = GetNext(EndNode);
 	FPoint2D SecondIntersection = IntersectingPoint(NextIntersection.Value, EndNode, TmpNode);
 
@@ -492,7 +68,10 @@ bool FIsoTriangulator::RemovePickToOutside(const TArray<FLoopNode*>& NodesOfLoop
 	TmpNode = GetNext(StartNode);
 	while (TmpNode && (TmpNode != EndNode) && !TmpNode->IsDelete())
 	{
-		RemoveNodeOfLoop(*TmpNode);
+		if (!RemoveNodeOfLoop(*TmpNode))
+		{
+			return false;
+		}
 		TmpNode = GetNext(StartNode);
 	}
 
@@ -501,6 +80,10 @@ bool FIsoTriangulator::RemovePickToOutside(const TArray<FLoopNode*>& NodesOfLoop
 	if (FMath::IsNearlyZero(Length))
 	{
 		FLoopNode* StartSegment = GetNodeAt(NodesOfLoop, (int32)Intersection.Key);
+		if (StartSegment == nullptr)
+		{
+			return false;
+		}
 		MoveDirection = GetNext(StartSegment)->Get2DPoint(EGridSpace::UniformScaled, Grid) - StartSegment->Get2DPoint(EGridSpace::UniformScaled, Grid);
 		Length = MoveDirection.Length();
 	}
@@ -518,276 +101,7 @@ bool FIsoTriangulator::RemovePickToOutside(const TArray<FLoopNode*>& NodesOfLoop
 	return true;
 }
 
-void FIsoTriangulator::RemoveSubLoop(FLoopNode* StartNode, FLoopNode* EndNode, IsoTriangulatorImpl::GetNextNodeMethod NextNode)
-{
-	FLoopNode* Node = NextNode(StartNode);
-	while (Node && (Node != EndNode) && !Node->IsDelete())
-	{
-		RemoveNodeOfLoop(*Node);
-		Node = NextNode(StartNode);
-	}
-
-	//Display(EGridSpace::UniformScaled, TEXT("Loop"), LoopSegments, true);
-	//Wait();
-}
-
-void FIsoTriangulator::MoveIntersectingSectionBehindOppositeSection(IsoTriangulatorImpl::FLoopSection IntersectingSection, IsoTriangulatorImpl::FLoopSection OppositeSection, IsoTriangulatorImpl::GetNextNodeMethod GetNext, IsoTriangulatorImpl::GetNextNodeMethod GetPrevious)
-{
-	using namespace IsoTriangulatorImpl;
-
-	FLoopNode* FirstNodeIntersectingSection = IntersectingSection.Key;
-	FLoopNode* LastNodeIntersectingSection = IntersectingSection.Value;
-
-	FLoopNode* FirstNodeOppositeSection = OppositeSection.Key;
-	FLoopNode* LastNodeOppositeSection = OppositeSection.Value;
-
-	int32 OppositeSectionCount = 1;
-	for (FLoopNode* SegmentNode = FirstNodeOppositeSection; SegmentNode != LastNodeOppositeSection; SegmentNode = GetNext(SegmentNode), ++OppositeSectionCount);
-
-	TArray<const FPoint2D*> OppositeSectionPoint;
-	OppositeSectionPoint.Reserve(OppositeSectionCount);
-
-	for (FLoopNode* SegmentNode = FirstNodeOppositeSection; SegmentNode != LastNodeOppositeSection; SegmentNode = GetNext(SegmentNode))
-	{
-		OppositeSectionPoint.Add(&SegmentNode->Get2DPoint(EGridSpace::UniformScaled, Grid));
-	}
-	OppositeSectionPoint.Add(&LastNodeOppositeSection->Get2DPoint(EGridSpace::UniformScaled, Grid));
-
-#ifdef DEBUG_MOVE_INTERSECTING_SECTION_BEHIND_OPPOSITE_SECTION
-	{
-		F3DDebugSession _(TEXT("IntersectingSection"));
-		for (FLoopNode* Node = FirstNodeIntersectingSection; Node != GetNext(LastNodeIntersectingSection); Node = GetNext(Node))
-		{
-			DisplayPoint(Node->Get2DPoint(EGridSpace::UniformScaled, Grid), EVisuProperty::PurplePoint);
-		}
-}
-	{
-		F3DDebugSession _(TEXT("OppositeSection"));
-		for (FLoopNode* Node = FirstNodeOppositeSection; Node != GetNext(LastNodeOppositeSection); Node = GetNext(Node))
-		{
-			DisplayPoint(Node->Get2DPoint(EGridSpace::UniformScaled, Grid), EVisuProperty::YellowPoint);
-		}
-		Wait();
-	}
-#endif
-
-	double Coordinate;
-	FLoopNode* NextNode = nullptr;
-	for (FLoopNode* Node = GetNext(FirstNodeIntersectingSection); Node != LastNodeIntersectingSection;)
-	{
-		FLoopNode* NodeToProceed = Node;
-		Node = GetNext(Node);
-
-		FPoint2D  CandiatePosition;
-		double MinSquareDistance = HUGE_VALUE;
-		FPoint2D PointToProject = NodeToProceed->Get2DPoint(EGridSpace::UniformScaled, Grid);
-		for (int32 Index = 1; Index < OppositeSectionCount; ++Index)
-		{
-			FPoint2D ProjectedPoint = ProjectPointOnSegment(PointToProject, *OppositeSectionPoint[Index - 1], *OppositeSectionPoint[Index], Coordinate);
-			double SquareDistance = PointToProject.SquareDistance(ProjectedPoint);
-			if (SquareDistance < MinSquareDistance)
-			{
-				MinSquareDistance = SquareDistance;
-				CandiatePosition = ProjectedPoint;
-			}
-		}
-		MoveNode(*NodeToProceed, CandiatePosition);
-	}
-
-	for (FLoopNode* Node = GetNext(FirstNodeIntersectingSection); Node != LastNodeIntersectingSection; )
-	{
-		FLoopNode* NodeToProceed = Node;
-		Node = GetNext(Node);
-		if (CheckAndRemovePick(NodeToProceed->GetPreviousNode().Get2DPoint(EGridSpace::UniformScaled, Grid), NodeToProceed->Get2DPoint(EGridSpace::UniformScaled, Grid), NodeToProceed->GetNextNode().Get2DPoint(EGridSpace::UniformScaled, Grid), *NodeToProceed))
-		{
-			Node = GetPrevious(Node);
-		}
-	}
-}
-
-bool FIsoTriangulator::RemoveIntersectionsOfSubLoop(const TArray<FLoopNode*>& NodesOfLoop, TArray<TPair<double, double>>& LoopIntersections, int32 IntersectionIndex, int32 IntersectionCount, bool bForward)
-{
-	using namespace IsoTriangulatorImpl;
-	GetNextNodeMethod GetNext = bForward ? GetNextNodeImpl : GetPreviousNodeImpl;
-	GetNextNodeMethod GetPrevious = bForward ? GetPreviousNodeImpl : GetNextNodeImpl;
-
-	const int32 NodeCount = NodesOfLoop.Num();
-
-	TFunction<void(FLoopNode*, const FPoint2D&, FPoint2D&)> MoveNodeToProjection = [&](FLoopNode* NodeToProject, const FPoint2D& PointToProject, FPoint2D& ProjectedPoint)
-	{
-		FPoint2D MoveDirection = ProjectedPoint - PointToProject;
-		MoveDirection.Normalize();
-		MoveDirection *= GeometricTolerance;
-		ProjectedPoint += MoveDirection;
-		NodeToProject->Set2DPoint(EGridSpace::UniformScaled, Grid, ProjectedPoint);
-	};
-
-	TFunction<void(FLoopNode*, const FPoint2D&, const FPoint2D&)> ProjectNodeOnSegment = [&](FLoopNode* NodeToProject, const FPoint2D& Point0, const FPoint2D& Point1)
-	{
-		const FPoint2D& PointToProject = NodeToProject->Get2DPoint(EGridSpace::UniformScaled, Grid);
-
-		FPoint2D ProjectedPoint;
-		double Coordinate;
-		ProjectedPoint = ProjectPointOnSegment(PointToProject, Point0, Point1, Coordinate);
-
-		MoveNodeToProjection(NodeToProject, PointToProject, ProjectedPoint);
-	};
-
-	TFunction<void(const int32, const int32, const int32)> ProjectNodesOnSegment = [&](const int32 StartIndex, const int32 EndIndex, const int32 SegmnentEndIndex)
-	{
-		FLoopNode* Node = GetNodeAt(NodesOfLoop, StartIndex);
-		FLoopNode* StopNode = GetNodeAt(NodesOfLoop, EndIndex);
-		FLoopNode* EndSegment = GetNodeAt(NodesOfLoop, SegmnentEndIndex);
-
-		StopNode = GetNext(StopNode);
-		FLoopNode* StartSegment = GetPrevious(EndSegment);
-
-		const FPoint2D& EndPoint = EndSegment->Get2DPoint(EGridSpace::UniformScaled, Grid);
-		const FPoint2D& StartPoint = StartSegment->Get2DPoint(EGridSpace::UniformScaled, Grid);
-
-		for (; Node != StopNode; Node = GetNext(Node))
-		{
-			ProjectNodeOnSegment(Node, StartPoint, EndPoint);
-		}
-
-		for (Node = GetNodeAt(NodesOfLoop, StartIndex); Node != StopNode; )
-		{
-			FLoopNode* NodeToProcess = Node;
-			Node = GetNext(Node);
-			CheckAndRemovePick(NodeToProcess->GetPreviousNode().Get2DPoint(EGridSpace::UniformScaled, Grid), NodeToProcess->Get2DPoint(EGridSpace::UniformScaled, Grid), NodeToProcess->GetNextNode().Get2DPoint(EGridSpace::UniformScaled, Grid), *NodeToProcess);
-		}
-	};
-
-	TFunction<void(FLoopNode*, FLoopNode*)> MoveNodeBehindOther = [&](FLoopNode* NodeToMove, FLoopNode* Node1Side1)
-	{
-		FLoopNode* Node0Side1 = GetPrevious(Node1Side1);
-		const FPoint2D& Point0 = Node0Side1->Get2DPoint(EGridSpace::UniformScaled, Grid);
-		const FPoint2D& Point1 = Node1Side1->Get2DPoint(EGridSpace::UniformScaled, Grid);
-
-		const FPoint2D& PointToMove = NodeToMove->Get2DPoint(EGridSpace::UniformScaled, Grid);
-
-		FPoint2D MoveDirection = Point1 - Point0;
-		MoveDirection.Normalize();
-		MoveDirection = MoveDirection.GetPerpendicularVector();
-		MoveDirection *= GeometricTolerance;
-
-		FPoint2D NewCoordinate = Point1 + MoveDirection;
-		NodeToMove->Set2DPoint(EGridSpace::UniformScaled, Grid, NewCoordinate);
-	};
-
-	for (int32 Index = IntersectionCount - 1; Index >= 0; --Index)
-	{
-		int32 SecondIntersectionIndex = IntersectionIndex + Index;
-		if (Index > 0)
-		{
-			const TPair<double, double>& SecondIntersection = LoopIntersections[SecondIntersectionIndex];
-			const TPair<double, double>& FirstIntersection = LoopIntersections[SecondIntersectionIndex - 1];
-
-			int32 Side0NodeCount = (int32)SecondIntersection.Key - (int32)FirstIntersection.Key;
-			int32 Side1NodeCount = (int32)FirstIntersection.Value - (int32)SecondIntersection.Value;
-
-			int32 IndexSide0 = NextIndex(NodeCount, (int32)FirstIntersection.Key);
-			int32 IndexSide1 = NextIndex(NodeCount, (int32)SecondIntersection.Value);
-
-			if (Side0NodeCount == 0)
-			{
-				ProjectNodesOnSegment(IndexSide1, (int32)FirstIntersection.Value, IndexSide0);
-			}
-			else if (Side1NodeCount == 0)
-			{
-				ProjectNodesOnSegment(IndexSide0, (int32)SecondIntersection.Key, IndexSide1);
-			}
-			else if (Side0NodeCount == 1 && Side1NodeCount == 1)
-			{
-				FLoopNode* NodeSide0 = GetNodeAt(NodesOfLoop, IndexSide0);
-				FLoopNode* NodeSide1 = GetNodeAt(NodesOfLoop, IndexSide1);
-
-				double SlopSide0 = ComputeUnorientedSlope(GetPrevious(NodeSide0)->Get2DPoint(EGridSpace::UniformScaled, Grid), NodeSide0->Get2DPoint(EGridSpace::UniformScaled, Grid), GetNext(NodeSide0)->Get2DPoint(EGridSpace::UniformScaled, Grid));
-				double SlopSide1 = ComputeUnorientedSlope(GetPrevious(NodeSide1)->Get2DPoint(EGridSpace::UniformScaled, Grid), NodeSide1->Get2DPoint(EGridSpace::UniformScaled, Grid), GetNext(NodeSide1)->Get2DPoint(EGridSpace::UniformScaled, Grid));
-
-				if (SlopSide0 < SlopSide1)
-				{
-					MoveNodeBehindOther(NodeSide1, NodeSide0);
-				}
-				else
-				{
-					MoveNodeBehindOther(NodeSide0, NodeSide1);
-				}
-			}
-			else
-			{
-				FLoopSection IntersectingSection;
-				IntersectingSection.Key = GetPrevious(GetNodeAt(NodesOfLoop, IndexSide0));
-				IntersectingSection.Value = GetNext(GetNodeAt(NodesOfLoop, (int32)SecondIntersection.Key));
-
-				FLoopSection OppositeSection;
-				OppositeSection.Key = GetPrevious(GetNodeAt(NodesOfLoop, IndexSide1));
-				OppositeSection.Value = GetNext(GetNodeAt(NodesOfLoop, (int32)FirstIntersection.Value));
-
-				int32 IntersectingSectionCount = (int32)FirstIntersection.Value - IndexSide1;
-				int32 OppositeSectionCount = (int32)SecondIntersection.Key - IndexSide0;
-				if (OppositeSectionCount < IntersectingSectionCount)
-				{
-					Swap(IntersectingSection, OppositeSection);
-				}
-
-				MoveIntersectingSectionBehindOppositeSection(IntersectingSection, OppositeSection, GetNext, GetPrevious);
-
-#ifdef DEBUG_REMOVE_INTERSECTIONS		
-				DisplayLoops(TEXT("RemoveLoopIntersections first step"), false, true);
-#endif
-
-				// Check if there is no more intersection, otherwise fix last intersections by moving node of OppositeSection
-				// As we don't know the existence of intersection, the process is chesk and fix
-				{
-					LoopSegmentsIntersectionTool.Empty(NodesOfLoop.Num());
-					for (FLoopNode* Node = IntersectingSection.Key; Node != IntersectingSection.Value; Node = GetNext(Node))
-					{
-						LoopSegmentsIntersectionTool.AddSegment(*Node->GetSegmentConnectedTo(GetNext(Node)));
-					}
-					}
-
-				for (FLoopNode* Node = OppositeSection.Key; Node != OppositeSection.Value; Node = GetNext(Node))
-				{
-					if (const FIsoSegment* Intersection = LoopSegmentsIntersectionTool.DoesIntersect(*Node, *GetNext(Node)))
-					{
-						RemoveIntersectionByMovingOutsideNodeInside(*Intersection, *GetNext(Node));
-					}
-				}
-
-#ifdef DEBUG_REMOVE_INTERSECTIONS		
-				DisplayLoops(TEXT("RemoveLoopIntersections second step"), false, true);
-#endif
-
-				// Remove the pick if neeeded
-				for (FLoopNode* Node = OppositeSection.Key; Node != OppositeSection.Value;)
-				{
-					FLoopNode* NodeToProceed = Node;
-					Node = GetNext(Node);
-					if (CheckAndRemovePick(NodeToProceed->GetPreviousNode().Get2DPoint(EGridSpace::UniformScaled, Grid), NodeToProceed->Get2DPoint(EGridSpace::UniformScaled, Grid), NodeToProceed->GetNextNode().Get2DPoint(EGridSpace::UniformScaled, Grid), *NodeToProceed))
-					{
-						Node = GetPrevious(Node);
-					}
-				}
-
-#ifdef DEBUG_REMOVE_INTERSECTIONS		
-				DisplayLoops(TEXT("RemoveLoopIntersections after remove pick"), false, true);
-				Wait();
-#endif
-				}
-			--Index;
-			}
-		else
-		{
-			if (!RemoveUniqueIntersection(NodesOfLoop, LoopIntersections[IntersectionIndex], bForward))
-			{
-				return false;
-			}
-		}
-		}
 
-	return true;
-	}
 
 bool FIsoTriangulator::IsSubLoopBiggerThanMainLoop(const TArray<FLoopNode*>& NodesOfLoop, const TPair<double, double>& Intersection, bool bForward)
 {
@@ -820,6 +134,10 @@ bool FIsoTriangulator::IsSubLoopBiggerThanMainLoop(const TArray<FLoopNode*>& Nod
 		// the sub-loop is not yet process so there are not deleted points
 		const FLoopNode* SubLoopStartNode = GetNodeAt(NodesOfLoop, NextIndex(NodeCount, (int32)Intersection.Key));
 		const FLoopNode* SubLoopEndNode = GetNodeAt(NodesOfLoop, (int32)Intersection.Value);
+		if (SubLoopStartNode->IsDelete() || SubLoopEndNode->IsDelete())
+		{
+			return false;
+		}
 
 		const FLoopNode* MainLoopStartNode = GetNext(SubLoopEndNode);
 		const FLoopNode* MainLoopEndNode = GetPrevious(SubLoopStartNode);
@@ -837,73 +155,6 @@ bool FIsoTriangulator::IsSubLoopBiggerThanMainLoop(const TArray<FLoopNode*>& Nod
 	return false;
 }
 
-bool FIsoTriangulator::RemoveUniqueIntersection(const TArray<FLoopNode*>& NodesOfLoop, TPair<double, double> Intersection, bool bForward)
-{
-	const int32 NodesCount = NodesOfLoop.Num();
-	int32 SubLoopSegmentCount = (int32)Intersection.Value - (int32)Intersection.Key;
-	const int32 OtherSubLoopSegmentCount = NodesCount - (int32)Intersection.Value + (int32)Intersection.Key;
-
-	bool bIntersectionKeyIsExtremity = FMath::IsNearlyEqual(Intersection.Key, (int32)Intersection.Key);
-	bool bIntersectionValueIsExtremity = FMath::IsNearlyEqual(Intersection.Value, (int32)Intersection.Value);
-
-	// TODO
-	/*
-	if(bIntersectionKeyIsExtremity && bIntersectionValueIsExtremity)
-	{   //       _   _
-		// Cas:   \./
-		//       _/ \_
-
-		return SpreadCoincidentNodes(NodesOfLoop, Intersection);
-	}
-	else if(bIntersectionKeyIsExtremity || bIntersectionValueIsExtremity)
-	{
-		// Cas:  ________
-		//         _/\_
-		return MovePickBehind(NodesOfLoop, Intersection, bIntersectionKeyIsExtremity);
-	}
-	*/
-
-	if (OtherSubLoopSegmentCount < SubLoopSegmentCount)
-	{
-		Swap(Intersection.Value, Intersection.Key);
-		SubLoopSegmentCount = OtherSubLoopSegmentCount;
-		bNeedCheckOrientation = true;
-	}
-
-#ifdef DEBUG_REMOVE_UNIQUE_INTERSECTION
-	if (bDisplay)
-	{
-		DisplayLoops(TEXT("before removed"));
-		Wait();
-	}
-#endif
-
-	switch (SubLoopSegmentCount)
-	{
-	case 0:
-		ensureCADKernel(false);
-		break;
-	case 1:
-	case 2:
-		RemoveThePick(NodesOfLoop, Intersection, bForward);
-		break;
-	case 3:
-		SwapNodes(NodesOfLoop, Intersection, bForward);
-		break;
-	default:
-		TryToSwapSegmentsOrRemoveLoop(NodesOfLoop, Intersection, bForward);
-		break;
-	}
-
-#ifdef DEBUG_REMOVE_UNIQUE_INTERSECTION
-	if (bDisplay)
-	{
-		DisplayLoops(TEXT("After removed"));
-		Wait();
-	}
-#endif
-	return true;
-};
 
 bool FIsoTriangulator::SpreadCoincidentNodes(const TArray<FLoopNode*>& NodesOfLoop, TPair<double, double> Intersection)
 {
@@ -934,240 +185,12 @@ bool FIsoTriangulator::MovePickBehind(const TArray<FLoopNode*>& NodesOfLoop, TPa
 }
 
 
-void FIsoTriangulator::RemoveIntersectionByMovingOutsideNodeInside(const FIsoSegment& IntersectingSegment, FLoopNode& NodeToMove)
-{
-	const FIsoNode* Nodes[2] = { nullptr, nullptr };
-	Nodes[0] = &IntersectingSegment.GetFirstNode();
-	Nodes[1] = &IntersectingSegment.GetSecondNode();
-
-	FPoint2D IntersectingPoints[2];
-	IntersectingPoints[0] = IntersectingSegment.GetFirstNode().Get2DPoint(EGridSpace::UniformScaled, Grid);
-	IntersectingPoints[1] = IntersectingSegment.GetSecondNode().Get2DPoint(EGridSpace::UniformScaled, Grid);
-
-	bool bEndNodeIsOutside = true;
-	FPoint2D PointToMove = NodeToMove.Get2DPoint(EGridSpace::UniformScaled, Grid);
-
-#ifdef DEBUG_CLOSED_OUSIDE_POINT
-	if (bDisplay)
-	{
-		F3DDebugSession _(TEXT("Outside Point"));
-		DisplayPoint(PointToMove, bEndNodeIsOutside ? EVisuProperty::GreenPoint : EVisuProperty::YellowPoint);
-	}
-#endif
-
-	double Coordinate;
-	FPoint2D ProjectedPoint = ProjectPointOnSegment(PointToMove, IntersectingPoints[0], IntersectingPoints[1], Coordinate);
-	MoveNode(NodeToMove, ProjectedPoint);
-}
-
-void FIsoTriangulator::MoveNode(FLoopNode& NodeToMove, FPoint2D& ProjectedPoint)
-{
-	const FPoint2D& PointToMove = NodeToMove.Get2DPoint(EGridSpace::UniformScaled, Grid);
-
-	FPoint2D ProjectedDirection = ProjectedPoint - PointToMove;
-	ProjectedDirection.Normalize();
-	ProjectedDirection *= GeometricTolerance;
-
-	FPoint2D NewCoordinate = ProjectedPoint + ProjectedDirection;
-
-	FLoopNode& PreviousNode = NodeToMove.GetPreviousNode();
-	FLoopNode& NextNode = NodeToMove.GetNextNode();
-
-	if ((PreviousNode.Get2DPoint(EGridSpace::UniformScaled, Grid).SquareDistance(NewCoordinate) < SquareGeometricTolerance2) ||
-		(NextNode.Get2DPoint(EGridSpace::UniformScaled, Grid).SquareDistance(NewCoordinate) < SquareGeometricTolerance2))
-	{
-		RemoveNodeOfLoop(NodeToMove);
-	}
-	else
-	{
-		NodeToMove.Set2DPoint(EGridSpace::UniformScaled, Grid, NewCoordinate);
-	}
-
-#ifdef DEBUG_MOVE_NODE
-	if (bDisplay)
-	{
-		{
-			F3DDebugSession _(TEXT("Point To Move"));
-			DisplayPoint(PointToMove, EVisuProperty::YellowPoint);
-}
-		{
-			F3DDebugSession _(TEXT("Projected Point"));
-			DisplayPoint(ProjectedPoint, EVisuProperty::GreenPoint);
-		}
-		{
-			F3DDebugSession _(TEXT("New Position"));
-			DisplayPoint(NewCoordinate, EVisuProperty::BluePoint);
-			Wait(false);
-		}
-	}
-#endif
-}
-
-void FIsoTriangulator::RemoveThePick(const TArray<FLoopNode*>& NodesOfLoop, const TPair<double, double>& Intersection, bool bForward)
-{
-	using namespace IsoTriangulatorImpl;
-	GetNextNodeMethod GetNext = bForward ? GetNextNodeImpl : GetPreviousNodeImpl;
-
-	FLoopNode* PreviousNode = GetNodeAt(NodesOfLoop, (int32)Intersection.Key);
-	FLoopNode* Node0 = GetNext(PreviousNode);
-	FLoopNode* Node1 = GetNext(Node0);
-	FLoopNode* NextNode = GetNext(Node1);
-
-	const FPoint2D* PreviousPoint = &PreviousNode->Get2DPoint(EGridSpace::UniformScaled, Grid);
-	const FPoint2D* Point0 = &Node0->Get2DPoint(EGridSpace::UniformScaled, Grid);
-	const FPoint2D* Point1 = &Node1->Get2DPoint(EGridSpace::UniformScaled, Grid);
-	const FPoint2D* NextPoint = &NextNode->Get2DPoint(EGridSpace::UniformScaled, Grid);
-
-	double Slop0 = ComputeUnorientedSlope(*Point0, *PreviousPoint, *Point1);
-	double Slop1 = ComputeUnorientedSlope(*Point1, *PreviousPoint, *NextPoint);
 
-	if (Slop0 < Slop1)
-	{
-		RemoveNodeOfLoop(*Node0);
-		Node0 = PreviousNode;
-	}
-	else
-	{
-		RemoveNodeOfLoop(*Node1);
-		Node1 = NextNode;
-	}
-};
 
-bool FIsoTriangulator::RemovePickRecursively(FLoopNode* Node0, FLoopNode* Node1)
-{
-	FLoopNode* PreviousNode = &Node0->GetPreviousNode();
-	FLoopNode* NextNode = &Node1->GetNextNode();
-
-	const FPoint2D* PreviousPoint = &PreviousNode->Get2DPoint(EGridSpace::UniformScaled, Grid);
-	const FPoint2D* Point0 = &Node0->Get2DPoint(EGridSpace::UniformScaled, Grid);
-	const FPoint2D* Point1 = &Node1->Get2DPoint(EGridSpace::UniformScaled, Grid);
-	const FPoint2D* NextPoint = &NextNode->Get2DPoint(EGridSpace::UniformScaled, Grid);
-
-	bool bNodeRemoved = true;
-	while (true)
-	{
-		bool bNodeIsToDelete = (Point0->SquareDistance(*Point1) < SquareGeometricTolerance2);
-		if (bNodeIsToDelete)
-		{
-			RemoveNodeOfLoop(*Node0);
-		}
-
-		if (bNodeIsToDelete || CheckAndRemovePick(*PreviousPoint, *Point0, *Point1, *Node0))
-		{
-			if (PreviousNode->IsDelete())
-			{
-				break;
-			}
-
-			Point0 = PreviousPoint;
-			Node0 = PreviousNode;
-			PreviousNode = &PreviousNode->GetPreviousNode();
-			PreviousPoint = &PreviousNode->Get2DPoint(EGridSpace::UniformScaled, Grid);
-			continue;
-		}
-
-		if (CheckAndRemovePick(*Point0, *Point1, *NextPoint, *Node1))
-		{
-			if (NextNode->IsDelete())
-			{
-				break;
-			}
-
-			Point1 = NextPoint;
-			Node1 = NextNode;
-			NextNode = &NextNode->GetNextNode();
-			NextPoint = &NextNode->Get2DPoint(EGridSpace::UniformScaled, Grid);
-			continue;
-		}
-		break;
-	}
-	return bNodeRemoved;
-}
-
-
-// Case  ______a  c       ______a__b
-//              \/ \o	            \o
-//       ______d/\b/	  ______d__c/
-//	    			  
-void FIsoTriangulator::SwapNodes(const TArray<FLoopNode*>& NodesOfLoop, const TPair<double, double>& Intersection, bool bForward)
-{
-	using namespace IsoTriangulatorImpl;
-	GetNextNodeMethod GetNext = bForward ? GetNextNodeImpl : GetPreviousNodeImpl;
-
-	const int32 NodeCount = NodesOfLoop.Num();
-
-	FLoopNode* Node0 = GetNodeAt(NodesOfLoop, NextIndex(NodeCount, (int32)Intersection.Key));
-	FLoopNode* Pick = GetNext(Node0);
-	FLoopNode* Node1 = GetNext(Pick);
-
-	const FPoint2D Point0Copy = Node0->Get2DPoint(EGridSpace::UniformScaled, Grid);
-	const FPoint2D* Point1Ptr = &Node1->Get2DPoint(EGridSpace::UniformScaled, Grid);
-
-	Node0->Set2DPoint(EGridSpace::UniformScaled, Grid, *Point1Ptr);
-	Node1->Set2DPoint(EGridSpace::UniformScaled, Grid, Point0Copy);
-};
-
-void FIsoTriangulator::TryToSwapSegmentsOrRemoveLoop(const TArray<FLoopNode*>& NodesOfLoop, const TPair<double, double>& Intersection, bool bForward)
-{
-	using namespace IsoTriangulatorImpl;
-	const GetNextNodeMethod GetNext = bForward ? GetNextNodeImpl : GetPreviousNodeImpl;
-	const GetNextNodeMethod GetPrevious = bForward ? GetPreviousNodeImpl : GetNextNodeImpl;
-
-	const int32 NodeCount = NodesOfLoop.Num();
-
-	const int32 Segment0StartIndex = NextIndex(NodeCount, (int32)Intersection.Key);
-	const int32 Segment1EndIndex = (int32)Intersection.Value;
-
-	FLoopNode* Segment0_Node1 = GetNodeAt(NodesOfLoop, Segment0StartIndex);
-	FLoopNode* Segment1_Node0 = GetNodeAt(NodesOfLoop, Segment1EndIndex);
-
-	FLoopNode* Segment0_Node0 = GetPrevious(Segment0_Node1);
-	FLoopNode* Segment1_Node1 = GetNext(Segment1_Node0);
-
-	const FPoint2D& Segment0_Point0 = Segment0_Node0->Get2DPoint(EGridSpace::UniformScaled, Grid);
-	const FPoint2D& Segment0_Point1 = Segment0_Node1->Get2DPoint(EGridSpace::UniformScaled, Grid);
-	const FPoint2D& Segment1_Point0 = Segment1_Node0->Get2DPoint(EGridSpace::UniformScaled, Grid);
-	const FPoint2D& Segment1_Point1 = Segment1_Node1->Get2DPoint(EGridSpace::UniformScaled, Grid);
-
-#ifdef DEBUG_SWAP_SEGMENTS_OR_REMOVE
-	if (bDisplay)
-	{
-		F3DDebugSession _(TEXT("Intersected Segments"));
-		Display(EGridSpace::UniformScaled, *Segment0_Node0, *Segment0_Node1, 0, EVisuProperty::RedCurve);
-		Display(EGridSpace::UniformScaled, *Segment1_Node0, *Segment1_Node1, 0, EVisuProperty::RedCurve);
-		Display(EGridSpace::UniformScaled, *Segment0_Node0, 0, EVisuProperty::RedPoint);
-		Display(EGridSpace::UniformScaled, *Segment1_Node0, 0, EVisuProperty::RedPoint);
-		Wait(false);
-	}
-#endif
 
-	double Slop = ComputeSlope(Segment0_Point0, Segment0_Point1);
-	Slop = ComputeUnorientedSlope(Segment1_Point0, Segment1_Point1, Slop);
 
-	bool bIsFixed = false;
-	if (Slop < 2)
-	{
-		FIsoSegment* Segment0 = Segment0_Node0->GetSegmentConnectedTo(Segment0_Node1);
-		FIsoSegment* Segment1 = Segment1_Node0->GetSegmentConnectedTo(Segment1_Node1);
-		bIsFixed = TryToRemoveIntersectionBySwappingSegments(*Segment0, *Segment1);
 
-#ifdef DEBUG_SWAP_SEGMENTS_OR_REMOVE
-		if (bDisplay)
-		{
-			F3DDebugSession _(TEXT("New Segments"));
-			Display(EGridSpace::UniformScaled, *Segment0_Node0, Segment0_Node0->GetNextNode(), 0, EVisuProperty::BlueCurve);
-			Display(EGridSpace::UniformScaled, *Segment1_Node1, Segment1_Node1->GetPreviousNode(), 0, EVisuProperty::BlueCurve);
-			Display(EGridSpace::UniformScaled, *Segment0_Node0, 0, EVisuProperty::BluePoint);
-			Display(EGridSpace::UniformScaled, *Segment1_Node0, 0, EVisuProperty::BluePoint);
-		}
-#endif
-}
 
-	if (!bIsFixed)
-	{
-		RemoveSubLoop(Segment0_Node0, Segment1_Node1, GetNext);
-	}
-};
 
 
 void FIsoTriangulator::FillIntersectionToolWithOuterLoop()
@@ -1178,44 +201,17 @@ void FIsoTriangulator::FillIntersectionToolWithOuterLoop()
 		{
 			break;
 		}
-		LoopSegmentsIntersectionTool.AddSegment(*Node.GetSegmentConnectedTo(&Node.GetNextNode()));
+
+		FIsoSegment* Segment = Node.GetSegmentConnectedTo(&Node.GetNextNode());
+		if(Segment != nullptr)
+		{
+			LoopSegmentsIntersectionTool.AddSegment(*Segment);
+		}
 	}
 }
 
 bool FIsoTriangulator::FindLoopIntersectionAndFixIt()
 {
-	// move points for test
-	{
-		//FPoint2D NewCoordinate = LoopNodes[56].Get2DPoint(EGridSpace::UniformScaled, Grid);
-		//FPoint2D NewCoordinate53 = LoopNodes[55].Get2DPoint(EGridSpace::UniformScaled, Grid);
-		//NewCoordinate += FPoint(2, 0);
-		//const_cast<FLoopNode&>(LoopNodes[56]).Set2DPoint(EGridSpace::UniformScaled, Grid, NewCoordinate);
-
-		//NewCoordinate += FPoint(0.3, 1.0);
-		//const_cast<FLoopNode&>(LoopNodes[55]).Set2DPoint(EGridSpace::UniformScaled, Grid, NewCoordinate);
-		//NewCoordinate += FPoint(-0.1, -0.3);
-		//const_cast<FLoopNode&>(LoopNodes[54]).Set2DPoint(EGridSpace::UniformScaled, Grid, NewCoordinate);
-		//const_cast<FLoopNode&>(LoopNodes[53]).Set2DPoint(EGridSpace::UniformScaled, Grid, NewCoordinate53);
-
-		//FPoint2D NewCoordinate42 = LoopNodes[43].Get2DPoint(EGridSpace::UniformScaled, Grid);
-		//NewCoordinate42 += FPoint(0, -0.1);
-		//const_cast<FLoopNode&>(LoopNodes[42]).Set2DPoint(EGridSpace::UniformScaled, Grid, NewCoordinate42);
-
-
-		//NewCoordinate = LoopNodes[53].Get2DPoint(EGridSpace::UniformScaled, Grid);
-		//NewCoordinate += FPoint(2, 0);
-		//const_cast<FLoopNode&>(LoopNodes[53]).Set2DPoint(EGridSpace::UniformScaled, Grid, NewCoordinate);
-
-
-		//NewCoordinate = LoopNodes[49].Get2DPoint(EGridSpace::UniformScaled, Grid);
-		//NewCoordinate += FPoint(-0.05, -1);
-		//const_cast<FLoopNode&>(LoopNodes[49]).Set2DPoint(EGridSpace::UniformScaled, Grid, NewCoordinate);
-
-		//FPoint2D NewCoordinate = LoopNodes[126].Get2DPoint(EGridSpace::UniformScaled, Grid);
-		//NewCoordinate += FPoint(0, 3);
-		//const_cast<FLoopNode&>(LoopNodes[126]).Set2DPoint(EGridSpace::UniformScaled, Grid, NewCoordinate);
-	}
-
 	TArray<FLoopNode*> BestStartNodeOfLoop;
 	FindBestLoopExtremity(BestStartNodeOfLoop);
 
@@ -1270,7 +266,11 @@ bool FIsoTriangulator::FindLoopIntersectionAndFixIt()
 		FindLoopIntersections(LoopNodesFromStartNode, bIsAnOuterLoop, Intersections);
 
 		// WARNING From this step, RemoveLoopIntersections can delete nodes. 
-		RemoveLoopIntersections(LoopNodesFromStartNode, Intersections, bIsAnOuterLoop);
+		if (!RemoveLoopIntersections(LoopNodesFromStartNode, Intersections, bIsAnOuterLoop))
+		{
+			FMessage::Printf(Log, TEXT("Loop intersections of the surface %d cannot be fixed. The mesh of this surface is canceled.\n"), Grid.GetFace()->GetId());
+			return false;
+		}
 
 #ifdef DEBUG_LOOP_INTERSECTION_AND_FIX_IT		
 		DisplayLoop(EGridSpace::UniformScaled, TEXT("LoopIntersections: remove self intersection"), LoopNodesFromStartNode, true, EVisuProperty::YellowPoint);
@@ -1292,7 +292,7 @@ bool FIsoTriangulator::FindLoopIntersectionAndFixIt()
 		FixLoopOrientation(LoopNodesFromStartNode);
 
 		bIsAnOuterLoop = false;
-		}
+	}
 
 	if (!CheckMainLoopConsistency())
 	{
@@ -1336,158 +336,12 @@ bool FIsoTriangulator::FindLoopIntersectionAndFixIt()
 #endif
 
 	return true;
-	}
-
-void FIsoTriangulator::FixLoopOrientation(const TArray<FLoopNode*>& NodesOfLoop)
-{
-	FLoopNode const* const* StartNode = NodesOfLoop.FindByPredicate([](const FLoopNode* Node) { return !Node->IsDelete(); });
-	if (StartNode != nullptr)
-	{
-		EOrientation Orientation = GetLoopOrientation(*StartNode);
-		if (Orientation == EOrientation::Back)
-		{
-			const int32 LoopIndex = (*StartNode)->GetLoopIndex();
-			const int32 FirstSegmentIndex = LoopSegments.IndexOfByPredicate([&](const FIsoSegment* Segment) { return ((const FLoopNode&)Segment->GetFirstNode()).GetLoopIndex() == LoopIndex; });
-			int32 LastSegmentIndex = (Grid.GetLoopCount() == LoopIndex + 1) ? LoopSegments.Num() : LoopSegments.IndexOfByPredicate([&](const FIsoSegment* Segment) { return ((const FLoopNode&)Segment->GetFirstNode()).GetLoopIndex() > LoopIndex; });
-			if (LastSegmentIndex == INDEX_NONE)
-			{
-				LastSegmentIndex = LoopSegments.Num();
-			}
-			SwapLoopOrientation(FirstSegmentIndex, LastSegmentIndex);
-#ifdef DEBUG_LOOP_ORIENTATION
-			Display(EGridSpace::UniformScaled, TEXT("After orientation"), LoopSegments, true, true);
-			DisplayLoops(TEXT("After orientation"), false, true, true);
-			Wait();
-#endif
-}
-	}
-}
-
-bool FIsoTriangulator::CheckMainLoopConsistency()
-{
-	int32 OuterNodeCount = 0;
-	for (const FLoopNode& Node : LoopNodes)
-	{
-		if (Node.GetLoopIndex() != 0)
-		{
-			break;
-		}
-
-		if (!Node.IsDelete())
-		{
-			++OuterNodeCount;
-			if (OuterNodeCount > 2)
-			{
-				return true;
-			}
-		}
-	}
-	return false;
 }
 
 //#define DEBUG_CLOSED_OUSIDE_POINT
-void FIsoTriangulator::RemoveIntersectionByMovingOutsideSegmentNodeInside(const FIsoSegment& IntersectingSegment, const FIsoSegment& Segment)
-{
-	const FIsoNode* Nodes[2][2] = { {nullptr, nullptr}, {nullptr, nullptr} };
-	Nodes[0][0] = &IntersectingSegment.GetFirstNode();
-	Nodes[0][1] = &IntersectingSegment.GetSecondNode();
-	Nodes[1][0] = &Segment.GetFirstNode();
-	Nodes[1][1] = &Segment.GetSecondNode();
-
-	FPoint2D IntersectingPoints[2];
-	IntersectingPoints[0] = IntersectingSegment.GetFirstNode().Get2DPoint(EGridSpace::UniformScaled, Grid);
-	IntersectingPoints[1] = IntersectingSegment.GetSecondNode().Get2DPoint(EGridSpace::UniformScaled, Grid);
-
-	FPoint2D SegmentPoints[2];
-	SegmentPoints[0] = Segment.GetFirstNode().Get2DPoint(EGridSpace::UniformScaled, Grid);
-	SegmentPoints[1] = Segment.GetSecondNode().Get2DPoint(EGridSpace::UniformScaled, Grid);
-
-	double SquareLenghtIntersectingSegment = IntersectingPoints[0].SquareDistance(IntersectingPoints[1]);
-	double SquareLenghtSegment = SegmentPoints[0].SquareDistance(SegmentPoints[1]);
-
-	if (SquareLenghtSegment > 10 * SquareLenghtIntersectingSegment)
-	{
-		return RemoveIntersectionByMovingOutsideSegmentNodeInside(Segment, const_cast<FIsoSegment&>(IntersectingSegment));
-	}
-
-	bool bFirstNodeIsOutside = false;
-	FPoint2D PointToMove = Segment.GetSecondNode().Get2DPoint(EGridSpace::UniformScaled, Grid);
-
-	// Is Second node, the outside node ?
-	double OrientedSlop = ComputeOrientedSlope(IntersectingPoints[0], IntersectingPoints[1], PointToMove);
-	if (OrientedSlop > 0)
-	{
-		PointToMove = Segment.GetFirstNode().Get2DPoint(EGridSpace::UniformScaled, Grid);
-		bFirstNodeIsOutside = true;
-	}
 
-#ifdef DEBUG_CLOSED_OUSIDE_POINT
-	if (bDisplay)
-	{
-		F3DDebugSession _(TEXT("Outside Point"));
-		DisplayPoint(PointToMove, bFirstNodeIsOutside ? EVisuProperty::GreenPoint : EVisuProperty::YellowPoint);
-	}
-#endif
-
-	double Coordinate;
-	FPoint2D ProjectedPoint = ProjectPointOnSegment(PointToMove, IntersectingPoints[0], IntersectingPoints[1], Coordinate);
-
-	FPoint2D MoveDirection = IntersectingPoints[1] - IntersectingPoints[0];
-	MoveDirection.Normalize();
-	MoveDirection = MoveDirection.GetPerpendicularVector();
-	MoveDirection *= GeometricTolerance;
-
-	FPoint2D NewCoordinate = ProjectedPoint + MoveDirection;
-
-	FLoopNode& Node = (FLoopNode&) const_cast<FIsoNode&>(!bFirstNodeIsOutside ? Segment.GetSecondNode() : Segment.GetFirstNode());
-	FLoopNode& PreviousNode = Node.GetPreviousNode();
-	FLoopNode& NextNode = Node.GetNextNode();
-
-	if ((PreviousNode.Get2DPoint(EGridSpace::UniformScaled, Grid).SquareDistance(NewCoordinate) < SquareGeometricTolerance2) ||
-		(NextNode.Get2DPoint(EGridSpace::UniformScaled, Grid).SquareDistance(NewCoordinate) < SquareGeometricTolerance2))
-	{
-		RemoveNodeOfLoop(Node);
-		return;
-}
-
-	Node.Set2DPoint(EGridSpace::UniformScaled, Grid, NewCoordinate);
-
-#ifdef DEBUG_CLOSED_OUSIDE_POINT
-	if (bDisplay)
-	{
-		F3DDebugSession _(TEXT("New Segs"));
-		DisplayPoint(NewCoordinate, EVisuProperty::BluePoint);
-		Display(EGridSpace::UniformScaled, Segment, 0, EVisuProperty::BlueCurve);
-		Display(EGridSpace::UniformScaled, IntersectingSegment, 0, EVisuProperty::RedCurve);
-	}
-#endif
-}
-
-
-void FIsoTriangulator::OffsetSegment(FIsoSegment& Segment, TSegment<FPoint2D>& Segment2D, TSegment<FPoint2D>& IntersectingSegment2D)
-{
-	FPoint2D SegmentTangent = IntersectingSegment2D.Point1 - IntersectingSegment2D.Point0;
-	SegmentTangent.Normalize();
-	FPoint2D MoveDirection = SegmentTangent.GetPerpendicularVector();
-	MoveDirection *= GeometricTolerance;
-
-	FPoint2D NewPoint0 = Segment2D.Point0 + MoveDirection;
-	FPoint2D NewPoint1 = Segment2D.Point1 + MoveDirection;
-
-	Segment.GetFirstNode().Set2DPoint(EGridSpace::UniformScaled, Grid, NewPoint0);
-	Segment.GetSecondNode().Set2DPoint(EGridSpace::UniformScaled, Grid, NewPoint1);
-}
 
-void FIsoTriangulator::OffsetNode(FLoopNode& Node, TSegment<FPoint2D>& IntersectingSegment2D)
-{
-	FPoint2D SegmentTangent = IntersectingSegment2D.Point1 - IntersectingSegment2D.Point0;
-	SegmentTangent.Normalize();
-	FPoint2D MoveDirection = SegmentTangent.GetPerpendicularVector();
-	MoveDirection *= GeometricTolerance;
 
-	FPoint2D NewPoint = Node.Get2DPoint(EGridSpace::UniformScaled, Grid) + MoveDirection;
-	Node.Set2DPoint(EGridSpace::UniformScaled, Grid, NewPoint);
-}
 
 //#define DEBUG_TWO_CONSECUTIVE_INTERSECTING
 bool FIsoTriangulator::TryToRemoveIntersectionOfTwoConsecutiveIntersectingSegments(const FIsoSegment& IntersectingSegment, FIsoSegment& Segment)
@@ -1547,7 +401,7 @@ bool FIsoTriangulator::TryToRemoveIntersectionOfTwoConsecutiveIntersectingSegmen
 		Node = (FLoopNode*)&Segment.GetFirstNode();
 		PreviousNode = (FLoopNode*)&Segment.GetSecondNode();
 		NextNode = (FLoopNode*)&Node->GetPreviousNode();
-}
+	}
 
 #ifdef DEBUG_TWO_CONSECUTIVE_INTERSECTING
 	if (bDisplay)
@@ -1626,7 +480,7 @@ bool FIsoTriangulator::TryToRemoveIntersectionOfTwoConsecutiveIntersectingSegmen
 	}
 
 	return true;
-	}
+}
 
 //#define DEBUG_FIND_LOOP_INTERSECTION_AND_FIX_IT
 void FIsoTriangulator::FixIntersectionBetweenLoops()
@@ -1744,7 +598,7 @@ void FIsoTriangulator::FixIntersectionBetweenLoops()
 				Wait(false);
 			}
 #endif
-			}
+		}
 		else
 		{
 			LoopSegmentsIntersectionTool.AddSegment(Segment);
@@ -1770,309 +624,6 @@ void FIsoTriangulator::FixIntersectionBetweenLoops()
 }
 
 //#define DEBUG_SELF_CLOSED_OUSIDE_POINT
-bool FIsoTriangulator::TryToRemoveSelfIntersectionByMovingTheClosedOusidePoint(const FIsoSegment& Segment0, const FIsoSegment& Segment1)
-{
-	const FIsoNode* Nodes[2][2] = { {nullptr, nullptr}, {nullptr, nullptr} };
-	Nodes[0][0] = &Segment0.GetFirstNode();
-	Nodes[0][1] = &Segment0.GetSecondNode();
-	Nodes[1][0] = &Segment1.GetFirstNode();
-	Nodes[1][1] = &Segment1.GetSecondNode();
-
-	FPoint2D Points[2][2];
-	Points[0][0] = Nodes[0][0]->Get2DPoint(EGridSpace::UniformScaled, Grid);
-	Points[0][1] = Nodes[0][1]->Get2DPoint(EGridSpace::UniformScaled, Grid);
-	Points[1][0] = Nodes[1][0]->Get2DPoint(EGridSpace::UniformScaled, Grid);
-	Points[1][1] = Nodes[1][1]->Get2DPoint(EGridSpace::UniformScaled, Grid);
-
-	int32 ProjectedNodeIndex[2][2] = { { 0, 0 }, { 0, 0 } };
-	FPoint2D ProjectedPoints[2][2];
-	double Distance[2][2] = { { HUGE_VALUE, HUGE_VALUE }, { HUGE_VALUE, HUGE_VALUE } };
-	FPoint2D ProjectedPointsOut[2][2];
-
-	//int32 OtherSegmentIndex = 0;
-	int32 SegmentIndex = 1;
-	int32 OtherSegmentIndex = 0;
-
-	TFunction<void(int32)> ProjectPoint = [&](int32 OtherSegmentNodeIndex)
-	{
-		double Coordinate;
-		ProjectedPoints[OtherSegmentIndex][OtherSegmentNodeIndex] = ProjectPointOnSegment(Points[OtherSegmentIndex][OtherSegmentNodeIndex], Points[SegmentIndex][0], Points[SegmentIndex][1], Coordinate, false);
-		bool bProjectedPoint1IsInside = (Coordinate >= -SMALL_NUMBER && Coordinate <= 1 + SMALL_NUMBER);
-		if (bProjectedPoint1IsInside)
-		{
-			Distance[OtherSegmentIndex][OtherSegmentNodeIndex] = ProjectedPoints[OtherSegmentIndex][OtherSegmentNodeIndex].SquareDistance(Points[OtherSegmentIndex][OtherSegmentNodeIndex]);
-		}
-
-#ifdef DEBUG_SELF_CLOSED_OUSIDE_POINT
-		if (bDisplay)
-		{
-			{
-				F3DDebugSession _(TEXT("Segs"));
-				DisplaySegment(Points[SegmentIndex][0], Points[SegmentIndex][1], 0, EVisuProperty::BlueCurve);
-				DisplayPoint(Points[OtherSegmentIndex][OtherSegmentNodeIndex], EVisuProperty::RedPoint);
-				DisplayPoint(ProjectedPoints[OtherSegmentIndex][OtherSegmentNodeIndex], EVisuProperty::RedPoint);
-				Wait(false);
-			}
-		}
-#endif
-};
-
-	SegmentIndex = 0;
-	OtherSegmentIndex = 1;
-	ProjectPoint(0);
-	ProjectPoint(1);
-	SegmentIndex = 1;
-	OtherSegmentIndex = 0;
-	ProjectPoint(0);
-	ProjectPoint(1);
-
-	TFunction<void(int32, int32)> MoveNode = [&](int32 MoveSegmentIndex, int32 MovePointIndex)
-	{
-		FPoint2D MoveDirection = ProjectedPoints[MoveSegmentIndex][MovePointIndex] - Points[MoveSegmentIndex][MovePointIndex];
-		MoveDirection.Normalize();
-
-		if (MoveDirection.SquareLength() < 0.5)
-		{
-			MoveDirection = Points[MoveSegmentIndex][MovePointIndex == 0 ? 1 : 0] - ProjectedPoints[MoveSegmentIndex][MovePointIndex];
-			MoveDirection.Normalize();
-		}
-
-		MoveDirection *= GeometricTolerance;
-		FPoint2D NewCoordinate = ProjectedPoints[MoveSegmentIndex][MovePointIndex] + MoveDirection;
-
-		FLoopNode& Node = (FLoopNode&) const_cast<FIsoNode&>(*Nodes[MoveSegmentIndex][MovePointIndex]);
-		FLoopNode& PreviousNode = Node.GetPreviousNode();
-		FLoopNode& NextNode = Node.GetNextNode();
-
-		if ((PreviousNode.Get2DPoint(EGridSpace::UniformScaled, Grid).SquareDistance(NewCoordinate) < SquareGeometricTolerance2) ||
-			(NextNode.Get2DPoint(EGridSpace::UniformScaled, Grid).SquareDistance(NewCoordinate) < SquareGeometricTolerance2))
-		{
-			RemoveNodeOfLoop(Node);
-			return;
-		}
-
-		const_cast<FIsoNode&>((*Nodes[MoveSegmentIndex][MovePointIndex])).Set2DPoint(EGridSpace::UniformScaled, Grid, NewCoordinate);
-
-#ifdef DEBUG_SELF_CLOSED_OUSIDE_POINT
-		if (bDisplay)
-		{
-			F3DDebugSession _(TEXT("New Segs"));
-			DisplayPoint(NewCoordinate, EVisuProperty::RedPoint);
-			Display(EGridSpace::UniformScaled, Segment1, 0, EVisuProperty::RedCurve);
-			Display(EGridSpace::UniformScaled, Segment0, 0, EVisuProperty::BlueCurve);
-			Wait(false);
-		}
-#endif
-	};
-
-
-	int32 MoveSegment = 0;
-	int32 MovePoint = 0;
-	double MinDistance = Distance[0][0];
-
-	TFunction<void(int32, int32)> FindNodeToMove = [&](int32 MoveSegmentIndex, int32 MovePointIndex)
-	{
-		if (Distance[MoveSegmentIndex][MovePointIndex] < MinDistance)
-		{
-			MinDistance = Distance[MoveSegmentIndex][MovePointIndex];
-			MoveSegment = MoveSegmentIndex;
-			MovePoint = MovePointIndex;
-		}
-	};
-
-	FindNodeToMove(0, 1);
-	FindNodeToMove(1, 0);
-	FindNodeToMove(1, 1);
-
-	MoveNode(MoveSegment, MovePoint);
-
-	return true;
-}
-
-bool FIsoTriangulator::TryToRemoveIntersectionByMovingTheClosedOusidePoint(const FIsoSegment& Segment0, const FIsoSegment& Segment1)
-{
-	// For Inner/outer segment intersecting outer loop (i.e. IntersectingSegment is Outer), the segment has to be moved inside IntersectingSegment
-	// For Inner segment intersecting its loop, the segment has to be moved inside IntersectingSegment but the orientation of the loop is counterclockwise => need to swap orientation
-	// For Inner segment intersecting another inner loop, the segment has to be moved outside the loop but the orientation of the loop is counterclockwise => don't need to swap orientation
-
-	// As the iteration of the segment start from outer loop's segments, Segment can not be from outer loop if IntersectingSegment is not from outer loop
-	int32 SegmentLoopIndex[2];
-	SegmentLoopIndex[0] = Segment0.GetFirstNode().IsALoopNode() ? ((FLoopNode&)Segment0.GetFirstNode()).GetLoopIndex() : 0;
-	SegmentLoopIndex[1] = Segment1.GetFirstNode().IsALoopNode() ? ((FLoopNode&)Segment1.GetFirstNode()).GetLoopIndex() : 0;
-	bool bIsOutLoopSegment = SegmentLoopIndex[0] == 0;
-	bool bIsSameLoop = SegmentLoopIndex[0] == SegmentLoopIndex[1];
-
-	bool bToInside = bIsOutLoopSegment || bIsSameLoop;
-
-	const FIsoNode* Nodes[2][2] = { {nullptr, nullptr}, {nullptr, nullptr} };
-	Nodes[0][0] = &Segment0.GetFirstNode();
-	Nodes[0][1] = &Segment0.GetSecondNode();
-	Nodes[1][0] = &Segment1.GetFirstNode();
-	Nodes[1][1] = &Segment1.GetSecondNode();
-
-	FPoint2D Points[2][2];
-	Points[0][0] = Nodes[0][0]->Get2DPoint(EGridSpace::UniformScaled, Grid);
-	Points[0][1] = Nodes[0][1]->Get2DPoint(EGridSpace::UniformScaled, Grid);
-	Points[1][0] = Nodes[1][0]->Get2DPoint(EGridSpace::UniformScaled, Grid);
-	Points[1][1] = Nodes[1][1]->Get2DPoint(EGridSpace::UniformScaled, Grid);
-
-	int32 ProjectedNodeIndex[2] = { 0, 0 };
-	FPoint2D ProjectedPoints[2];
-	double Distance[2] = { HUGE_VALUE, HUGE_VALUE };
-	int32 OtherSegmentIndex = 0;
-	int32 SegmentIndex = 1;
-
-	TFunction<void(int32)> ProjectInnerPoint = [&](int32 OtherSegmentNodeIndex)
-	{
-		double Coordinate;
-		ProjectedPoints[OtherSegmentIndex] = ProjectPointOnSegment(Points[OtherSegmentIndex][OtherSegmentNodeIndex], Points[SegmentIndex][0], Points[SegmentIndex][1], Coordinate, false);
-		bool bProjectedPoint1IsInside = (Coordinate >= -SMALL_NUMBER && Coordinate <= 1 + SMALL_NUMBER);
-		if (bProjectedPoint1IsInside)
-		{
-			Distance[OtherSegmentIndex] = ProjectedPoints[OtherSegmentIndex].Distance(Points[OtherSegmentIndex][OtherSegmentNodeIndex]);
-		}
-
-#ifdef DEBUG_CLOSED_OUSIDE_POINT
-		if (bDisplay)
-		{
-			{
-				F3DDebugSession _(TEXT("Segs"));
-				DisplaySegment(Points[SegmentIndex][0], Points[SegmentIndex][1], 0, EVisuProperty::BlueCurve);
-				DisplayPoint(Points[OtherSegmentIndex][OtherSegmentNodeIndex], EVisuProperty::RedPoint);
-				DisplayPoint(ProjectedPoints[OtherSegmentIndex], EVisuProperty::RedPoint);
-				Wait(false);
-			}
-		}
-#endif
-};
-
-	TFunction<void()> FindInnerNodeAndProjectIt = [&]()
-	{
-		OtherSegmentIndex = SegmentIndex ? 0 : 1;
-
-		double OrientedSlop = ComputeOrientedSlope(Points[SegmentIndex][0], Points[SegmentIndex][1], Points[OtherSegmentIndex][0]);
-		ProjectedNodeIndex[OtherSegmentIndex] = ((SegmentLoopIndex[0] == 0) == (OrientedSlop < 0)) ? 0 : 1;
-		ProjectInnerPoint(ProjectedNodeIndex[OtherSegmentIndex]);
-	};
-
-	SegmentIndex = 0;
-	FindInnerNodeAndProjectIt();
-	SegmentIndex = 1;
-	FindInnerNodeAndProjectIt();
-
-	if (Distance[0] == HUGE_VALUE && Distance[1] == HUGE_VALUE)
-	{
-		return false;
-	}
-
-	TFunction<void(int32)> MoveNode = [&](int32 MovePointIndex)
-	{
-		FPoint2D MoveDirection = ProjectedPoints[MovePointIndex] - Points[MovePointIndex][ProjectedNodeIndex[MovePointIndex]];
-		MoveDirection.Normalize();
-
-		if (MoveDirection.SquareLength() < 0.5)
-		{
-			MoveDirection = Points[MovePointIndex][ProjectedNodeIndex[MovePointIndex] == 0 ? 1 : 0] - ProjectedPoints[MovePointIndex];
-			MoveDirection.Normalize();
-		}
-
-		MoveDirection *= GeometricTolerance;
-		FPoint2D NewCoordinate = ProjectedPoints[MovePointIndex] + MoveDirection;
-
-		FLoopNode& Node = (FLoopNode&) const_cast<FIsoNode&>(*Nodes[MovePointIndex][ProjectedNodeIndex[MovePointIndex]]);
-		FLoopNode& PreviousNode = Node.GetPreviousNode();
-		FLoopNode& NextNode = Node.GetNextNode();
-
-		if ((PreviousNode.Get2DPoint(EGridSpace::UniformScaled, Grid).SquareDistance(NewCoordinate) < SquareGeometricTolerance2) ||
-			(NextNode.Get2DPoint(EGridSpace::UniformScaled, Grid).SquareDistance(NewCoordinate) < SquareGeometricTolerance2))
-		{
-			RemoveNodeOfLoop(Node);
-			return;
-		}
-
-		const_cast<FIsoNode&>((*Nodes[MovePointIndex][ProjectedNodeIndex[MovePointIndex]])).Set2DPoint(EGridSpace::UniformScaled, Grid, NewCoordinate);
-
-#ifdef DEBUG_CLOSED_OUSIDE_POINT
-		if (bDisplay)
-		{
-			F3DDebugSession _(TEXT("New Segs"));
-			DisplayPoint(NewCoordinate, EVisuProperty::RedPoint);
-			Display(EGridSpace::UniformScaled, Segment1, 0, EVisuProperty::RedCurve);
-			Display(EGridSpace::UniformScaled, Segment0, 0, EVisuProperty::BlueCurve);
-			Wait(false);
-		}
-#endif
-	};
-
-	int32 MovePoint = (Distance[0] < Distance[1]) ? 0 : 1;
-	MoveNode(MovePoint);
-	return true;
-}
-
-bool FIsoTriangulator::RemoveNodeOfLoop(FLoopNode& NodeToRemove)
-{
-	if (NodeToRemove.GetConnectedSegments().Num() != 2)
-	{
-		return false;
-	}
-
-	FLoopNode& PreviousNode = NodeToRemove.GetPreviousNode();
-	FLoopNode& NextNode = NodeToRemove.GetNextNode();
-
-	FIsoSegment* Segment = PreviousNode.GetSegmentConnectedTo(&NodeToRemove);
-	if (Segment == nullptr)
-	{
-		return false;
-	}
-
-	FIsoSegment* SegmentToDelete = NextNode.GetSegmentConnectedTo(&NodeToRemove);
-	if (SegmentToDelete == nullptr)
-	{
-		return false;
-	}
-
-	NextNode.DisconnectSegment(*SegmentToDelete);
-	NextNode.ConnectSegment(*Segment);
-	Segment->SetSecondNode(NextNode);
-
-	if (&NextNode.GetNextNode() == &NextNode.GetPreviousNode())
-	{
-		NextNode.DisconnectSegment(*Segment);
-		PreviousNode.DisconnectSegment(*Segment);
-		LoopSegments.RemoveSingle(Segment);
-		IsoSegmentFactory.DeleteEntity(Segment);
-
-		FIsoSegment* ThirdSegment = PreviousNode.GetSegmentConnectedTo(&NextNode);
-		if (!ThirdSegment)
-		{
-			return false;
-		}
-
-		NextNode.DisconnectSegment(*ThirdSegment);
-		PreviousNode.DisconnectSegment(*ThirdSegment);
-
-		LoopSegments.RemoveSingle(ThirdSegment);
-		IsoSegmentFactory.DeleteEntity(ThirdSegment);
-
-		NextNode.Delete();
-		PreviousNode.Delete();
-
-		LoopSegmentsIntersectionTool.Remove(Segment);
-		LoopSegmentsIntersectionTool.Remove(ThirdSegment);
-	}
-
-	LoopSegments.RemoveSingle(SegmentToDelete);
-	IsoSegmentFactory.DeleteEntity(SegmentToDelete);
-	NodeToRemove.Delete();
-
-	LoopSegmentsIntersectionTool.Remove(SegmentToDelete);
-	if (!Segment->IsDelete())
-	{
-		LoopSegmentsIntersectionTool.Update(Segment);
-	}
-
-	return true;
-}
 
 //#define DEBUG_REMOVE_PICK_OF_LOOP
 void FIsoTriangulator::RemovePickOfLoop(FIsoSegment& Segment)
@@ -2138,64 +689,8 @@ void FIsoTriangulator::RemovePickOfLoop(FIsoSegment& Segment)
 #endif
 }
 
-void FIsoTriangulator::SwapLoopOrientation(int32 FirstSegmentIndex, int32 LastSegmentIndex)
-{
-	TArray<FIsoSegment*> Segments;
-	Segments.Reserve(LastSegmentIndex - FirstSegmentIndex);
-	for (int32 Index = FirstSegmentIndex; Index < LastSegmentIndex; ++Index)
-	{
-		LoopSegments[Index]->SwapOrientation();
-		Segments.Add(LoopSegments[Index]);
-	}
-
-	for (int32 ReverseIndex = Segments.Num() - 1, Index = FirstSegmentIndex; Index < LastSegmentIndex; ++Index, --ReverseIndex)
-	{
-		LoopSegments[Index] = Segments[ReverseIndex];
-	}
-};
-
-bool FIsoTriangulator::TryToRemoveIntersectionBySwappingSegments(FIsoSegment& IntersectingSegment, FIsoSegment& Segment)
-{
-	if (LoopSegmentsIntersectionTool.DoesIntersect(IntersectingSegment.GetFirstNode(), Segment.GetFirstNode()))
-	{
-		return false;
-	}
-
-	if (LoopSegmentsIntersectionTool.DoesIntersect(IntersectingSegment.GetSecondNode(), Segment.GetSecondNode()))
-	{
-		return false;
-	}
-
-	int32 StartSegmentIndex = LoopSegments.IndexOfByKey(&IntersectingSegment);
-	int32 EndSegmentIndex = LoopSegments.IndexOfByKey(&Segment);
-
-	IntersectingSegment.GetSecondNode().DisconnectSegment(IntersectingSegment);
-	Segment.GetFirstNode().DisconnectSegment(Segment);
-
-	FIsoNode& Node = IntersectingSegment.GetSecondNode();
-	IntersectingSegment.SetSecondNode(Segment.GetFirstNode());
-	Segment.SetFirstNode(Node);
-
-	Segment.GetFirstNode().ConnectSegment(Segment);
-	IntersectingSegment.GetSecondNode().ConnectSegment(IntersectingSegment);
 
-#ifdef DEBUG_BY_SWAPPING_SEGMENTS
-	if (bDisplay)
-	{
-		F3DDebugSession _(TEXT("New Segments"));
-		Display(EGridSpace::UniformScaled, Segment, StartSegmentIndex, EVisuProperty::RedCurve/*, true*/);
-		Display(EGridSpace::UniformScaled, IntersectingSegment, EndSegmentIndex, EVisuProperty::RedCurve/*, true*/);
-		//Wait();
-	}
 #endif
 
-	SwapLoopOrientation(StartSegmentIndex + 1, EndSegmentIndex);
-
-	LoopSegmentsIntersectionTool.Update(&IntersectingSegment);
-	LoopSegmentsIntersectionTool.Update(&Segment);
-
-	return true;
-}
-
 
 } //namespace CADKernel
\ No newline at end of file
diff --git a/Engine/Source/Runtime/Datasmith/CADKernel/Private/CADKernel/Mesh/Meshers/ParametricMesher.cpp b/Engine/Source/Runtime/Datasmith/CADKernel/Private/CADKernel/Mesh/Meshers/ParametricMesher.cpp
--- a/Engine/Source/Runtime/Datasmith/CADKernel/Private/CADKernel/Mesh/Meshers/ParametricMesher.cpp
+++ b/Engine/Source/Runtime/Datasmith/CADKernel/Private/CADKernel/Mesh/Meshers/ParametricMesher.cpp
@@ -123,7 +123,7 @@ void FParametricMesher::MeshEntities()
 		ApplyFaceCriteria(Face.ToSharedRef());
 	}
 
-	Chronos.ApplyCriteriaDuration = FChrono::Elapse(ApplyCriteriaStartTime);
+	MesherReport.Chronos.ApplyCriteriaDuration = FChrono::Elapse(ApplyCriteriaStartTime);
 
 	FTimePoint MeshingStartTime = FChrono::Now();
 
@@ -142,7 +142,7 @@ void FParametricMesher::MeshEntities()
 
 		IsolateQuadFace(QuadTrimmedSurfaceSet, OtherEntities);
 
-		Chronos.IsolateQuadPatchDuration = FChrono::Elapse(IsolateQuadPatchStartTime);
+		MesherReport.Chronos.IsolateQuadPatchDuration = FChrono::Elapse(IsolateQuadPatchStartTime);
 		FMessage::Printf(Log, TEXT("  %d Quad Surfaces found\n"), QuadTrimmedSurfaceSet.Num());
 	}
 
@@ -152,8 +152,8 @@ void FParametricMesher::MeshEntities()
 
 	FTimePoint MeshStartTime = FChrono::Now();
 	MeshSurfaceByFront(QuadTrimmedSurfaceSet);
-	Chronos.GlobalMeshDuration = FChrono::Elapse(MeshStartTime);
-	Chronos.GlobalDuration = FChrono::Elapse(StartTime);
+	MesherReport.Chronos.GlobalMeshDuration = FChrono::Elapse(MeshStartTime);
+	MesherReport.Chronos.GlobalDuration = FChrono::Elapse(StartTime);
 
 	//Chronos.PrintTimeElapse();
 }
@@ -217,10 +217,9 @@ void FParametricMesher::Mesh(TSharedRef<FTopologicalFace> Face)
 
 	FProgress _(1, TEXT("Meshing Entities : Mesh Surface"));
 
-	if (BoolDisplayDebugMeshStep)
-	{
-		Open3DDebugSession(FString::Printf(TEXT("Mesh of surface %d"), Face->GetId()));
-	}
+#ifdef DEBUG_CADKERNEL
+	F3DDebugSession S(bDisplayDebugMeshStep, FString::Printf(TEXT("Mesh of surface %d"), Face->GetId()));
+#endif
 
 	FTimePoint StartTime = FChrono::Now();
 
@@ -231,10 +230,6 @@ void FParametricMesher::Mesh(TSharedRef<FTopologicalFace> Face)
 
 	if (Grid.IsDegenerated())
 	{
-		if (BoolDisplayDebugMeshStep)
-		{
-			Close3DDebugSession();
-		}
 		FMessage::Printf(EVerboseLevel::Log, TEXT("The meshing of the surface %d failed due to a degenerated grid\n"), Face->GetId());
 		Face->SetMeshed();
 		return;
@@ -246,7 +241,7 @@ void FParametricMesher::Mesh(TSharedRef<FTopologicalFace> Face)
 
 	TSharedRef<FFaceMesh> SurfaceMesh = StaticCastSharedRef<FFaceMesh>(Face->GetOrCreateMesh(MeshModel));
 
-	FIsoTriangulator IsoTrianguler(Grid, SurfaceMesh);
+	FIsoTriangulator IsoTrianguler(Grid, SurfaceMesh, MesherReport);
 	if (IsoTrianguler.Triangulate())
 	{
 		if (Face->IsBackOriented())
@@ -261,19 +256,15 @@ void FParametricMesher::Mesh(TSharedRef<FTopologicalFace> Face)
 	FDuration Duration = FChrono::Elapse(StartTime);
 
 #ifdef CADKERNEL_DEV
-	Chronos.GlobalPointCloudDuration += Grid.Chronos.GeneratePointCloudDuration;
-	Chronos.GlobalGeneratePointCloudDuration += GenerateCloudDuration;
-	Chronos.GlobalTriangulateDuration += TriangulateDuration;
-	Chronos.GlobalDelaunayDuration += IsoTrianguler.Chronos.FindSegmentToLinkLoopToLoopByDelaunayDuration;
-	Chronos.GlobalMeshDuration += Duration;
+	MesherReport.Chronos.GlobalPointCloudDuration += Grid.Chronos.GeneratePointCloudDuration;
+	MesherReport.Chronos.GlobalGeneratePointCloudDuration += GenerateCloudDuration;
+	MesherReport.Chronos.GlobalTriangulateDuration += TriangulateDuration;
+	MesherReport.Chronos.GlobalDelaunayDuration += IsoTrianguler.Chronos.FindSegmentToLinkLoopToLoopByDelaunayDuration;
+	MesherReport.Chronos.GlobalMeshDuration += Duration;
 #endif
 
 	FChrono::PrintClockElapse(EVerboseLevel::Debug, TEXT("   "), TEXT("Meshing"), Duration);
 
-	if (BoolDisplayDebugMeshStep)
-	{
-		Close3DDebugSession();
-	}
 }
 
 void FParametricMesher::GenerateCloud(FGrid& Grid)
@@ -300,20 +291,22 @@ void FParametricMesher::GenerateCloud(FGrid& Grid)
 #endif
 			FTimePoint MeshThinZonesTime = FChrono::Now();
 			MeshThinZoneEdges(Grid);
-			Chronos.GlobalMeshThinZones += FChrono::Elapse(MeshThinZonesTime);
+			MesherReport.Chronos.GlobalMeshThinZones += FChrono::Elapse(MeshThinZonesTime);
 		}
-		Chronos.GlobalThinZones += FChrono::Elapse(StartTime);
+		MesherReport.Chronos.GlobalThinZones += FChrono::Elapse(StartTime);
 	}
 
+#ifdef DEBUG_THIN_ZONES
 	Grid.DisplayInnerDomainPoints(TEXT("FGrid::PointCloud 2D"), Grid.GetInner2DPoints(EGridSpace::Default2D));
 	//Wait(Grid.bDisplay);
+#endif
 
 	FTimePoint StartTime = FChrono::Now();
 	MeshFaceLoops(Grid);
 
 	Grid.ProcessPointCloud();
 
-	Chronos.GlobalMeshAndGetLoopNodes += FChrono::Elapse(StartTime);
+	MesherReport.Chronos.GlobalMeshAndGetLoopNodes += FChrono::Elapse(StartTime);
 }
 
 void FParametricMesher::MeshFaceLoops(FGrid& Grid)
@@ -330,7 +323,7 @@ void FParametricMesher::MeshFaceLoops(FGrid& Grid)
 		}
 	}
 
-	Chronos.GlobalMeshEdges += FChrono::Elapse(StartTime);
+	MesherReport.Chronos.GlobalMeshEdges += FChrono::Elapse(StartTime);
 }
 
 static void FillImposedIsoCuttingPoints(TArray<double>& UEdgeSetOfIntersectionWithIso, ECoordinateType CoordinateType, double EdgeToleranceGeo, const FTopologicalEdge& Edge, TArray<FCuttingPoint>& OutImposedIsoVertexSet)
@@ -1363,7 +1356,7 @@ void FParametricMesher::MeshThinZoneEdges(FGrid& Grid)
 	}
 	Close3DDebugSession();
 #endif
-	Chronos.GlobalMeshThinZones += FChrono::Elapse(MeshStartTime);
+	MesherReport.Chronos.GlobalMeshThinZones += FChrono::Elapse(MeshStartTime);
 
 }
 
diff --git a/Engine/Source/Runtime/Datasmith/CADKernel/Private/CADKernel/Mesh/Structure/Grid.cpp b/Engine/Source/Runtime/Datasmith/CADKernel/Private/CADKernel/Mesh/Structure/Grid.cpp
--- a/Engine/Source/Runtime/Datasmith/CADKernel/Private/CADKernel/Mesh/Structure/Grid.cpp
+++ b/Engine/Source/Runtime/Datasmith/CADKernel/Private/CADKernel/Mesh/Structure/Grid.cpp
@@ -1602,7 +1602,9 @@ void FGrid::FindInnerFacePoints()
 		}
 	}
 
-	DisplayLoop(TEXT("FGrid::Loop 2D After move according tol"), GetLoops2D(EGridSpace::UniformScaled), true, false);
+#ifdef DEBUG_FIND_INNER_FACE_POINTS
+	DisplayGridLoop(TEXT("FGrid::Loop 2D After move according tol"), GetLoops2D(EGridSpace::UniformScaled), true, false);
+#endif
 
 	// Intersection along U axis
 	for (const TArray<FPoint2D>& Loop : FaceLoops2D[EGridSpace::UniformScaled])
@@ -1815,124 +1817,6 @@ void FGrid::FindInnerFacePoints()
 	Chronos.FindInnerDomainPointsDuration += FChrono::Elapse(StartTime);
 }
 
-void FGrid::DisplayGridPoints(EGridSpace DisplaySpace) const
-{
-	if (!bDisplay)
-	{
-		return;
-	}
-	int32 NbNum = 0;
-	{
-		F3DDebugSession _(TEXT("FGrid::FindInnerDomainPoints Inside Point"));
-		for (int32 Index = 0; Index < CuttingSize; ++Index)
-		{
-			if (IsInsideFace[Index])
-			{
-				DisplayPoint(Points2D[DisplaySpace][Index], Index);
-				NbNum++;
-			}
-		}
-	}
-	ensureCADKernel(NbNum == CountOfInnerNodes);
-
-	{
-		F3DDebugSession _(TEXT("FGrid::FindInnerDomainPoints Outside Point"));
-		for (int32 Index = 0; Index < CuttingSize; ++Index)
-		{
-			if (!IsInsideFace[Index])
-			{
-				DisplayPoint(Points2D[DisplaySpace][Index], EVisuProperty::OrangePoint, Index);
-			}
-		}
-	}
-}
-
-void FGrid::DisplayGridNormal() const
-{
-	if (!bDisplay)
-	{
-		return;
-	}
-
-	double NormalLength = FSystem::Get().GetVisu()->GetParameters()->NormalLength;
-	{
-		F3DDebugSession _(TEXT("FGrid::Inner Normal"));
-
-		for (int32 Index = 0; Index < CuttingSize; ++Index)
-		{
-			F3DDebugSegment GraphicSegment(Index);
-			FVector Normal = Normals[Index];
-			Normal.Normalize();
-			Normal *= NormalLength;
-			DrawSegment(Points3D[Index], Points3D[Index] + Normal, EVisuProperty::GreenCurve);
-		}
-	}
-
-	{
-		F3DDebugSession _(TEXT("FGrid::Loop Normal"));
-		for (int32 LoopIndex = 0; LoopIndex < FaceLoops3D.Num(); ++LoopIndex)
-		{
-			const TArray<FPoint>& LoopPoints = FaceLoops3D[LoopIndex];
-			const TArray<FVector>& LoopNormals = NormalsOfFaceLoops[LoopIndex];
-
-			for (int32 Index = 0; Index < LoopPoints.Num(); ++Index)
-			{
-				F3DDebugSegment GraphicSegment(Index);
-				FVector Normal = LoopNormals[Index];
-				Normal.Normalize();
-				Normal *= NormalLength;
-				DrawSegment(LoopPoints[Index], LoopPoints[Index] + Normal, EVisuProperty::YellowCurve);
-			}
-		}
-	}
-}
-
-
-void FGrid::DisplayFindPointsCloseToLoop(EGridSpace DisplaySpace) const
-{
-	if (!bDisplay)
-	{
-		return;
-	}
-
-	F3DDebugSession _(TEXT("FGrid::FindPointsClosedToLoop result"));
-	for (int32 Index = 0; Index < CuttingSize; ++Index)
-	{
-		if (IsCloseToLoop[Index])
-		{
-			DisplayPoint(Points2D[DisplaySpace][Index]);
-		}
-		else
-		{
-			DisplayPoint(Points2D[DisplaySpace][Index], EVisuProperty::YellowPoint);
-		}
-	}
-}
-
-void FGrid::DisplayGridInnerPoints(EGridSpace DisplaySpace, TCHAR* Message) const
-{
-	if (!bDisplay)
-	{
-		return;
-	}
-
-	F3DDebugSession _(Message);
-	for (int32 Index = 0; Index < CuttingSize; ++Index)
-	{
-		if (IsInsideFace[Index])
-		{
-			if (IsCloseToLoop[Index])
-			{
-				DisplayPoint(Points2D[DisplaySpace][Index], EVisuProperty::BluePoint, Index);
-			}
-			else
-			{
-				DisplayPoint(Points2D[DisplaySpace][Index], EVisuProperty::YellowPoint, Index);
-			}
-		}
-	}
-}
-
 bool FGrid::CheckIfDegenerated()
 {
 	if (FaceLoops2D[EGridSpace::Default2D].Num() == 0)
