CADKernel Update

diff --git a/Engine/Source/Runtime/Datasmith/CADKernel/Private/CADKernel/Mesh/Structure/Grid.h b/Engine/Source/Runtime/Datasmith/CADKernel/Private/CADKernel/Mesh/Structure/Grid.h
--- a/Engine/Source/Runtime/Datasmith/CADKernel/Private/CADKernel/Mesh/Structure/Grid.h
+++ b/Engine/Source/Runtime/Datasmith/CADKernel/Private/CADKernel/Mesh/Structure/Grid.h
@@ -10,630 +10,692 @@
 #include "CADKernel/Mesh/Structure/ThinZone2D.h"
 #include "CADKernel/UI/Display.h"
 
-
 namespace CADKernel
 {
-	struct FCuttingPoint;
-	struct FCuttingGrid;
-	
-	class FModelMesh;
-	class FPoint;
-	class FPoint2D;
-	class FThinZone2DFinder;
-	class FTopologicalFace;
-
-	struct FGridChronos
-	{
-		FDuration DefineCuttingParametersDuration;
-		FDuration GeneratePointCloudDuration;
-		FDuration ProcessPointCloudDuration;
-		FDuration FindInnerDomainPointsDuration;
-		FDuration Build2DLoopDuration;
-		FDuration RemovePointsClosedToLoopDuration;
-		FDuration FindPointsCloseToLoopDuration;
-		FDuration ScaleGridDuration;
-
-		FGridChronos()
-			: DefineCuttingParametersDuration(FChrono::Init())
-			, GeneratePointCloudDuration(FChrono::Init())
-			, ProcessPointCloudDuration(FChrono::Init())
-			, FindInnerDomainPointsDuration(FChrono::Init())
-			, Build2DLoopDuration(FChrono::Init())
-			, RemovePointsClosedToLoopDuration(FChrono::Init())
-			, FindPointsCloseToLoopDuration(FChrono::Init())
-			, ScaleGridDuration(FChrono::Init())
-		{}
-
-		void PrintTimeElapse() const
-		{
-			FDuration GridDuration = FChrono::Elapse(FChrono::Now());
-			GridDuration += DefineCuttingParametersDuration;
-			GridDuration += GeneratePointCloudDuration;
-			GridDuration += FindInnerDomainPointsDuration;
-			GridDuration += Build2DLoopDuration;
-			GridDuration += RemovePointsClosedToLoopDuration;
-			GridDuration += FindPointsCloseToLoopDuration;
-			FChrono::PrintClockElapse(Log, TEXT(""), TEXT("Grid"), GridDuration);
-			FChrono::PrintClockElapse(Log, TEXT("  "), TEXT("DefineCuttingParameters"), DefineCuttingParametersDuration);
-			FChrono::PrintClockElapse(Log, TEXT("  "), TEXT("GeneratePointCloud"), GeneratePointCloudDuration);
-			FChrono::PrintClockElapse(Log, TEXT("  "), TEXT("GenerateDomainPoints"), ProcessPointCloudDuration);
-			FChrono::PrintClockElapse(Log, TEXT("    "), TEXT("FindInnerDomainPointsDuration"), FindInnerDomainPointsDuration);
-			FChrono::PrintClockElapse(Log, TEXT("    "), TEXT("FindPointsCloseToLoop"), FindPointsCloseToLoopDuration);
-			FChrono::PrintClockElapse(Log, TEXT("    "), TEXT("RemovePointsClosedToLoop"), RemovePointsClosedToLoopDuration);
-			FChrono::PrintClockElapse(Log, TEXT("    "), TEXT("Build2DLoopDuration"), Build2DLoopDuration);
-		}
-	};
-
-	class FGrid : public FHaveStates
-	{
-	protected:
-
-		const TSharedRef<FTopologicalFace> Face;
-		const FSurfacicTolerance FaceTolerance;
-		const double Tolerance3D;
-		const double MinimumElementSize;
-
-		TSharedRef<FModelMesh> MeshModel;
-
-		FThinZone2DFinder ThinZoneFinder;
-
-		/**
-		 * 2D Coordinate of Loop's nodes in each space
-		 */
-		TArray<TArray<FPoint2D>> FaceLoops2D[EGridSpace::EndGridSpace];
-
-		/**
-		 * 3D Coordinate of Loop nodes in each space
-		 */
-		TArray<TArray<FPoint>> FaceLoops3D;
 
-		/**
-		 * Surface Normal at each boundary nodes
-		 */
-		TArray<TArray<FVector>> NormalsOfFaceLoops;
+struct FCuttingPoint;
+struct FCuttingGrid;
 
-		TArray<TArray<int32>> NodeIdsOfFaceLoops;
+class FIsoNode;
+class FIsoSegment;
+class FLoopNode;
+class FModelMesh;
+class FPoint;
+class FPoint2D;
+class FThinZone2DFinder;
+class FTopologicalFace;
 
-		/**
-		 * grid point cloud size
-		 */ 
-		int32 CuttingCount[2] = {0 , 0};
-		int32 CuttingSize = 0;
-
-		/**
-		 * count of node inside the face i.e. node inside the external loop and outer the inner loops 
-		 */
-		int32 CountOfInnerNodes = 0;
-
-		/*
-		 * Cutting coordinates of the face respecting the meshing criteria
-		 */
-		const FCoordinateGrid& CuttingCoordinates;
-		FCoordinateGrid UniformCuttingCoordinates;
+struct FGridChronos
+{
+	FDuration DefineCuttingParametersDuration;
+	FDuration GeneratePointCloudDuration;
+	FDuration ProcessPointCloudDuration;
+	FDuration FindInnerDomainPointsDuration;
+	FDuration Build2DLoopDuration;
+	FDuration RemovePointsClosedToLoopDuration;
+	FDuration FindPointsCloseToLoopDuration;
+	FDuration ScaleGridDuration;
+
+	FGridChronos()
+		: DefineCuttingParametersDuration(FChrono::Init())
+		, GeneratePointCloudDuration(FChrono::Init())
+		, ProcessPointCloudDuration(FChrono::Init())
+		, FindInnerDomainPointsDuration(FChrono::Init())
+		, Build2DLoopDuration(FChrono::Init())
+		, RemovePointsClosedToLoopDuration(FChrono::Init())
+		, FindPointsCloseToLoopDuration(FChrono::Init())
+		, ScaleGridDuration(FChrono::Init())
+	{}
+
+	void PrintTimeElapse() const
+	{
+		FDuration GridDuration = FChrono::Elapse(FChrono::Now());
+		GridDuration += DefineCuttingParametersDuration;
+		GridDuration += GeneratePointCloudDuration;
+		GridDuration += FindInnerDomainPointsDuration;
+		GridDuration += Build2DLoopDuration;
+		GridDuration += RemovePointsClosedToLoopDuration;
+		GridDuration += FindPointsCloseToLoopDuration;
+		FChrono::PrintClockElapse(Log, TEXT(""), TEXT("Grid"), GridDuration);
+		FChrono::PrintClockElapse(Log, TEXT("  "), TEXT("DefineCuttingParameters"), DefineCuttingParametersDuration);
+		FChrono::PrintClockElapse(Log, TEXT("  "), TEXT("GeneratePointCloud"), GeneratePointCloudDuration);
+		FChrono::PrintClockElapse(Log, TEXT("  "), TEXT("GenerateDomainPoints"), ProcessPointCloudDuration);
+		FChrono::PrintClockElapse(Log, TEXT("    "), TEXT("FindInnerDomainPointsDuration"), FindInnerDomainPointsDuration);
+		FChrono::PrintClockElapse(Log, TEXT("    "), TEXT("FindPointsCloseToLoop"), FindPointsCloseToLoopDuration);
+		FChrono::PrintClockElapse(Log, TEXT("    "), TEXT("RemovePointsClosedToLoop"), RemovePointsClosedToLoopDuration);
+		FChrono::PrintClockElapse(Log, TEXT("    "), TEXT("Build2DLoopDuration"), Build2DLoopDuration);
+	}
+};
+
+class FGrid : public FHaveStates
+{
+protected:
+
+	const TSharedRef<FTopologicalFace> Face;
+	const FSurfacicTolerance FaceTolerance;
+	const double Tolerance3D;
+	const double MinimumElementSize;
+
+	TSharedRef<FModelMesh> MeshModel;
+
+	FThinZone2DFinder ThinZoneFinder;
+
+	/**
+	 * 2D Coordinate of Loop's nodes in each space
+	 */
+	TArray<TArray<FPoint2D>> FaceLoops2D[EGridSpace::EndGridSpace];
+
+	/**
+	 * 3D Coordinate of Loop nodes in each space
+	 */
+	TArray<TArray<FPoint>> FaceLoops3D;
+
+	/**
+	 * Surface Normal at each boundary nodes
+	 */
+	TArray<TArray<FVector>> NormalsOfFaceLoops;
+
+	TArray<TArray<int32>> NodeIdsOfFaceLoops;
+
+	/**
+	 * grid point cloud size
+	 */
+	int32 CuttingCount[2] = { 0 , 0 };
+	int32 CuttingSize = 0;
+
+	/**
+	 * count of node inside the face i.e. node inside the external loop and outer the inner loops
+	 */
+	int32 CountOfInnerNodes = 0;
+
+	/*
+	 * Cutting coordinates of the face respecting the meshing criteria
+	 */
+	const FCoordinateGrid& CuttingCoordinates;
+	FCoordinateGrid UniformCuttingCoordinates;
+
+	/*
+	 * Maximum difference of coordinate along the specified axis of two successive cutting points
+	 */
+	FPoint2D MaxDeltaUV = { 0., 0. };
+
+	/**
+	 * Maximum 3d distance along each axis of two successive cutting points
+	 */
+	FPoint2D MaxElementSize = { 0., 0. };
+
+	double MinOfMaxElementSize = 0;
+
+	/**
+	 * Array to flag each points as inside the face (if IsInsideDomain[Index] = 1) or outer
+	 * @see IsNodeInsideFace
+	 */
+	TArray<char> IsInsideFace;
+
+	/**
+	 * Array to flag each points as close to a loop of the face if IsCloseToLoop[Index] = 1
+	 * @see IsNodeCloseToLoop
+	 */
+	TArray<char> IsCloseToLoop;
+
+	/**
+	 * 2D Coordinate of grid nodes in each space
+	 */
+	TArray<FPoint2D> Points2D[EGridSpace::EndGridSpace];
+
+	/**
+	 * 3D Coordinate of inner nodes
+	 */
+	TArray<FPoint> Points3D;
+
+	/**
+	 * Surface Normal at each inner nodes
+	 */
+	TArray<FVector> Normals;
+
+public:
+	FGridChronos Chronos;
+
+public:
+	FGrid(TSharedRef<FTopologicalFace>& InFace, TSharedRef<FModelMesh>& InShellMesh);
 
-		/*
-		 * Maximum difference of coordinate along the specified axis of two successive cutting points
-		 */
-		FPoint2D MaxDeltaUV = { 0., 0. };
+#ifndef CADKERNEL_DEV
+	virtual ~FGrid() = default;
+#else
+	virtual ~FGrid()
+	{
+		Close3DDebugSession(bDisplay);
+	}
+#endif
 
-		/**
-		 * Maximum 3d distance along each axis of two successive cutting points
-		 */
-		FPoint2D MaxElementSize = { 0., 0. };
+	// ======================================================================================================================================================================================================================
+	// Meshing tools =========================================================================================================================================================================================================
+	// ======================================================================================================================================================================================================================
+
+	/**
+	 * Return true if the grid is not consistent to build a mesh as composed of only two border nodes.
+	 */
+	bool CheckIfDegenerated();
+
+	/**
+	 * Defines the cutting coordinate of the grid according to mesh criteria and the existing mesh of bordering edges (loop's edges)
+	 * @see GetPreferredUVCoordinatesFromNeighbours
+	 */
+	void DefineCuttingParameters();
+	void DefineCuttingParameters(EIso Iso, FCuttingGrid& Neighbors);
+
+	/**
+	 * Computes the 2d Points, 3D points and Normals of the grid.
+	 * call of DefineCuttingParameters is mandatory called before
+	 * @return false if the grid is degenerated
+	 */
+	bool GeneratePointCloud();
+
+	/**
+	 * Process the generated point cloud to compute the scaled parametric spaces and to identify outer points of the face
+	 * @see IsNodeInsideFace
+	 * @see Points2D
+	 */
+	void ProcessPointCloud();
+
+	/**
+	 * Convert an array of points of "DefaultParametric" space into a scaled parametric space
+	 * @see FThinZone2DFinder::BuildBoundarySegments()
+	 */
+	void TransformPoints(EGridSpace DestinationSpace, const TArray<FPoint2D>& InPointsToScale, TArray<FPoint2D>& OutScaledPoints) const;
+
+protected:
+
+	/**
+	 * Builds the scaled parametric spaces
+	 * @see Points2D
+	 * @see GeneratePointCloud (ended GeneratePointCloud process)
+	 */
+	void ScaleGrid();
+
+	/**
+	 * Projects loop's points in the scaled parametric spaces
+	 * @see FaceLoops2D
+	 * @see ProcessPointClound (called in ProcessPointClound)
+	 */
+	void ScaleLoops();
+
+	/**
+	 * After the loops have been scaled, the distance between two consecutive vertices is verified to avoid degenerated mesh elements
+	 * @ return false if the external loop is degenerated after the process
+	 */
+	bool RemoveCoincidentNodes();
+
+	/**
+	 * Convert Coordinate of "DefaultParametric" space into a scaled parametric space
+	 * @see ScaleLoops
+	 */
+	void ComputeNewCoordinate(const TArray<FPoint2D>& NewGrid, int32 IndexU, int32 IndexV, const FPoint2D& InPoint, FPoint2D& OutNewScaledPoint) const
+	{
+		OutNewScaledPoint =
+			NewGrid[(IndexV + 0) * CuttingCount[EIso::IsoU] + (IndexU + 0)] * (CuttingCoordinates[EIso::IsoU][IndexU + 1] - InPoint.U) * (CuttingCoordinates[EIso::IsoV][IndexV + 1] - InPoint.V) -
+			NewGrid[(IndexV + 0) * CuttingCount[EIso::IsoU] + (IndexU + 1)] * (CuttingCoordinates[EIso::IsoU][IndexU + 0] - InPoint.U) * (CuttingCoordinates[EIso::IsoV][IndexV + 1] - InPoint.V) -
+			NewGrid[(IndexV + 1) * CuttingCount[EIso::IsoU] + (IndexU + 0)] * (CuttingCoordinates[EIso::IsoU][IndexU + 1] - InPoint.U) * (CuttingCoordinates[EIso::IsoV][IndexV + 0] - InPoint.V) +
+			NewGrid[(IndexV + 1) * CuttingCount[EIso::IsoU] + (IndexU + 1)] * (CuttingCoordinates[EIso::IsoU][IndexU + 0] - InPoint.U) * (CuttingCoordinates[EIso::IsoV][IndexV + 0] - InPoint.V);
+		OutNewScaledPoint /= (CuttingCoordinates[EIso::IsoU][IndexU + 1] - CuttingCoordinates[EIso::IsoU][IndexU + 0]) * (CuttingCoordinates[EIso::IsoV][IndexV + 1] - CuttingCoordinates[EIso::IsoV][IndexV + 0]);
+	};
 
-		double MinOfMaxElementSize = 0;
+	/**
+	 * @see ProcessPointClound (called in ProcessPointClound)
+	 */
+	void FindInnerFacePoints();
+
+	/**
+	 * Define if an inner node is close to a loop i.e. if the loop cross the cell [[IndexU - 1, IndexU], [IndexV - 1, IndexV]], each corner node of the cell is close to the loop
+	 *
+	 * This algorithm is inspired of Bresenham's line algorithm
+	 * The loop is traversed segment by segment
+	 * For each segment, according to the slop of the segment, each cell intersecting the segment is selected (FindIntersectionsCaseSlop_) i.e. each corner node of the cell is flagged IsCloseToLoop (SetCellCloseToLoop)
+	 * @see SlopeUtils
+	 * @see IsNodeCloseToLoop
+	 * @see ProcessPointCloud (called in ProcessPointClound)
+	 */
+	void FindPointsCloseToLoop();
+
+	/**
+	 * Find all node close to the boundary.
+	 * For each node, the 3d distance to the boundary is estimated
+	 * If this distance is bigger than the Tolerance, the node is remove
+	 * Removing a node can make its neighbor removable,
+	 * Each removed node make its direct neighbor removable, the process is recursive as long as it is possible to remove a node
+	 *
+	 * @see ProcessPointClound (called in ProcessPointClound)
+	 */
+	void RemovePointsClosedToLoop();
+
+
+	/**
+	 * Gets the cutting coordinates of the existing mesh of bordering edges (loop's edges)
+	 * @see DefineCuttingParameters (called in DefineCuttingParameters)
+	 */
+	void GetPreferredUVCuttingParametersFromLoops(FCuttingGrid& CuttingFromLoops);
+
+	void ComputeMaxElementSize();
+	void ComputeMaxDeltaUV();
+
+	/**
+	 * @ return false if the mesh of the loop is degenerated
+	 */
+	bool GetMeshOfLoops();
+
+public:
+
+	// ======================================================================================================================================================================================================================
+	// ThinZone properties for meshing scheduling optimisation ==============================================================================================================================================================
+	// ======================================================================================================================================================================================================================
+
+	void SearchThinZones();
+	const TArray<FThinZone2D>& GetThinZones() const
+	{
+		return ThinZoneFinder.GetThinZones();
+	}
 
-		/**
-		 * Array to flag each points as inside the face (if IsInsideDomain[Index] = 1) or outer 
-		 * @see IsNodeInsideFace
-		 */
-		TArray<char> IsInsideFace;
+	// ======================================================================================================================================================================================================================
+	// GET Methodes =========================================================================================================================================================================================================
+	// ======================================================================================================================================================================================================================
 
-		/**
-		 * Array to flag each points as close to a loop of the face if IsCloseToLoop[Index] = 1
-		 * @see IsNodeCloseToLoop
-		 */
-		TArray<char> IsCloseToLoop;
+	constexpr const int32 GetCuttingCount(EIso Iso) const
+	{
+		return CuttingCount[Iso];
+	}
 
-		/**
-		 * 2D Coordinate of grid nodes in each space
-		 */
-		TArray<FPoint2D> Points2D[EGridSpace::EndGridSpace];
+	/**
+	 * Return the number of points (inner and outer points) of the grid (i.e. CuttingCount[IsoU] x CuttingCount[IsoV]).
+	 */
+	const int32 GetTotalCuttingCount() const
+	{
+		return CuttingSize;
+	}
 
-		/**
-		 * 3D Coordinate of inner nodes
-		 */
-		TArray<FPoint> Points3D;
+	const double GetTolerance(EIso Iso) const
+	{
+		return FaceTolerance[Iso];
+	}
 
-		/**
-		 * Surface Normal at each inner nodes
-		 */
-		TArray<FVector> Normals;
+	/**
+	 * @return true if the node is inner the external loop and outer the inner loops
+	 */
+	const bool IsNodeInsideFace(int32 IndexU, int32 IndexV) const
+	{
+		return IsInsideFace[GobalIndex(IndexU, IndexV)] == 1;
+	}
 
-	public:
-		FGridChronos Chronos;
+	/**
+	 * @return true if the node is inner the external loop and outer the inner loops
+	 */
+	const bool IsNodeInsideFace(int32 Index) const
+	{
+		return IsInsideFace[Index] == 1;
+	}
 
-	public:
-		FGrid(TSharedRef<FTopologicalFace>& InFace, TSharedRef<FModelMesh>& InShellMesh);
+	/**
+	 * @return the count of inner nodes
+	 */
+	const int32 InnerNodesCount() const
+	{
+		return CountOfInnerNodes;
+	}
 
-#ifndef CADKERNEL_DEV
-		virtual ~FGrid() = default;
-#else
-		virtual ~FGrid()
-		{
-			Close3DDebugSession(bDisplay);
-		}
-#endif
+	/**
+	 * @return true if the node is close to a loop i.e. the loop cross the space [[IndexU - 1, IndexU + 1], [IndexV - 1, IndexV + 1]]
+	 */
+	const bool IsNodeCloseToLoop(int32 IndexU, int32 IndexV) const
+	{
+		return IsCloseToLoop[GobalIndex(IndexU, IndexV)] == 1;
+	}
 
-		// ======================================================================================================================================================================================================================
-		// Meshing tools =========================================================================================================================================================================================================
-		// ======================================================================================================================================================================================================================
-		
-		/**
-		 * Return true if the grid is not consistent to build a mesh as composed of only two border nodes.
-		 */
-		bool CheckIfDegenerated();
-
-		/**
-		 * Defines the cutting coordinate of the grid according to mesh criteria and the existing mesh of bordering edges (loop's edges)
-		 * @see GetPreferredUVCoordinatesFromNeighbours
-		 */
-		void DefineCuttingParameters();
-		void DefineCuttingParameters(EIso Iso, FCuttingGrid& Neighbors);
-
-		/**
-		 * Computes the 2d Points, 3D points and Normals of the grid.
-		 * call of DefineCuttingParameters is mandatory called before
-		 * @return false if the grid is degenerated
-		 */
-		bool GeneratePointCloud();
-
-		/**
-		 * Process the generated point cloud to compute the scaled parametric spaces and to identify outer points of the face
-		 * @see IsNodeInsideFace
-		 * @see Points2D
-		 */
-		void ProcessPointCloud();
-
-		/**
-		 * Convert an array of points of "DefaultParametric" space into a scaled parametric space
-		 * @see FThinZone2DFinder::BuildBoundarySegments()
-		 */
-		void TransformPoints(EGridSpace DestinationSpace, const TArray<FPoint2D>& InPointsToScale, TArray<FPoint2D>& OutScaledPoints) const;
-
-	protected:
-
-		/**
-		 * Builds the scaled parametric spaces 
-		 * @see Points2D
-		 * @see GeneratePointCloud (ended GeneratePointCloud process)
-		 */
-		void ScaleGrid();
-
-		/**
-		 * Projects loop's points in the scaled parametric spaces
-		 * @see FaceLoops2D
-		 * @see ProcessPointClound (called in ProcessPointClound)
-		 */
-		void ScaleLoops();
-
-		/**
-		 * After the loops have been scaled, the distance between two consecutive vertices is verified to avoid degenerated mesh elements
-		 * @ return false if the external loop is degenerated after the process
-		 */
-		bool RemoveCoincidentNodes();
-
-		/**
-		 * Convert Coordinate of "DefaultParametric" space into a scaled parametric space
-		 * @see ScaleLoops
-		 */
-		void ComputeNewCoordinate(const TArray<FPoint2D>& NewGrid, int32 IndexU, int32 IndexV, const FPoint2D& InPoint, FPoint2D& OutNewScaledPoint) const
-		{
-			OutNewScaledPoint =
-				NewGrid[(IndexV + 0) * CuttingCount[EIso::IsoU] + (IndexU + 0)] * (CuttingCoordinates[EIso::IsoU][IndexU + 1] - InPoint.U) * (CuttingCoordinates[EIso::IsoV][IndexV + 1] - InPoint.V) -
-				NewGrid[(IndexV + 0) * CuttingCount[EIso::IsoU] + (IndexU + 1)] * (CuttingCoordinates[EIso::IsoU][IndexU + 0] - InPoint.U) * (CuttingCoordinates[EIso::IsoV][IndexV + 1] - InPoint.V) -
-				NewGrid[(IndexV + 1) * CuttingCount[EIso::IsoU] + (IndexU + 0)] * (CuttingCoordinates[EIso::IsoU][IndexU + 1] - InPoint.U) * (CuttingCoordinates[EIso::IsoV][IndexV + 0] - InPoint.V) +
-				NewGrid[(IndexV + 1) * CuttingCount[EIso::IsoU] + (IndexU + 1)] * (CuttingCoordinates[EIso::IsoU][IndexU + 0] - InPoint.U) * (CuttingCoordinates[EIso::IsoV][IndexV + 0] - InPoint.V);
-			OutNewScaledPoint /= (CuttingCoordinates[EIso::IsoU][IndexU + 1] - CuttingCoordinates[EIso::IsoU][IndexU + 0]) * (CuttingCoordinates[EIso::IsoV][IndexV + 1] - CuttingCoordinates[EIso::IsoV][IndexV + 0]);
-		};
-
-		/**
-		 * @see ProcessPointClound (called in ProcessPointClound)
-		 */
-		void FindInnerFacePoints();
-
-		/**
-		 * Define if an inner node is close to a loop i.e. if the loop cross the cell [[IndexU - 1, IndexU], [IndexV - 1, IndexV]], each corner node of the cell is close to the loop
-		 * 
-		 * This algorithm is inspired of Bresenham's line algorithm
-		 * The loop is traversed segment by segment
-		 * For each segment, according to the slop of the segment, each cell intersecting the segment is selected (FindIntersectionsCaseSlop_) i.e. each corner node of the cell is flagged IsCloseToLoop (SetCellCloseToLoop)
-		 * @see SlopeUtils
-		 * @see IsNodeCloseToLoop
-		 * @see ProcessPointCloud (called in ProcessPointClound)
-		 */
-		void FindPointsCloseToLoop();
-
-		/**
-		 * Find all node close to the boundary.
-		 * For each node, the 3d distance to the boundary is estimated
-		 * If this distance is bigger than the Tolerance, the node is remove
-		 * Removing a node can make its neighbor removable,
-		 * Each removed node make its direct neighbor removable, the process is recursive as long as it is possible to remove a node
-		 *
-		 * @see ProcessPointClound (called in ProcessPointClound)
-		 */
-		void RemovePointsClosedToLoop();
-
-
-		/**
-		 * Gets the cutting coordinates of the existing mesh of bordering edges (loop's edges)
-		 * @see DefineCuttingParameters (called in DefineCuttingParameters)
-		 */
-		void GetPreferredUVCuttingParametersFromLoops(FCuttingGrid& CuttingFromLoops);
-
-		void ComputeMaxElementSize();
-		void ComputeMaxDeltaUV();
-
-		/**
-		 * @ return false if the mesh of the loop is degenerated
-		 */
-		bool GetMeshOfLoops();
-
-	public:
-
-		// ======================================================================================================================================================================================================================
-		// ThinZone properties for meshing scheduling optimisation ==============================================================================================================================================================
-		// ======================================================================================================================================================================================================================
-
-		void SearchThinZones();
-		const TArray<FThinZone2D>& GetThinZones() const
-		{
-			return ThinZoneFinder.GetThinZones();
-		}
+	/**
+	 * @return true if the node is close to a loop i.e. the loop cross the space [[IndexU - 1, IndexU + 1], [IndexV - 1, IndexV + 1]]
+	 */
+	const bool IsNodeCloseToLoop(int32 Index) const
+	{
+		return IsCloseToLoop[Index] == 1;
+	}
+
+	/**
+	 * @return the FPoint2D (parametric coordinates) of the point at the Index of the grid in the defined grid space
+	 * @see EGridSpace
+	 * @see Points2D
+	 */
+	const FPoint2D& GetInner2DPoint(EGridSpace Space, int32 Index) const
+	{
+		return Points2D[(int32)Space][Index];
+	}
 
-		// ======================================================================================================================================================================================================================
-		// GET Methodes =========================================================================================================================================================================================================
-		// ======================================================================================================================================================================================================================
+	void SetInner2DPoint(EGridSpace Space, int32 Index, const FPoint2D& NewCoordinate)
+	{
+		Points2D[(int32)Space][Index] = NewCoordinate;
+	}
 
-		constexpr const int32 GetCuttingCount(EIso Iso) const
-		{
-			return CuttingCount[Iso];
-		}
+	/**
+	 * @return the FPoint2D (parametric coordinates) of the point at the Index of the grid in the defined grid space @see EGridSpace
+	 */
+	const FPoint2D& GetInner2DPoint(EGridSpace Space, int32 IndexU, int32 IndexV) const
+	{
+		return Points2D[(int32)Space][GobalIndex(IndexU, IndexV)];
+	}
 
-		/**
-		 * Return the number of points (inner and outer points) of the grid (i.e. CuttingCount[IsoU] x CuttingCount[IsoV]).
-		 */
-		const int32 GetTotalCuttingCount() const
-		{
-			return CuttingSize;
-		}
+	/**
+	 * @return the FPoint (3D coordinates) of the point at the Index of the grid
+	 */
+	const FPoint& GetInner3DPoint(int32 Index) const
+	{
+		return Points3D[Index];
+	}
 
-		const double GetTolerance(EIso Iso) const
-		{
-			return FaceTolerance[Iso];
-		}
+	/**
+	 * @return the FPoint (3D coordinates) of the point at the Index of the grid
+	 */
+	const FPoint& GetInner3DPoint(int32 IndexU, int32 IndexV) const
+	{
+		return Points3D[GobalIndex(IndexU, IndexV)];
+	}
 
-		/**
-		 * @return true if the node is inner the external loop and outer the inner loops 
-		 */
-		const bool IsNodeInsideFace(int32 IndexU, int32 IndexV) const
-		{
-			return IsInsideFace[GobalIndex(IndexU, IndexV)] == 1;
-		}
+	/**
+	 * @return the normal of the surface at the point at the Index of the grid
+	 */
+	const FVector& GetPointNormal(int32 IndexU, int32 IndexV) const
+	{
+		return Normals[GobalIndex(IndexU, IndexV)];
+	}
 
-		/**
-		 * @return true if the node is inner the external loop and outer the inner loops
-		 */
-		const bool IsNodeInsideFace(int32 Index) const
-		{
-			return IsInsideFace[Index] == 1;
-		}
+	/**
+	 * @return the normal of the surface at the point at the Index of the grid
+	 */
+	const FVector& GetPointNormal(int32 Index) const
+	{
+		return Normals[Index];
+	}
 
-		/**
-		 * @return the count of inner nodes
-		 */
-		const int32 InnerNodesCount() const
-		{
-			return CountOfInnerNodes;
-		}
+	constexpr const TArray<double>& GetCuttingCoordinatesAlongIso(EIso Iso) const
+	{
+		return CuttingCoordinates[Iso];
+	}
 
-		/**
-		 * @return true if the node is close to a loop i.e. the loop cross the space [[IndexU - 1, IndexU + 1], [IndexV - 1, IndexV + 1]]
-		 */
-		const bool IsNodeCloseToLoop(int32 IndexU, int32 IndexV) const
-		{
-			return IsCloseToLoop[GobalIndex(IndexU, IndexV)] == 1;
-		}
+	const FCoordinateGrid& GetCuttingCoordinates() const
+	{
+		return CuttingCoordinates;
+	}
 
-		/**
-		 * @return true if the node is close to a loop i.e. the loop cross the space [[IndexU - 1, IndexU + 1], [IndexV - 1, IndexV + 1]]
-		 */
-		const bool IsNodeCloseToLoop(int32 Index) const
-		{
-			return IsCloseToLoop[Index] == 1;
-		}
+	constexpr const TArray<double>& GetUniformCuttingCoordinatesAlongIso(EIso Iso) const
+	{
+		return UniformCuttingCoordinates[Iso];
+	}
 
-		/**
-		 * @return the FPoint2D (parametric coordinates) of the point at the Index of the grid in the defined grid space 
-		 * @see EGridSpace
-		 * @see Points2D
-		 */
-		const FPoint2D& GetInner2DPoint(EGridSpace Space, int32 Index) const
-		{
-			return Points2D[(int32)Space][Index];
-		}
+	const FCoordinateGrid& GetUniformCuttingCoordinates() const
+	{
+		return UniformCuttingCoordinates;
+	}
 
-		void SetInner2DPoint(EGridSpace Space, int32 Index, const FPoint2D& NewCoordinate)
-		{
-			Points2D[(int32)Space][Index] = NewCoordinate;
-		}
+	/**
+	 * @return the array of 3d points of the grid
+	 */
+	TArray<FPoint>& GetInner3DPoints()
+	{
+		return Points3D;
+	}
 
-		/**
-		 * @return the FPoint2D (parametric coordinates) of the point at the Index of the grid in the defined grid space @see EGridSpace
-		 */
-		const FPoint2D& GetInner2DPoint(EGridSpace Space, int32 IndexU, int32 IndexV) const
-		{
-			return Points2D[(int32)Space][GobalIndex(IndexU, IndexV)];
-		}
+	/**
+	 * @return the array of 3d points of the grid
+	 */
+	const TArray<FPoint>& GetInner3DPoints() const
+	{
+		return Points3D;
+	}
 
-		/**
-		 * @return the FPoint (3D coordinates) of the point at the Index of the grid
-		 */
-		const FPoint& GetInner3DPoint(int32 Index) const
-		{
-			return Points3D[Index];
-		}
+	/**
+	 * @return the array of 2d points of the grid in the defined space
+	 */
+	const TArray<FPoint2D>& GetInner2DPoints(EGridSpace Space) const
+	{
+		return Points2D[(int32)Space];
+	}
 
-		/**
-		 * @return the FPoint (3D coordinates) of the point at the Index of the grid
-		 */
-		const FPoint& GetInner3DPoint(int32 IndexU, int32 IndexV) const
-		{
-			return Points3D[GobalIndex(IndexU, IndexV)];
-		}
+	/**
+	 * @return the array of normal of the points of the grid
+	 */
+	TArray<FVector>& GetNormals()
+	{
+		return Normals;
+	}
 
-		/**
-		 * @return the normal of the surface at the point at the Index of the grid
-		 */
-		const FVector& GetPointNormal(int32 IndexU, int32 IndexV) const
-		{
-			return Normals[GobalIndex(IndexU, IndexV)];
-		}
+	TSharedRef<FTopologicalFace> GetFace() const
+	{
+		return Face;
+	}
 
-		/**
-		 * @return the normal of the surface at the point at the Index of the grid
-		 */
-		const FVector& GetPointNormal(int32 Index) const
-		{
-			return Normals[Index];
-		}
+	const TArray<TArray<int32>>& GetNodeIdsOfFaceLoops() const
+	{
+		return NodeIdsOfFaceLoops;
+	}
 
-		constexpr const TArray<double>& GetCuttingCoordinatesAlongIso(EIso Iso) const
-		{
-			return CuttingCoordinates[Iso];
-		}
+	const FPoint2D& GetLoop2DPoint(EGridSpace Space, int32 LoopIndex, int32 Index) const
+	{
+		return FaceLoops2D[(int32)Space][LoopIndex][Index];
+	}
 
-		const FCoordinateGrid& GetCuttingCoordinates() const
-		{
-			return CuttingCoordinates;
-		}
+	void SetLoop2DPoint(EGridSpace Space, int32 LoopIndex, int32 Index, const FPoint2D& NewCoordinate)
+	{
+		FaceLoops2D[(int32)Space][LoopIndex][Index] = NewCoordinate;
+	}
 
-		constexpr const TArray<double>& GetUniformCuttingCoordinatesAlongIso(EIso Iso) const
-		{
-			return UniformCuttingCoordinates[Iso];
-		}
+	const FPoint& GetLoop3DPoint(int32 LoopIndex, int32 Index) const
+	{
+		return FaceLoops3D[LoopIndex][Index];
+	}
 
-		const FCoordinateGrid& GetUniformCuttingCoordinates() const
-		{
-			return UniformCuttingCoordinates;
-		}
+	const int32 GetLoopCount() const
+	{
+		return FaceLoops2D[0].Num();
+	}
 
-		/**
-		 * @return the array of 3d points of the grid
-		 */
-		TArray<FPoint>& GetInner3DPoints()
-		{
-			return Points3D;
-		}
+	/**
+	 * @return the array of array of 2d points of the loops according to the defined grid space
+	 */
+	const TArray<TArray<FPoint2D>>& GetLoops2D(EGridSpace Space) const
+	{
+		return FaceLoops2D[(int32)Space];
+	}
 
-		/**
-		 * @return the array of 3d points of the grid
-		 */
-		const TArray<FPoint>& GetInner3DPoints() const
-		{
-			return Points3D;
-		}
+	/**
+	 * @return the array of array of 3d points of the loops
+	 */
+	TArray<TArray<FPoint>>& GetLoops3D()
+	{
+		return FaceLoops3D;
+	}
 
-		/**
-		 * @return the array of 2d points of the grid in the defined space
-		 */
-		const TArray<FPoint2D>& GetInner2DPoints(EGridSpace Space) const
-		{
-			return Points2D[(int32)Space];
-		}
+	/**
+	 * @return the array of array of 3d points of the loops
+	 */
+	const TArray<TArray<FPoint>>& GetLoops3D() const
+	{
+		return FaceLoops3D;
+	}
 
-		/**
-		 * @return the array of normal of the points of the grid
-		 */
-		TArray<FVector>& GetNormals()
-		{
-			return Normals;
-		}
+	/**
+	 * @return the array of array of normal of the points of the loops
+	 */
+	const TArray<TArray<FVector>>& GetLoopNormals() const
+	{
+		return NormalsOfFaceLoops;
+	}
 
-		TSharedRef<FTopologicalFace> GetFace() const
-		{
-			return Face;
-		}
+	/**
+	 * @return the Index of the position in the arrays of a point [IndexU, IndexV] of the grid
+	 */
+	int32 GobalIndex(int32 IndexU, int32 IndexV) const
+	{
+		return IndexV * CuttingCount[EIso::IsoU] + IndexU;
+	}
 
-		const TArray<TArray<int32>>& GetNodeIdsOfFaceLoops() const
-		{
-			return NodeIdsOfFaceLoops;
-		}
+	/**
+	 * @return the Index of the position in the arrays of a point [IndexU, IndexV] of the grid
+	 */
+	void UVIndexFromGlobalIndex(int32 GLobalIndex, int32& OutIndexU, int32& OutIndexV) const
+	{
+		OutIndexU = GLobalIndex % CuttingCount[EIso::IsoU];
+		OutIndexV = GLobalIndex / CuttingCount[EIso::IsoU];
+	}
+
+	/**
+	 * @return the maximum difference of coordinate along the specified axis of two two successive points
+	 */
+	double GetMaxDeltaU(EIso Iso) const
+	{
+		return MaxDeltaUV[Iso];
+	}
 
-		const FPoint2D& GetLoop2DPoint(EGridSpace Space, int32 LoopIndex, int32 Index) const
+	/**
+	 * @return the minimal size of the mesh elements according to 3d geometric tolerance
+	 */
+	double GetMinElementSize() const
+	{
+		return MinimumElementSize;
+	}
+
+	// ======================================================================================================================================================================================================================
+	// Display Methodes   ================================================================================================================================================================================================
+	// ======================================================================================================================================================================================================================
+#ifdef CADKERNEL_DEV
+	bool bDisplay = false;
+
+	void DisplayIsoNode(EGridSpace Space, const int32 PointIndex, FIdent Ident = 0, EVisuProperty Property = EVisuProperty::BluePoint) const;
+	void DisplayIsoNode(EGridSpace Space, const FIsoNode& Node, FIdent Ident = 0, EVisuProperty Property = EVisuProperty::BluePoint) const;
+	void DisplayIsoNodes(EGridSpace Space, const TArray<FIsoNode*>& Nodes, EVisuProperty Property = EVisuProperty::BluePoint) const;
+	void DisplayIsoNodes(const FString& Message, EGridSpace Space, const TArray<FIsoNode*>& Nodes, EVisuProperty Property = EVisuProperty::BluePoint) const
+	{
+		if (!bDisplay)
 		{
-			return FaceLoops2D[(int32)Space][LoopIndex][Index];
+			return;
 		}
+		F3DDebugSession _(Message);
+		DisplayIsoNodes(Space, Nodes, Property);
+	}
 
-		void SetLoop2DPoint(EGridSpace Space, int32 LoopIndex, int32 Index, const FPoint2D& NewCoordinate)
-		{
-			FaceLoops2D[(int32)Space][LoopIndex][Index] = NewCoordinate;
-		}
+	void DisplayIsoSegment(EGridSpace Space, const FIsoSegment& Segment, FIdent Ident = 0, EVisuProperty Property = EVisuProperty::BlueCurve, bool bDisplayOrientation = false) const;
+	void DisplayIsoSegment(EGridSpace Space, const FIsoNode& NodeA, const FIsoNode& NodeB, FIdent Ident = 0, EVisuProperty Property = EVisuProperty::BlueCurve) const;
 
-		const FPoint& GetLoop3DPoint(int32 LoopIndex, int32 Index) const
+	void DisplayIsoSegments(EGridSpace Space, const TArray<FIsoSegment*>& Segments, bool bDisplayNode = false, bool bDisplayOrientation = false, EVisuProperty Property = EVisuProperty::BlueCurve) const;
+	void DisplayIsoSegments(const FString& Message, EGridSpace Space, const TArray<FIsoSegment*>& Segments, bool bDisplayNode = false, bool bDisplayOrientation = false, EVisuProperty Property = EVisuProperty::BlueCurve) const
+	{
+		if (!bDisplay)
 		{
-			return FaceLoops3D[LoopIndex][Index];
+			return;
 		}
+		F3DDebugSession _(Message);
+		DisplayIsoSegments(Space, Segments, bDisplayNode, bDisplayOrientation, Property);
+	}
 
-		const int32 GetLoopCount() const
-		{
-			return FaceLoops2D[0].Num();
-		}
+	void DisplayTriangle(EGridSpace Space, const FIsoNode& NodeA, const FIsoNode& NodeB, const FIsoNode& NodeC) const;
 
-		/**
-		 * @return the array of array of 2d points of the loops according to the defined grid space
-		 */
-		const TArray<TArray<FPoint2D>>& GetLoops2D(EGridSpace Space) const
+	void DisplayGridPolyline(EGridSpace Space, const FLoopNode& StartNode, bool bDisplayNode, EVisuProperty Property = EVisuProperty::BlueCurve) const;
+	void DisplayGridPolyline(EGridSpace Space, const TArray<FIsoNode*>& Nodes, bool bDisplayNode, EVisuProperty Property = EVisuProperty::BlueCurve) const;
+	void DisplayGridPolyline(EGridSpace Space, const TArray<FLoopNode*>& Nodes, bool bDisplayNode, EVisuProperty Property = EVisuProperty::BlueCurve) const;
+	void DisplayGridPolyline(EGridSpace Space, const TArray<FLoopNode>& Nodes, bool bDisplayNode, EVisuProperty Property = EVisuProperty::BlueCurve) const;
+	void DisplayGridPolyline(const FString& Message, EGridSpace Space, const TArray<FLoopNode*>& Nodes, bool bDisplayNode, EVisuProperty Property = EVisuProperty::BlueCurve) const
+	{
+		if (!bDisplay)
 		{
-			return FaceLoops2D[(int32)Space];
+			return;
 		}
+		F3DDebugSession _(Message);
+		DisplayGridPolyline(Space, Nodes, bDisplayNode, Property);
+	}
 
-		/**
-		 * @return the array of array of 3d points of the loops
-		 */
-		TArray<TArray<FPoint>>& GetLoops3D()
+	void DisplayGridPolyline(const FString& Message, EGridSpace Space, const TArray<FLoopNode>& Nodes, bool bDisplayNode, EVisuProperty Property = EVisuProperty::BlueCurve) const
+	{
+		if (!bDisplay)
 		{
-			return FaceLoops3D;
+			return;
 		}
+		F3DDebugSession _(Message);
+		DisplayGridPolyline(Space, Nodes, bDisplayNode, Property);
+	}
 
-		/**
-		 * @return the array of array of 3d points of the loops
-		 */
-		const TArray<TArray<FPoint>>& GetLoops3D() const
-		{
-			return FaceLoops3D;
-		}
+	void DisplayNodes(const TCHAR* Message, EGridSpace Space, const TArray<const FIsoNode*>& Nodes, EVisuProperty Property) const;
 
-		/**
-		 * @return the array of array of normal of the points of the loops
-		 */
-		const TArray<TArray<FVector>>& GetLoopNormals() const
-		{
-			return NormalsOfFaceLoops;
-		}
+	void DisplayGridPoints(EGridSpace DisplaySpace) const;
+	void DisplayInnerPoints(TCHAR* Message, EGridSpace DisplaySpace) const;
 
-		/**
-		 * @return the Index of the position in the arrays of a point [IndexU, IndexV] of the grid
-		 */
-		int32 GobalIndex(int32 IndexU, int32 IndexV) const
-		{
-			return IndexV * CuttingCount[EIso::IsoU] + IndexU;
-		}
+	void DisplayGridNormal() const;
 
-		/**
-		 * @return the Index of the position in the arrays of a point [IndexU, IndexV] of the grid
-		 */
-		void UVIndexFromGlobalIndex(int32 GLobalIndex, int32& OutIndexU, int32& OutIndexV) const
+	template<typename TPoint>
+	void DisplayInnerDomainPoints(FString Message, const TArray<TPoint>& Points) const
+	{
+		if (!bDisplay)
 		{
-			OutIndexU = GLobalIndex % CuttingCount[EIso::IsoU];
-			OutIndexV = GLobalIndex / CuttingCount[EIso::IsoU];
+			return;
 		}
 
-		/**
-		 * @return the maximum difference of coordinate along the specified axis of two two successive points
-		 */
-		double GetMaxDeltaU(EIso Iso) const
+		F3DDebugSession _(Message);
+		for (int32 Index = 0; Index < Points.Num(); ++Index)
 		{
-			return MaxDeltaUV[Iso];
+			if (IsInsideFace[Index])
+			{
+				DisplayPoint(Points[Index], IsCloseToLoop[Index] ? EVisuProperty::BluePoint : EVisuProperty::GreenPoint, Index);
+			}
+			else
+			{
+				DisplayPoint(Points[Index], EVisuProperty::OrangePoint, Index);
+			}
 		}
+	}
 
-		/**
-		 * @return the minimal size of the mesh elements according to 3d geometric tolerance
-		 */
-		double GetMinElementSize() const
+	template<typename TPoint>
+	void DisplayGridLoop(FString Message, const TArray<TArray<TPoint>>& Loops, bool bDisplayNodes, bool bMakeGroup) const
+	{
+		if (!bDisplay)
 		{
-			return MinimumElementSize;
+			return;
 		}
 
-		// ======================================================================================================================================================================================================================
-		// Display Methodes   ================================================================================================================================================================================================
-		// ======================================================================================================================================================================================================================
-		bool bDisplay = false;
-
-		void DisplayGridPoints(EGridSpace DisplaySpace) const;
-		void DisplayGridInnerPoints(EGridSpace DisplaySpace, TCHAR* Message) const;
-
-		void DisplayGridNormal() const;
-
-		template<typename TPoint>
-		void DisplayInnerDomainPoints(FString Message, const TArray<TPoint>& Points) const
+		F3DDebugSession _(*Message);
+		int32 LoopIndex = 0;
+		for (const TArray<TPoint>& Loop : Loops)
 		{
-			if (!bDisplay)
-			{
-				return;
-			}
-
-			F3DDebugSession _(Message);
-			for (int32 Index = 0; Index < Points.Num(); ++Index)
+			if (bMakeGroup)
 			{
-				if (IsInsideFace[Index])
-				{
-					if (IsCloseToLoop[Index])
-					{
-						DisplayPoint(Points[Index], EVisuProperty::BluePoint, Index);
-					}
-					else
-					{
-						DisplayPoint(Points[Index], EVisuProperty::OrangePoint, Index);
-					}
-				}
+				Open3DDebugSession(FString::Printf(TEXT("Loop %d"), LoopIndex++));
 			}
-		}
 
-		template<typename TPoint>
-		void DisplayLoop(FString Message, const TArray<TArray<TPoint>>& Loops, bool bDisplayNodes, bool bMakeGroup) const
-		{
-			if (!bDisplay)
+			const TPoint* FirstSegmentPoint = &Loop[0];
+			if (bDisplayNodes)
 			{
-				return;
+				DisplayPoint(*FirstSegmentPoint, EVisuProperty::BluePoint, 0);
 			}
 
-			F3DDebugSession _(*Message);
-			int32 LoopIndex = 0;
-			for (const TArray<TPoint>& Loop : Loops)
+			for (int32 Index = 1; Index < Loop.Num(); ++Index)
 			{
-				if (bMakeGroup)
-				{
-					Open3DDebugSession(FString::Printf(TEXT("Loop %d"), LoopIndex++));
-				}
-
-				const TPoint* FirstSegmentPoint = &Loop[0];
+				const TPoint* SecondSegmentPoint = &Loop[Index];
+				DisplaySegment(*FirstSegmentPoint, *SecondSegmentPoint);
 				if (bDisplayNodes)
 				{
-					DisplayPoint(*FirstSegmentPoint, EVisuProperty::BluePoint, 0);
-				}
-
-				for (int32 Index = 1; Index < Loop.Num(); ++Index)
-				{
-					const TPoint* SecondSegmentPoint = &Loop[Index];
-					DisplaySegment(*FirstSegmentPoint, *SecondSegmentPoint);
-					if (bDisplayNodes)
-					{
-						DisplayPoint(*SecondSegmentPoint, EVisuProperty::BluePoint, Index);
-					}
-					FirstSegmentPoint = SecondSegmentPoint;
-				}
-				DisplaySegment(*FirstSegmentPoint, *Loop.begin());
-				if (bMakeGroup)
-				{
-					Close3DDebugSession();
+					DisplayPoint(*SecondSegmentPoint, EVisuProperty::BluePoint, Index);
 				}
+				FirstSegmentPoint = SecondSegmentPoint;
+			}
+			DisplaySegment(*FirstSegmentPoint, *Loop.begin());
+			if (bMakeGroup)
+			{
+				Close3DDebugSession();
 			}
 		}
+	}
 
-		void DisplayFindPointsCloseToLoop(EGridSpace DisplaySpace) const;
-		void PrintTimeElapse() const;
-	};
+	void DisplayGridLoops(FString Message, EGridSpace DisplaySpace, bool bDisplayNodes, bool bMakeGroup) const
+	{
+		DisplayGridLoop(Message, FaceLoops2D[DisplaySpace], bDisplayNodes, bMakeGroup);
+	}
+
+#endif
+
+	void PrintTimeElapse() const;
+};
 
 }
diff --git a/Engine/Source/Runtime/Datasmith/CADKernel/Private/CADKernel/Topo/TopologicalLoop.cpp b/Engine/Source/Runtime/Datasmith/CADKernel/Private/CADKernel/Topo/TopologicalLoop.cpp
--- a/Engine/Source/Runtime/Datasmith/CADKernel/Private/CADKernel/Topo/TopologicalLoop.cpp
+++ b/Engine/Source/Runtime/Datasmith/CADKernel/Private/CADKernel/Topo/TopologicalLoop.cpp
@@ -138,7 +138,7 @@ bool FTopologicalLoop::Get2DSamplingWithoutDegeneratedEdges(TArray<FPoint2D>& Lo
 	double LoopLength = 0;
 	int32 EdgeCount = 0;
 	int32 PointCount = 0;
-	for (const FOrientedEdge& Edge : Edges)
+ 	for (const FOrientedEdge& Edge : Edges)
 	{
 		if (Edge.Entity->IsDegenerated())
 		{
@@ -151,7 +151,7 @@ bool FTopologicalLoop::Get2DSamplingWithoutDegeneratedEdges(TArray<FPoint2D>& Lo
 
 	double LoopMeanLength = LoopLength / EdgeCount;
 	double MinEdgeLength = LoopMeanLength * 0.01;
-	MinEdgeLength = FMath::Max(MinEdgeLength, Face.Pin()->GetCarrierSurface()->Get3DTolerance() * 10);
+	MinEdgeLength = FMath::Max(MinEdgeLength, Face.Pin()->GetCarrierSurface()->Get3DTolerance());
 
 	LoopSampling.Empty(PointCount);
 
@@ -219,6 +219,13 @@ void FTopologicalLoop::Orient()
 		return;
 	}
 
+#ifdef DEBUG_ORIENT
+	{
+		F3DDebugSession _(*FString::Printf(TEXT("Loop before orientation")));
+		DisplayOrientedPolyline(LoopSampling, EVisuProperty::BlueCurve);
+	}
+#endif
+
 	LoopSampling.Pop();
 	TSet<int32> ExtremityIndex;
 	ExtremityIndex.Reserve(8);
@@ -456,6 +463,20 @@ void FTopologicalLoop::Orient()
 	{
 		SwapOrientation();
 	}
+
+#ifdef DEBUG_ORIENT
+	{
+		LoopSampling.Empty();
+		if (!Get2DSamplingWithoutDegeneratedEdges(LoopSampling))
+		{
+			return;
+		}
+
+		F3DDebugSession _(*FString::Printf(TEXT("Loop oriented")));
+		DisplayOrientedPolyline(LoopSampling, EVisuProperty::BlueCurve);
+	}
+#endif
+
 }
 
 void FTopologicalLoop::SwapOrientation()
