CADKernel Update

diff --git a/Engine/Source/Runtime/Datasmith/CADKernel/Private/CADKernel/UI/Display.cpp b/Engine/Source/Runtime/Datasmith/CADKernel/Private/CADKernel/UI/Display.cpp
--- a/Engine/Source/Runtime/Datasmith/CADKernel/Private/CADKernel/UI/Display.cpp
+++ b/Engine/Source/Runtime/Datasmith/CADKernel/Private/CADKernel/UI/Display.cpp
@@ -30,1174 +30,1176 @@
 
 namespace CADKernel
 {
+
 #ifdef CADKERNEL_DEV
-	void Open3DDebugSession(FString name, const TArray<FIdent>& IdArray)
-	{
-		FSystem::Get().GetVisu()->Open3DDebugSession(*name, IdArray);
-	}
+void Open3DDebugSession(FString name, const TArray<FIdent>& IdArray)
+{
+	FSystem::Get().GetVisu()->Open3DDebugSession(*name, IdArray);
+}
 
-	void Close3DDebugSession(bool bIsDisplayed)
+void Close3DDebugSession(bool bIsDisplayed)
+{
+	if (bIsDisplayed)
 	{
-		if (bIsDisplayed)
-		{
-			FSystem::Get().GetVisu()->Close3DDebugSession();
-		}
+		FSystem::Get().GetVisu()->Close3DDebugSession();
 	}
+}
 
-	void Wait(bool bMakeWait)
+void Wait(bool bMakeWait)
+{
+	if (bMakeWait)
 	{
-		if (bMakeWait)
-		{
-			system("PAUSE");
-			printf("Hello\n");
-		}
+		system("PAUSE");
+		printf("Hello\n");
 	}
+}
 #endif
 
-	void Open3DDebugSegment(FIdent Ident)
-	{
+void Open3DDebugSegment(FIdent Ident)
+{
 #ifdef CADKERNEL_DEV
-		FSystem::Get().GetVisu()->Open3DDebugSegment(Ident);
+	FSystem::Get().GetVisu()->Open3DDebugSegment(Ident);
 #endif
-	}
+}
 
-	void Close3DDebugSegment()
-	{
+void Close3DDebugSegment()
+{
 #ifdef CADKERNEL_DEV
-		FSystem::Get().GetVisu()->Close3DDebugSegment();
+	FSystem::Get().GetVisu()->Close3DDebugSegment();
 #endif
-	}
+}
 
-	void FlushVisu()
-	{
+void FlushVisu()
+{
 #ifdef CADKERNEL_DEV
-		FSystem::Get().GetVisu()->UpdateViewer();
+	FSystem::Get().GetVisu()->UpdateViewer();
 #endif
-	}
+}
 
-	void DrawElement(int32 Dimension, TArray<FPoint>& Points, EVisuProperty Property)
-	{
+void DrawElement(int32 Dimension, TArray<FPoint>& Points, EVisuProperty Property)
+{
 #ifdef CADKERNEL_DEV
-		FSystem::Get().GetVisu()->DrawElement(Dimension, Points, Property);
+	FSystem::Get().GetVisu()->DrawElement(Dimension, Points, Property);
 #endif
-	}
+}
 
-	void DrawMesh(const FMesh& Mesh)
-	{
+void DrawMesh(const FMesh& Mesh)
+{
 #ifdef CADKERNEL_DEV
-		FSystem::Get().GetVisu()->DrawMesh(Mesh.GetId());
+	FSystem::Get().GetVisu()->DrawMesh(Mesh.GetId());
 #endif
-	}
+}
 
-	void Draw(const FLinearBoundary& Boundary, const FRestrictionCurve& Curve, EVisuProperty Property)
-	{
+void Draw(const FLinearBoundary& Boundary, const FRestrictionCurve& Curve, EVisuProperty Property)
+{
 #ifdef CADKERNEL_DEV
-		if (Boundary.IsDegenerated())
-		{
-			return;
-		}
+	if (Boundary.IsDegenerated())
+	{
+		return;
+	}
 
-		bool bShowOrientation = FSystem::Get().GetVisu()->GetParameters()->bDisplayCADOrient;
+	bool bShowOrientation = FSystem::Get().GetVisu()->GetParameters()->bDisplayCADOrient;
 
-		TArray<FPoint> Polyline;
-		Curve.GetDiscretizationPoints<FPoint>(Boundary, EOrientation::Front, Polyline);
-		Draw(Polyline, Property);
+	TArray<FPoint> Polyline;
+	Curve.GetDiscretizationPoints<FPoint>(Boundary, EOrientation::Front, Polyline);
+	Draw(Polyline, Property);
 
-		if (bShowOrientation)
+	if (bShowOrientation)
+	{
+		double Length = 0;
+		for (int32 Index = 1; Index < Polyline.Num(); ++Index)
 		{
-			double Length = 0;
-			for (int32 Index = 1; Index < Polyline.Num(); ++Index)
-			{
-				Length += Polyline[Index].Distance(Polyline[Index - 1]);
-			}
+			Length += Polyline[Index].Distance(Polyline[Index - 1]);
+		}
 
-			double Coordinate = (Boundary.Max + Boundary.Min) / 2.;
-			FCurvePoint Point;
-			Curve.EvaluatePoint(Coordinate, Point, 1);
+		double Coordinate = (Boundary.Max + Boundary.Min) / 2.;
+		FCurvePoint Point;
+		Curve.EvaluatePoint(Coordinate, Point, 1);
 
-			double Height = Length / 20.0;
-			double Base = Height / 2;
+		double Height = Length / 20.0;
+		double Base = Height / 2;
 
-			DrawQuadripode(Height, Base, Point.Point, Point.Gradient);
-		}
-#endif
+		DrawQuadripode(Height, Base, Point.Point, Point.Gradient);
 	}
+#endif
+}
 
-	void Draw2D(const FCurve& Curve, const FLinearBoundary& Boundary, EVisuProperty Property)
-	{
+void Draw2D(const FCurve& Curve, const FLinearBoundary& Boundary, EVisuProperty Property)
+{
 #ifdef CADKERNEL_DEV
-		if (Boundary.IsDegenerated())
-		{
-			return;
-		}
+	if (Boundary.IsDegenerated())
+	{
+		return;
+	}
 
-		double DiscVisu = ((double)(FSystem::Get().GetVisu()->GetParameters()->ChordError)) / 10.;
-		bool bVoirOrientation = FSystem::Get().GetVisu()->GetParameters()->bDisplayCADOrient;
+	double DiscVisu = ((double)(FSystem::Get().GetVisu()->GetParameters()->ChordError)) / 10.;
+	bool bVoirOrientation = FSystem::Get().GetVisu()->GetParameters()->bDisplayCADOrient;
 
-		FPolyline2D Polyline;
-		FCurve2DSamplerOnChord Sampler(Curve, Boundary, DiscVisu, Polyline);
-		Sampler.Sample();
+	FPolyline2D Polyline;
+	FCurve2DSamplerOnChord Sampler(Curve, Boundary, DiscVisu, Polyline);
+	Sampler.Sample();
 
-		Draw(Polyline.GetPoints(), Property);
+	Draw(Polyline.GetPoints(), Property);
 
-		if (bVoirOrientation)
-		{
-			double Length = Polyline.GetLength(Boundary);
+	if (bVoirOrientation)
+	{
+		double Length = Polyline.GetLength(Boundary);
 
-			double Coordinate = (Boundary.Max + Boundary.Min) / 2.;
-			FCurvePoint Point;
-			Curve.EvaluatePoint(Coordinate, Point, 1);
+		double Coordinate = (Boundary.Max + Boundary.Min) / 2.;
+		FCurvePoint Point;
+		Curve.EvaluatePoint(Coordinate, Point, 1);
 
-			double Height = Length / 20.0;
-			double Base = Height / 2;
+		double Height = Length / 20.0;
+		double Base = Height / 2;
 
-			DrawQuadripode(Height, Base, Point.Point, Point.Gradient);
-		}
-#endif
+		DrawQuadripode(Height, Base, Point.Point, Point.Gradient);
 	}
+#endif
+}
 
-	void Draw3D(const FCurve& Curve, const FLinearBoundary& Boundary, EVisuProperty Property)
-	{
+void Draw3D(const FCurve& Curve, const FLinearBoundary& Boundary, EVisuProperty Property)
+{
 #ifdef CADKERNEL_DEV
-		if (Boundary.IsDegenerated())
-		{
-			return;
-		}
+	if (Boundary.IsDegenerated())
+	{
+		return;
+	}
 
-		double DiscVisu = FSystem::Get().GetVisu()->GetParameters()->ChordError;
-		bool bVoirOrientation = FSystem::Get().GetVisu()->GetParameters()->bDisplayCADOrient;
+	double DiscVisu = FSystem::Get().GetVisu()->GetParameters()->ChordError;
+	bool bVoirOrientation = FSystem::Get().GetVisu()->GetParameters()->bDisplayCADOrient;
 
-		FPolyline3D Polyline;
-		FCurveSamplerOnChord Sampler(Curve, Boundary, DiscVisu, Polyline);
-		Sampler.Sample();
+	FPolyline3D Polyline;
+	FCurveSamplerOnChord Sampler(Curve, Boundary, DiscVisu, Polyline);
+	Sampler.Sample();
 
-		Draw(Polyline.GetPoints(), Property);
+	Draw(Polyline.GetPoints(), Property);
 
-		if (bVoirOrientation)
-		{
-			double Length = Polyline.GetLength(Boundary);
+	if (bVoirOrientation)
+	{
+		double Length = Polyline.GetLength(Boundary);
 
-			double Coordinate = (Boundary.Max + Boundary.Min) / 2.;
-			FCurvePoint Point;
-			Curve.EvaluatePoint(Coordinate, Point, 1);
+		double Coordinate = (Boundary.Max + Boundary.Min) / 2.;
+		FCurvePoint Point;
+		Curve.EvaluatePoint(Coordinate, Point, 1);
 
-			double Height = Length / 20.0;
-			double Base = Height / 2;
+		double Height = Length / 20.0;
+		double Base = Height / 2;
 
-			DrawQuadripode(Height, Base, Point.Point, Point.Gradient);
-		}
-#endif
+		DrawQuadripode(Height, Base, Point.Point, Point.Gradient);
 	}
+#endif
+}
 
-	void Draw(const FCurve& Curve, const FLinearBoundary& Boundary, EVisuProperty Property)
-	{
+void Draw(const FCurve& Curve, const FLinearBoundary& Boundary, EVisuProperty Property)
+{
 #ifdef CADKERNEL_DEV
-		if (Curve.GetDimension() == 3)
-		{
-			Draw3D(Curve, Boundary, Property);
-		}
-		else
-		{
-			Draw2D(Curve, Boundary, Property);
-		}
-#endif
+	if (Curve.GetDimension() == 3)
+	{
+		Draw3D(Curve, Boundary, Property);
+	}
+	else
+	{
+		Draw2D(Curve, Boundary, Property);
 	}
+#endif
+}
 
 
-	void Draw(const FCurve& Curve, EVisuProperty Property)
-	{
+void Draw(const FCurve& Curve, EVisuProperty Property)
+{
 #ifdef CADKERNEL_DEV
-		const FLinearBoundary& Bounds = Curve.GetBoundary();
-		Draw(Curve, Bounds, Property);
+	const FLinearBoundary& Bounds = Curve.GetBoundary();
+	Draw(Curve, Bounds, Property);
 #endif
-	}
+}
 
-	void DrawQuadripode(double Height, double Base, FPoint& Center, FPoint& InDirection)
-	{
+void DrawQuadripode(double Height, double Base, FPoint& Center, FPoint& InDirection)
+{
 #ifdef CADKERNEL_DEV
-		FPoint Direction = InDirection;
-		Direction.Normalize();
+	FPoint Direction = InDirection;
+	Direction.Normalize();
 
-		FPoint Normal(-Direction[1], Direction[0], 0.0);
+	FPoint Normal(-Direction[1], Direction[0], 0.0);
 
-		double UVNorm = Normal.Length();
-		Normal /= UVNorm;
+	double UVNorm = Normal.Length();
+	Normal /= UVNorm;
 
-		FPoint BiNormal(-Direction[2] * Normal[1], Direction[2] * Normal[0], UVNorm);
+	FPoint BiNormal(-Direction[2] * Normal[1], Direction[2] * Normal[0], UVNorm);
 
-		FPoint Point0 = Center;
-		FPoint PointBase = Point0 - Direction * Height;
+	FPoint Point0 = Center;
+	FPoint PointBase = Point0 - Direction * Height;
 
-		Normal *= Base;
-		BiNormal *= Base;
+	Normal *= Base;
+	BiNormal *= Base;
 
-		FPoint Point1 = PointBase + Normal;
-		FPoint Point2 = PointBase + BiNormal;
-		FPoint Point3 = PointBase - Normal;
-		FPoint Point4 = PointBase - BiNormal;
+	FPoint Point1 = PointBase + Normal;
+	FPoint Point2 = PointBase + BiNormal;
+	FPoint Point3 = PointBase - Normal;
+	FPoint Point4 = PointBase - BiNormal;
 
-		TArray<FPoint> Polygone;
-		Polygone.Reserve(7);
+	TArray<FPoint> Polygone;
+	Polygone.Reserve(7);
 
-		Polygone.Add(Point1);
-		Polygone.Add(Point0);
-		Polygone.Add(Point2);
-		Polygone.Add(Point0);
-		Polygone.Add(Point3);
-		Polygone.Add(Point0);
-		Polygone.Add(Point4);
+	Polygone.Add(Point1);
+	Polygone.Add(Point0);
+	Polygone.Add(Point2);
+	Polygone.Add(Point0);
+	Polygone.Add(Point3);
+	Polygone.Add(Point0);
+	Polygone.Add(Point4);
 
-		Draw(Polygone);
+	Draw(Polygone);
 #endif
-	}
+}
 
-	void DisplayEntity(const FEntity& Entity)
-	{
+void DisplayEntity(const FEntity& Entity)
+{
 #ifdef CADKERNEL_DEV
-		FTimePoint StartTime = FChrono::Now();
+	FTimePoint StartTime = FChrono::Now();
 
-		FProgress Progress;
+	FProgress Progress;
 
-		F3DDebugSession GraphicSession(FString::Printf(TEXT("%s %d"), Entity.GetTypeName(), Entity.GetId()), { Entity.GetId() });
+	F3DDebugSession GraphicSession(FString::Printf(TEXT("%s %d"), Entity.GetTypeName(), Entity.GetId()), { Entity.GetId() });
 
-		switch (Entity.GetEntityType())
-		{
-		case EEntity::TopologicalVertex:
-			Display((const FTopologicalVertex&) Entity);
-			break;
-		case EEntity::Curve:
-			Display((const FCurve&) Entity);
-			break;
-		case EEntity::Surface:
-			Display((const FSurface&) Entity);
-			break;
-		case EEntity::TopologicalFace:
-			Display((const FTopologicalFace&) Entity);
-			break;
-		case EEntity::TopologicalLoop:
-			Display((const FTopologicalLoop&) Entity);
-			break;
-		case EEntity::Shell:
-			Display((const FShell&) Entity);
-			break;
-		case EEntity::Body:
-			Display((const FBody&) Entity);
-			break;
-		case EEntity::TopologicalEdge:
-			Display((const FTopologicalEdge&) Entity);
-			break;
-		case EEntity::Model:
-			Display((const FModel&) Entity);
-			break;
-		case EEntity::Group:
-			Display((const FGroup&) Entity);
-			break;
-		case EEntity::MeshModel:
-			Display((const FModelMesh&) Entity);
-			break;
-		default:
-			FMessage::Printf(Log, TEXT("Unable to display Entity of type %s"), FEntity::GetTypeName(Entity.GetEntityType()));
-		}
+	switch (Entity.GetEntityType())
+	{
+	case EEntity::TopologicalVertex:
+		Display((const FTopologicalVertex&)Entity);
+		break;
+	case EEntity::Curve:
+		Display((const FCurve&)Entity);
+		break;
+	case EEntity::Surface:
+		Display((const FSurface&)Entity);
+		break;
+	case EEntity::TopologicalFace:
+		Display((const FTopologicalFace&)Entity);
+		break;
+	case EEntity::TopologicalLoop:
+		Display((const FTopologicalLoop&)Entity);
+		break;
+	case EEntity::Shell:
+		Display((const FShell&)Entity);
+		break;
+	case EEntity::Body:
+		Display((const FBody&)Entity);
+		break;
+	case EEntity::TopologicalEdge:
+		Display((const FTopologicalEdge&)Entity);
+		break;
+	case EEntity::Model:
+		Display((const FModel&)Entity);
+		break;
+	case EEntity::Group:
+		Display((const FGroup&)Entity);
+		break;
+	case EEntity::MeshModel:
+		Display((const FModelMesh&)Entity);
+		break;
+	default:
+		FMessage::Printf(Log, TEXT("Unable to display Entity of type %s"), FEntity::GetTypeName(Entity.GetEntityType()));
+	}
 
-		FDuration DisplayDuration = FChrono::Elapse(StartTime);
-		FChrono::PrintClockElapse(Log, TEXT("  "), TEXT("Display "), DisplayDuration);
+	FDuration DisplayDuration = FChrono::Elapse(StartTime);
+	FChrono::PrintClockElapse(Log, TEXT("  "), TEXT("Display "), DisplayDuration);
 
 #endif
-	}
+}
 
-	void DisplayEntity2D(const FEntity& Entity)
-	{
+void DisplayEntity2D(const FEntity& Entity)
+{
 #ifdef CADKERNEL_DEV
-		FTimePoint StartTime = FChrono::Now();
+	FTimePoint StartTime = FChrono::Now();
 
-		FProgress Progress;
+	FProgress Progress;
 
-		F3DDebugSession GraphicSession(FString::Printf(TEXT("%s %d"), Entity.GetTypeName(), Entity.GetId()), { Entity.GetId() });
+	F3DDebugSession GraphicSession(FString::Printf(TEXT("%s %d"), Entity.GetTypeName(), Entity.GetId()), { Entity.GetId() });
 
-		switch (Entity.GetEntityType())
-		{
-		case EEntity::Surface:
-			Display2D((const FSurface&) Entity);
-			break;
-		case EEntity::TopologicalFace:
-			Display2D((const FTopologicalFace&) Entity);
-			break;
-		case EEntity::TopologicalLoop:
-			Display2D((const FTopologicalLoop&) Entity);
-			break;
-		case EEntity::TopologicalEdge:
-			Display2D((const FTopologicalEdge&) Entity);
-			break;
-		default:
-			FMessage::Printf(Log, TEXT("Unable to display Entity of type %s"), FEntity::GetTypeName(Entity.GetEntityType()));
-		}
+	switch (Entity.GetEntityType())
+	{
+	case EEntity::Surface:
+		Display2D((const FSurface&)Entity);
+		break;
+	case EEntity::TopologicalFace:
+		Display2D((const FTopologicalFace&)Entity);
+		break;
+	case EEntity::TopologicalLoop:
+		Display2D((const FTopologicalLoop&)Entity);
+		break;
+	case EEntity::TopologicalEdge:
+		Display2D((const FTopologicalEdge&)Entity);
+		break;
+	default:
+		FMessage::Printf(Log, TEXT("Unable to display Entity of type %s"), FEntity::GetTypeName(Entity.GetEntityType()));
+	}
 
-		FDuration DisplayDuration = FChrono::Elapse(StartTime);
-		FChrono::PrintClockElapse(Log, TEXT("  "), TEXT("Display "), DisplayDuration);
+	FDuration DisplayDuration = FChrono::Elapse(StartTime);
+	FChrono::PrintClockElapse(Log, TEXT("  "), TEXT("Display "), DisplayDuration);
 
 #endif
-	}
+}
 
-	void Display(const FPlane& Plane, FIdent Ident)
-	{
+void Display(const FPlane& Plane, FIdent Ident)
+{
 #ifdef CADKERNEL_DEV
-		FPoint Normal = Plane.GetNormal();
-		if (Normal.Length() < SMALL_NUMBER)
-		{
-			return;
-		}
+	FPoint Normal = Plane.GetNormal();
+	if (Normal.Length() < SMALL_NUMBER)
+	{
+		return;
+	}
 
-		FPoint UAxis;
-		// Find the first frame axis non collinear to the plane normal
-		for (int32 Index = 0; Index < 3; ++Index)
+	FPoint UAxis;
+	// Find the first frame axis non collinear to the plane normal
+	for (int32 Index = 0; Index < 3; ++Index)
+	{
+		FPoint Axis = FPoint::ZeroPoint;
+		Axis[Index] = 1;
+		UAxis = Normal ^ Axis;
+		if (UAxis.Length() > SMALL_NUMBER)
 		{
-			FPoint Axis = FPoint::ZeroPoint;
-			Axis[Index] = 1;
-			UAxis = Normal ^ Axis;
-			if (UAxis.Length() > SMALL_NUMBER)
-			{
-				break;
-			}
+			break;
 		}
-		FPoint VAxis = UAxis ^ Normal;
-
-		UAxis.Normalize();
-		UAxis *= 10;
-		VAxis.Normalize();
-		VAxis *= 10;
-
-		TArray<FPoint> Points;
-		Points.Init(FPoint(), 5);
-		const FPoint& Point = Plane.GetPoint();
-		Points[0] = Point + UAxis + VAxis;
-		Points[1] = Point + UAxis - VAxis;
-		Points[2] = Point - UAxis - VAxis;
-		Points[3] = Point - UAxis + VAxis;
-		Points[4] = Point + UAxis + VAxis;
-
-		F3DDebugSegment G(Ident);
-		DrawElement(2, Points);
-#endif
 	}
+	FPoint VAxis = UAxis ^ Normal;
+
+	UAxis.Normalize();
+	UAxis *= 10;
+	VAxis.Normalize();
+	VAxis *= 10;
+
+	TArray<FPoint> Points;
+	Points.Init(FPoint(), 5);
+	const FPoint& Point = Plane.GetPoint();
+	Points[0] = Point + UAxis + VAxis;
+	Points[1] = Point + UAxis - VAxis;
+	Points[2] = Point - UAxis - VAxis;
+	Points[3] = Point - UAxis + VAxis;
+	Points[4] = Point + UAxis + VAxis;
+
+	F3DDebugSegment G(Ident);
+	DrawElement(2, Points);
+#endif
+}
 
-	void DisplayEdgeCriteriaGrid(int32 EdgeId, const TArray<FPoint>& Points3D)
-	{
+void DisplayEdgeCriteriaGrid(int32 EdgeId, const TArray<FPoint>& Points3D)
+{
 #ifdef CADKERNEL_DEV
-		FString Name = FString::Printf(TEXT("Edge Grid %d"), EdgeId);
-		F3DDebugSession G(Name);
+	FString Name = FString::Printf(TEXT("Edge Grid %d"), EdgeId);
+	F3DDebugSession G(Name);
+	{
+		F3DDebugSession _(TEXT("CriteriaGrid Point 3d"));
+		for (int32 Index = 0; Index < Points3D.Num(); Index += 2)
 		{
-			F3DDebugSession _(TEXT("CriteriaGrid Point 3d"));
-			for (int32 Index = 0; Index < Points3D.Num(); Index += 2)
-			{
-				DisplayPoint(Points3D[Index]);
-			}
+			DisplayPoint(Points3D[Index]);
 		}
+	}
+	{
+		F3DDebugSession _(TEXT("CriteriaGrid IntermediateU"));
+		for (int32 Index = 1; Index < Points3D.Num(); Index += 2)
 		{
-			F3DDebugSession _(TEXT("CriteriaGrid IntermediateU"));
-			for (int32 Index = 1; Index < Points3D.Num(); Index += 2)
-			{
-				DisplayPoint(Points3D[Index], EVisuProperty::ControlPoint);
-			}
+			DisplayPoint(Points3D[Index], EVisuProperty::ControlPoint);
 		}
-#endif
 	}
+#endif
+}
 
-	void DisplayAABB(const FAABB& Aabb, FIdent Ident)
-	{
+void DisplayAABB(const FAABB& Aabb, FIdent Ident)
+{
 #ifdef CADKERNEL_DEV
-		TFunction<void(int32, FPoint&, FPoint&)> GetWire = [&](int32 Segment, FPoint& A, FPoint& B)
-		{
-			int32 Corner1;
-			int32 Corner2;
+	TFunction<void(int32, FPoint&, FPoint&)> GetWire = [&](int32 Segment, FPoint& A, FPoint& B)
+	{
+		int32 Corner1;
+		int32 Corner2;
 
-			if (Segment < 4)
-			{
-				Corner1 = Segment * 2;
-				Corner2 = Corner1 + 1;
-			}
-			else if (Segment < 8)
-			{
-				if (Segment < 6)
-				{
-					Corner1 = Segment - 4;
-				}
-				else
-				{
-					Corner1 = Segment - 2;
-				}
-				Corner2 = Corner1 + 2;
-			}
-			else if (Segment < 12)
+		if (Segment < 4)
+		{
+			Corner1 = Segment * 2;
+			Corner2 = Corner1 + 1;
+		}
+		else if (Segment < 8)
+		{
+			if (Segment < 6)
 			{
-				Corner1 = Segment - 8;
-				Corner2 = Corner1 + 4;
+				Corner1 = Segment - 4;
 			}
 			else
 			{
-				Corner1 = 0;
-				Corner2 = 1;
-				ERROR_NOT_EXPECTED;
+				Corner1 = Segment - 2;
 			}
-			A = Aabb.GetCorner(Corner1);
-			B = Aabb.GetCorner(Corner2);
-		};
-
-		F3DDebugSegment G(Ident);
-		for (int32 SegmentIndex = 0; SegmentIndex < 12; SegmentIndex++)
+			Corner2 = Corner1 + 2;
+		}
+		else if (Segment < 12)
 		{
-			TArray<FPoint> Points;
-			Points.SetNum(2);
-			GetWire(SegmentIndex, Points[0], Points[1]);
-			Draw(Points);
+			Corner1 = Segment - 8;
+			Corner2 = Corner1 + 4;
 		}
-#endif
-	}
+		else
+		{
+			Corner1 = 0;
+			Corner2 = 1;
+			ERROR_NOT_EXPECTED;
+		}
+		A = Aabb.GetCorner(Corner1);
+		B = Aabb.GetCorner(Corner2);
+	};
 
-	void DisplayAABB2D(const FAABB2D& aabb, FIdent Ident)
+	F3DDebugSegment G(Ident);
+	for (int32 SegmentIndex = 0; SegmentIndex < 12; SegmentIndex++)
 	{
-#ifdef CADKERNEL_DEV
-		FPoint A(aabb.GetMin().U, aabb.GetMin().V, 0.);
-		FPoint B(aabb.GetMin().U, aabb.GetMax().V, 0.);
-		FPoint C(aabb.GetMax().U, aabb.GetMax().V, 0.);
-		FPoint D(aabb.GetMax().U, aabb.GetMin().V, 0.);
-
 		TArray<FPoint> Points;
-		Points.Add(A);
-		Points.Add(B);
-		Points.Add(C);
-		Points.Add(D);
-		Points.Add(A);
+		Points.SetNum(2);
+		GetWire(SegmentIndex, Points[0], Points[1]);
 		Draw(Points);
-#endif
 	}
+#endif
+}
 
-	void Display(const FTopologicalVertex& Vertex)
-	{
+void DisplayAABB2D(const FAABB2D& aabb, FIdent Ident)
+{
 #ifdef CADKERNEL_DEV
-		F3DDebugSegment G(Vertex.GetId());
-		DrawPoint(Vertex.GetCoordinates());
+	FPoint A(aabb.GetMin().U, aabb.GetMin().V, 0.);
+	FPoint B(aabb.GetMin().U, aabb.GetMax().V, 0.);
+	FPoint C(aabb.GetMax().U, aabb.GetMax().V, 0.);
+	FPoint D(aabb.GetMax().U, aabb.GetMin().V, 0.);
+
+	TArray<FPoint> Points;
+	Points.Add(A);
+	Points.Add(B);
+	Points.Add(C);
+	Points.Add(D);
+	Points.Add(A);
+	Draw(Points);
 #endif
-	}
+}
 
-	void Display(const FCurve& Curve)
-	{
+void Display(const FTopologicalVertex& Vertex)
+{
 #ifdef CADKERNEL_DEV
-		F3DDebugSegment G(Curve.GetId());
-		Draw(Curve, Curve.GetBoundary());
+	F3DDebugSegment G(Vertex.GetId());
+	DrawPoint(Vertex.GetCoordinates());
 #endif
-	}
+}
 
-	void Display(const FShell& Shell)
-	{
+void Display(const FCurve& Curve)
+{
 #ifdef CADKERNEL_DEV
-		FProgress Progress(Shell.GetFaces().Num(), TEXT("Display Shell"));
-		for (const FOrientedFace& Face : Shell.GetFaces())
-		{
-			Display(*Face.Entity);
-		}
+	F3DDebugSegment G(Curve.GetId());
+	Draw(Curve, Curve.GetBoundary());
 #endif
-	}
+}
 
-	void Draw(const FShell& Shell)
-	{
+void Display(const FShell& Shell)
+{
 #ifdef CADKERNEL_DEV
-		FProgress Progress(Shell.GetFaces().Num(), TEXT("Display Shell"));
-		for (const FOrientedFace& Face : Shell.GetFaces())
-		{
-			F3DDebugSegment _(Face.Entity->GetId());
-			Draw(*Face.Entity);
-		}
-#endif
+	FProgress Progress(Shell.GetFaces().Num(), TEXT("Display Shell"));
+	for (const FOrientedFace& Face : Shell.GetFaces())
+	{
+		Display(*Face.Entity);
 	}
+#endif
+}
 
-	void Display(const FBody& Body)
+void Draw(const FShell& Shell)
+{
+#ifdef CADKERNEL_DEV
+	FProgress Progress(Shell.GetFaces().Num(), TEXT("Display Shell"));
+	for (const FOrientedFace& Face : Shell.GetFaces())
 	{
+		F3DDebugSegment _(Face.Entity->GetId());
+		Draw(*Face.Entity);
+	}
+#endif
+}
+
+void Display(const FBody& Body)
+{
 #ifdef CADKERNEL_DEV
-		FProgress Progress(Body.GetShells().Num(), TEXT("Display Body"));
+	FProgress Progress(Body.GetShells().Num(), TEXT("Display Body"));
 
-		F3DDebugSegment GraphicSegment(Body.GetId());
-		for (const TSharedPtr<FShell>& Shell : Body.GetShells())
+	F3DDebugSegment GraphicSegment(Body.GetId());
+	for (const TSharedPtr<FShell>& Shell : Body.GetShells())
+	{
+		if (!Shell.IsValid())
 		{
-			if (!Shell.IsValid())
-			{
-				continue;
-			}
-			Draw(*Shell);
+			continue;
 		}
-#endif
+		Draw(*Shell);
 	}
+#endif
+}
 
-	void Display(const FSurface& Surface)
-	{
+void Display(const FSurface& Surface)
+{
 #ifdef CADKERNEL_DEV
-		FProgress Progress(TEXT("Display Surface"));
+	FProgress Progress(TEXT("Display Surface"));
 
-		F3DDebugSegment GraphicSegment(Surface.GetId());
+	F3DDebugSegment GraphicSegment(Surface.GetId());
 
-		int32 IsoUNum = FSystem::Get().GetVisu()->GetParameters()->IsoUNumber;
-		int32 IsoVNum = FSystem::Get().GetVisu()->GetParameters()->IsoVNumber;
+	int32 IsoUNum = FSystem::Get().GetVisu()->GetParameters()->IsoUNumber;
+	int32 IsoVNum = FSystem::Get().GetVisu()->GetParameters()->IsoVNumber;
 
-		const double VisuSag = FSystem::Get().GetVisu()->GetParameters()->ChordError;
+	const double VisuSag = FSystem::Get().GetVisu()->GetParameters()->ChordError;
 
-		FPolyline3D Polyline;
-		FIsoCurve3DSamplerOnChord Sampler(Surface, VisuSag, Polyline);
+	FPolyline3D Polyline;
+	FIsoCurve3DSamplerOnChord Sampler(Surface, VisuSag, Polyline);
 
-		TFunction <void(const int32, const EIso)> DrawIsos = [&](int32 IsoCount, const EIso IsoType)
-		{
-			const FLinearBoundary& CurveBounds = Surface.GetBoundary().Get(IsoType == EIso::IsoU ? EIso::IsoV : EIso::IsoU);
-			const FLinearBoundary& Bounds = Surface.GetBoundary().Get(IsoType);
+	TFunction <void(const int32, const EIso)> DrawIsos = [&](int32 IsoCount, const EIso IsoType)
+	{
+		const FLinearBoundary& CurveBounds = Surface.GetBoundary().Get(IsoType == EIso::IsoU ? EIso::IsoV : EIso::IsoU);
+		const FLinearBoundary& Bounds = Surface.GetBoundary().Get(IsoType);
 
-			double Coordinate = Bounds.Min;
-			IsoCount++;
-			const double Step = (Bounds.Max - Bounds.Min) / IsoCount;
+		double Coordinate = Bounds.Min;
+		IsoCount++;
+		const double Step = (Bounds.Max - Bounds.Min) / IsoCount;
 
-			for (int32 Index = 0; Index <= IsoCount; Index++)
-			{
-				Polyline.Empty();
-				Sampler.Set(IsoType, Coordinate, CurveBounds);
-				Sampler.Sample();
+		for (int32 Index = 0; Index <= IsoCount; Index++)
+		{
+			Polyline.Empty();
+			Sampler.Set(IsoType, Coordinate, CurveBounds);
+			Sampler.Sample();
 
-				Draw(Polyline.GetPoints(), EVisuProperty::Iso);
-				Coordinate += Step;
-			}
-		};
+			Draw(Polyline.GetPoints(), EVisuProperty::Iso);
+			Coordinate += Step;
+		}
+	};
 
-		DrawIsos(IsoVNum, EIso::IsoV);
-		DrawIsos(IsoUNum, EIso::IsoU);
+	DrawIsos(IsoVNum, EIso::IsoV);
+	DrawIsos(IsoUNum, EIso::IsoU);
 #endif
-	}
+}
 
-	void Display2D(const FSurface& Surface)
-	{
+void Display2D(const FSurface& Surface)
+{
 #ifdef CADKERNEL_DEV
-		double StepU, StepV;
+	double StepU, StepV;
 
-		F3DDebugSegment G(Surface.GetId());
+	F3DDebugSegment G(Surface.GetId());
 
-		int32 IsoUCount = FSystem::Get().GetVisu()->GetParameters()->IsoUNumber;
-		int32 IsoVCount = FSystem::Get().GetVisu()->GetParameters()->IsoVNumber;
-		FSurfacicBoundary Bounds = Surface.GetBoundary();
+	int32 IsoUCount = FSystem::Get().GetVisu()->GetParameters()->IsoUNumber;
+	int32 IsoVCount = FSystem::Get().GetVisu()->GetParameters()->IsoVNumber;
+	FSurfacicBoundary Bounds = Surface.GetBoundary();
 
-		StepU = (Bounds[EIso::IsoU].Max - Bounds[EIso::IsoU].Min) / (IsoUCount + 1);
-		StepV = (Bounds[EIso::IsoV].Max - Bounds[EIso::IsoV].Min) / (IsoVCount + 1);
+	StepU = (Bounds[EIso::IsoU].Max - Bounds[EIso::IsoU].Min) / (IsoUCount + 1);
+	StepV = (Bounds[EIso::IsoV].Max - Bounds[EIso::IsoV].Min) / (IsoVCount + 1);
 
-		for (int32 iIso = 0; iIso <= (IsoUCount + 1); iIso++)
-		{
-			FPoint2D StartPoint(Bounds[EIso::IsoU].Min + (iIso)*StepU, Bounds[EIso::IsoV].Min);
-			FPoint2D EndPoint(Bounds[EIso::IsoU].Min + (iIso)*StepU, Bounds[EIso::IsoV].Max);
-			DrawSegment(StartPoint, EndPoint, EVisuProperty::Iso);
-		}
-
-		for (int32 iIso = 0; iIso <= (IsoVCount + 1); iIso++)
-		{
-			FPoint2D StartPoint(Bounds[EIso::IsoU].Min, Bounds[EIso::IsoV].Min + (iIso)*StepV);
-			FPoint2D EndPoint(Bounds[EIso::IsoU].Max, Bounds[EIso::IsoV].Min + (iIso)*StepV);
-			DrawSegment(StartPoint, EndPoint, EVisuProperty::Iso);
-		}
-#endif
+	for (int32 iIso = 0; iIso <= (IsoUCount + 1); iIso++)
+	{
+		FPoint2D StartPoint(Bounds[EIso::IsoU].Min + (iIso)*StepU, Bounds[EIso::IsoV].Min);
+		FPoint2D EndPoint(Bounds[EIso::IsoU].Min + (iIso)*StepU, Bounds[EIso::IsoV].Max);
+		DrawSegment(StartPoint, EndPoint, EVisuProperty::Iso);
 	}
 
-	void DisplayIsoCurve(const FSurface& Surface, double Coordinate, EIso IsoType)
+	for (int32 iIso = 0; iIso <= (IsoVCount + 1); iIso++)
 	{
+		FPoint2D StartPoint(Bounds[EIso::IsoU].Min, Bounds[EIso::IsoV].Min + (iIso)*StepV);
+		FPoint2D EndPoint(Bounds[EIso::IsoU].Max, Bounds[EIso::IsoV].Min + (iIso)*StepV);
+		DrawSegment(StartPoint, EndPoint, EVisuProperty::Iso);
+	}
+#endif
+}
+
+void DisplayIsoCurve(const FSurface& Surface, double Coordinate, EIso IsoType)
+{
 #ifdef CADKERNEL_DEV
-		const double VisuSag = FSystem::Get().GetVisu()->GetParameters()->ChordError;
+	const double VisuSag = FSystem::Get().GetVisu()->GetParameters()->ChordError;
 
-		F3DDebugSegment G(Surface.GetId());
+	F3DDebugSegment G(Surface.GetId());
 
-		const FLinearBoundary& CurveBounds = Surface.GetBoundary().Get(IsoType == EIso::IsoU ? EIso::IsoV : EIso::IsoU);
+	const FLinearBoundary& CurveBounds = Surface.GetBoundary().Get(IsoType == EIso::IsoU ? EIso::IsoV : EIso::IsoU);
 
-		FPolyline3D Polyline;
-		FIsoCurve3DSamplerOnChord Sampler(Surface, VisuSag, Polyline);
-		Sampler.Set(IsoType, Coordinate, CurveBounds);
-		Draw(Polyline.GetPoints(), EVisuProperty::Iso);
+	FPolyline3D Polyline;
+	FIsoCurve3DSamplerOnChord Sampler(Surface, VisuSag, Polyline);
+	Sampler.Set(IsoType, Coordinate, CurveBounds);
+	Draw(Polyline.GetPoints(), EVisuProperty::Iso);
 #endif
-	}
+}
 
 
 
-	void DisplayControlPolygon(const FCurve& Curve)
-	{
+void DisplayControlPolygon(const FCurve& Curve)
+{
 #ifdef CADKERNEL_DEV
-		TFunction<void(const TArray<FPoint>&)> DisplayHull = [](const TArray<FPoint>& Poles)
-		{
-			for (const FPoint& Pole : Poles)
-			{
-				Display(Pole, EVisuProperty::GreenPoint);
-			}
-
-			for (int32 Index = 1; Index < Poles.Num(); Index++)
-			{
-				FPoint Segment = Poles[Index] - Poles[Index - 1];
-				DisplaySegment(Poles[Index - 1] + Segment * 0.1, Poles[Index] - Segment * 0.1, 0, EVisuProperty::GreenCurve);
-			}
-		};
-
-		F3DDebugSegment GraphicSegment(Curve.GetId());
-		if (Curve.GetCurveType() == ECurve::Bezier)
+	TFunction<void(const TArray<FPoint>&)> DisplayHull = [](const TArray<FPoint>& Poles)
+	{
+		for (const FPoint& Pole : Poles)
 		{
-			const FBezierCurve& Bezier = (const FBezierCurve&) Curve;
-			DisplayHull(Bezier.GetPoles());
-			return;
+			Display(Pole, EVisuProperty::GreenPoint);
 		}
-		if (Curve.GetCurveType() == ECurve::Nurbs)
+
+		for (int32 Index = 1; Index < Poles.Num(); Index++)
 		{
-			const FNURBSCurve& Nurbs = (const FNURBSCurve&) Curve;
-			const TArray<FPoint>& Poles = Nurbs.GetPoles();
-			DisplayHull(Nurbs.GetPoles());
-			return;
+			FPoint Segment = Poles[Index] - Poles[Index - 1];
+			DisplaySegment(Poles[Index - 1] + Segment * 0.1, Poles[Index] - Segment * 0.1, 0, EVisuProperty::GreenCurve);
 		}
-#endif
-	}
+	};
 
-	void DisplayControlPolygon(const FSurface& Surface)
+	F3DDebugSegment GraphicSegment(Curve.GetId());
+	if (Curve.GetCurveType() == ECurve::Bezier)
+	{
+		const FBezierCurve& Bezier = (const FBezierCurve&)Curve;
+		DisplayHull(Bezier.GetPoles());
+		return;
+	}
+	if (Curve.GetCurveType() == ECurve::Nurbs)
 	{
+		const FNURBSCurve& Nurbs = (const FNURBSCurve&)Curve;
+		const TArray<FPoint>& Poles = Nurbs.GetPoles();
+		DisplayHull(Nurbs.GetPoles());
+		return;
+	}
+#endif
+}
+
+void DisplayControlPolygon(const FSurface& Surface)
+{
 #ifdef CADKERNEL_DEV
-		TFunction<void(const TArray<FPoint>&, int32, int32)> DisplayHull = [](const TArray<FPoint>& Poles, int32 PoleUNum, int32 PoleVNum)
+	TFunction<void(const TArray<FPoint>&, int32, int32)> DisplayHull = [](const TArray<FPoint>& Poles, int32 PoleUNum, int32 PoleVNum)
+	{
+		for (int32 Index = 0; Index < Poles.Num(); Index++)
 		{
-			for (int32 Index = 0; Index < Poles.Num(); Index++)
-			{
-				Display(Poles[Index], EVisuProperty::GreenPoint);
-			}
+			Display(Poles[Index], EVisuProperty::GreenPoint);
+		}
 
-			for (int32 IndexV = 0, Index = 0; IndexV < PoleVNum; IndexV++)
+		for (int32 IndexV = 0, Index = 0; IndexV < PoleVNum; IndexV++)
+		{
+			Index = Index + 1;
+			for (int32 IndexU = 1; IndexU < PoleUNum; IndexU++, Index++)
 			{
-				Index = Index + 1;
-				for (int32 IndexU = 1; IndexU < PoleUNum; IndexU++, Index++)
-				{
-					FPoint Segment = Poles[Index] - Poles[Index - 1];
-					DisplaySegment(Poles[Index - 1] + Segment * 0.1, Poles[Index] - Segment * 0.1, 0, EVisuProperty::GreenCurve);
-				}
+				FPoint Segment = Poles[Index] - Poles[Index - 1];
+				DisplaySegment(Poles[Index - 1] + Segment * 0.1, Poles[Index] - Segment * 0.1, 0, EVisuProperty::GreenCurve);
 			}
+		}
 
-			for (int32 IndexU = 0; IndexU < PoleUNum; IndexU++)
+		for (int32 IndexU = 0; IndexU < PoleUNum; IndexU++)
+		{
+			for (int32 IndexV = 1, Index = PoleUNum + IndexU; IndexV < PoleVNum; IndexV++, Index += PoleUNum)
 			{
-				for (int32 IndexV = 1, Index = PoleUNum + IndexU; IndexV < PoleVNum; IndexV++, Index += PoleUNum)
-				{
-					FPoint Segment = Poles[Index] - Poles[Index - PoleUNum];
-					DisplaySegment(Poles[Index - PoleUNum] + Segment * 0.1, Poles[Index] - Segment * 0.1, 0, EVisuProperty::GreenCurve);
-				}
+				FPoint Segment = Poles[Index] - Poles[Index - PoleUNum];
+				DisplaySegment(Poles[Index - PoleUNum] + Segment * 0.1, Poles[Index] - Segment * 0.1, 0, EVisuProperty::GreenCurve);
 			}
-		};
-
-		F3DDebugSegment GraphicSegment(Surface.GetId());
-		if (Surface.GetSurfaceType() == ESurface::Bezier)
-		{
-			const FBezierSurface& Bezier = (const FBezierSurface&) Surface;
-			DisplayHull(Bezier.GetPoles(), Bezier.GetUDegre() + 1, Bezier.GetVDegre() + 1);
-			return;
-		}
-		if (Surface.GetSurfaceType() == ESurface::Nurbs)
-		{
-			const FNURBSSurface& Nurbs = (const FNURBSSurface&) Surface;
-			DisplayHull(Nurbs.GetPoles(), Nurbs.GetPoleCount(EIso::IsoU), Nurbs.GetPoleCount(EIso::IsoV));
-			return;
 		}
-#endif
-	}
+	};
 
-	void Display(const FTopologicalFace& Face)
+	F3DDebugSegment GraphicSegment(Surface.GetId());
+	if (Surface.GetSurfaceType() == ESurface::Bezier)
 	{
+		const FBezierSurface& Bezier = (const FBezierSurface&)Surface;
+		DisplayHull(Bezier.GetPoles(), Bezier.GetUDegre() + 1, Bezier.GetVDegre() + 1);
+		return;
+	}
+	if (Surface.GetSurfaceType() == ESurface::Nurbs)
+	{
+		const FNURBSSurface& Nurbs = (const FNURBSSurface&)Surface;
+		DisplayHull(Nurbs.GetPoles(), Nurbs.GetPoleCount(EIso::IsoU), Nurbs.GetPoleCount(EIso::IsoV));
+		return;
+	}
+#endif
+}
+
+void Display(const FTopologicalFace& Face)
+{
 #ifdef CADKERNEL_DEV
-		F3DDebugSegment GraphicSegment(Face.GetId());
-		Draw(Face);
+	F3DDebugSegment GraphicSegment(Face.GetId());
+	Draw(Face);
 #endif
-	}
+}
 
-	void Display2D(const FTopologicalFace& Face)
-	{
+void Display2D(const FTopologicalFace& Face)
+{
 #ifdef CADKERNEL_DEV
-		F3DDebugSegment GraphicSegment(Face.GetId());
-		Draw2D(Face);
+	F3DDebugSegment GraphicSegment(Face.GetId());
+	Draw2D(Face);
 #endif
-	}
+}
 
-	void Draw(const FTopologicalFace& Face)
-	{
+void Draw(const FTopologicalFace& Face)
+{
 #ifdef CADKERNEL_DEV
+	{
+		for (const TSharedPtr<FTopologicalLoop>& Loop : Face.GetLoops())
 		{
-			for (const TSharedPtr<FTopologicalLoop>& Loop : Face.GetLoops())
+			for (const FOrientedEdge& Edge : Loop->GetEdges())
 			{
-				for (const FOrientedEdge& Edge : Loop->GetEdges())
+				EVisuProperty Property = EVisuProperty::BlueCurve;
+				switch (Edge.Entity->GetTwinsEntityCount())
 				{
-					EVisuProperty Property = EVisuProperty::BlueCurve;
-					switch (Edge.Entity->GetTwinsEntityCount())
-					{
-					case 1:
-						Property = EVisuProperty::BorderEdge;
-						break;
-					case 2:
-						Property = EVisuProperty::BlueCurve;
-						break;
-					default:
-						Property = EVisuProperty::NonManifoldEdge;
-					}
-					Draw(*Edge.Entity, Property);
+				case 1:
+					Property = EVisuProperty::BorderEdge;
+					break;
+				case 2:
+					Property = EVisuProperty::BlueCurve;
+					break;
+				default:
+					Property = EVisuProperty::NonManifoldEdge;
 				}
+				Draw(*Edge.Entity, Property);
 			}
 		}
+	}
 
-		DrawIsoCurves(Face);
+	DrawIsoCurves(Face);
 #endif
-	}
+}
 
-	void Draw2D(const FTopologicalFace& Face)
-	{
+void Draw2D(const FTopologicalFace& Face)
+{
 #ifdef CADKERNEL_DEV
-		TArray<TArray<FPoint2D>> BoundaryApproximation;
-		Face.Get2DLoopSampling(BoundaryApproximation);
+	TArray<TArray<FPoint2D>> BoundaryApproximation;
+	Face.Get2DLoopSampling(BoundaryApproximation);
 
-		for (const TArray<FPoint2D>& Boudary : BoundaryApproximation)
-		{
-			Draw(Boudary, EVisuProperty::BlueCurve);
-		}
+	for (const TArray<FPoint2D>& Boudary : BoundaryApproximation)
+	{
+		Draw(Boudary, EVisuProperty::BlueCurve);
+	}
 
-		TFunction <void(const int32, const EIso)> DrawIsos = [&](int32 IsoCount, const EIso IsoType)
-		{
-			EIso OtherIso = IsoType == EIso::IsoU ? EIso::IsoV : EIso::IsoU;
-			const FLinearBoundary& Bounds = Face.GetBoundary().Get(IsoType);
-			const FLinearBoundary& CurveBounds = Face.GetBoundary().Get(OtherIso);
+	TFunction <void(const int32, const EIso)> DrawIsos = [&](int32 IsoCount, const EIso IsoType)
+	{
+		EIso OtherIso = IsoType == EIso::IsoU ? EIso::IsoV : EIso::IsoU;
+		const FLinearBoundary& Bounds = Face.GetBoundary().Get(IsoType);
+		const FLinearBoundary& CurveBounds = Face.GetBoundary().Get(OtherIso);
 
 
-			double Coordinate = Bounds.Min;
-			IsoCount++;
-			const double Step = (Bounds.Max - Bounds.Min) / IsoCount;
+		double Coordinate = Bounds.Min;
+		IsoCount++;
+		const double Step = (Bounds.Max - Bounds.Min) / IsoCount;
 
-			for (int32 iIso = 1; iIso < IsoCount; iIso++)
-			{
-				Coordinate += Step;
+		for (int32 iIso = 1; iIso < IsoCount; iIso++)
+		{
+			Coordinate += Step;
 
-				TArray<double> Intersections;
-				FindLoopIntersectionsWithIso(IsoType, Coordinate, BoundaryApproximation, Intersections);
+			TArray<double> Intersections;
+			FindLoopIntersectionsWithIso(IsoType, Coordinate, BoundaryApproximation, Intersections);
 
-				FPoint2D Start;
-				FPoint2D End;
+			FPoint2D Start;
+			FPoint2D End;
 
-				Start[IsoType] = Coordinate;
-				End[IsoType] = Coordinate;
+			Start[IsoType] = Coordinate;
+			End[IsoType] = Coordinate;
 
-				if (Intersections.Num() % 2 != 0)
-				{
-					Start[OtherIso] = Intersections[Intersections.Num() - 1];
-					End[OtherIso] = CurveBounds.GetMax();
-					DrawSegment(Start, End, EVisuProperty::YellowCurve);
-					Intersections.Pop();
-				}
+			if (Intersections.Num() % 2 != 0)
+			{
+				Start[OtherIso] = Intersections[Intersections.Num() - 1];
+				End[OtherIso] = CurveBounds.GetMax();
+				DrawSegment(Start, End, EVisuProperty::YellowCurve);
+				Intersections.Pop();
+			}
 
-				for (int32 ISection = 0; ISection < Intersections.Num(); ISection += 2)
-				{
-					Start[OtherIso] = Intersections[ISection];
-					End[OtherIso] = Intersections[ISection + 1];
-					DrawSegment(Start, End, EVisuProperty::Iso);
-				}
+			for (int32 ISection = 0; ISection < Intersections.Num(); ISection += 2)
+			{
+				Start[OtherIso] = Intersections[ISection];
+				End[OtherIso] = Intersections[ISection + 1];
+				DrawSegment(Start, End, EVisuProperty::Iso);
 			}
-		};
+		}
+	};
 
-		int32 IsoUCount = FSystem::Get().GetVisu()->GetParameters()->IsoUNumber;
-		int32 IsoVCount = FSystem::Get().GetVisu()->GetParameters()->IsoVNumber;
+	int32 IsoUCount = FSystem::Get().GetVisu()->GetParameters()->IsoUNumber;
+	int32 IsoVCount = FSystem::Get().GetVisu()->GetParameters()->IsoVNumber;
 
-		DrawIsos(IsoUCount, EIso::IsoU);
-		DrawIsos(IsoVCount, EIso::IsoV);
+	DrawIsos(IsoUCount, EIso::IsoU);
+	DrawIsos(IsoVCount, EIso::IsoV);
 #endif
-	}
+}
 
 
-	void DrawIsoCurves(const FTopologicalFace& Face)
-	{
+void DrawIsoCurves(const FTopologicalFace& Face)
+{
 #ifdef CADKERNEL_DEV
-		TArray<TArray<FPoint2D>> BoundaryApproximation;
-		Face.Get2DLoopSampling(BoundaryApproximation);
+	TArray<TArray<FPoint2D>> BoundaryApproximation;
+	Face.Get2DLoopSampling(BoundaryApproximation);
 
-		const double VisuSag = FSystem::Get().GetVisu()->GetParameters()->ChordError;
+	const double VisuSag = FSystem::Get().GetVisu()->GetParameters()->ChordError;
 
-		FPolyline3D Polyline;
-		const FSurface& Surface = Face.GetCarrierSurface().Get();
-		FIsoCurve3DSamplerOnChord Sampler(Surface, VisuSag, Polyline);
+	FPolyline3D Polyline;
+	const FSurface& Surface = Face.GetCarrierSurface().Get();
+	FIsoCurve3DSamplerOnChord Sampler(Surface, VisuSag, Polyline);
 
-		TFunction <void(const int32, const EIso)> DrawIsos = [&](int32 IsoCount, const EIso IsoType)
-		{
-			const FLinearBoundary& Bounds = Face.GetBoundary().Get(IsoType);
+	TFunction <void(const int32, const EIso)> DrawIsos = [&](int32 IsoCount, const EIso IsoType)
+	{
+		const FLinearBoundary& Bounds = Face.GetBoundary().Get(IsoType);
 
-			EIso Other = IsoType == EIso::IsoU ? EIso::IsoV : EIso::IsoU;
+		EIso Other = IsoType == EIso::IsoU ? EIso::IsoV : EIso::IsoU;
 
-			double Coordinate = Bounds.Min;
-			IsoCount++;
-			const double Step = (Bounds.Max - Bounds.Min) / IsoCount;
+		double Coordinate = Bounds.Min;
+		IsoCount++;
+		const double Step = (Bounds.Max - Bounds.Min) / IsoCount;
 
-			for (int32 iIso = 1; iIso < IsoCount; iIso++)
-			{
-				Coordinate += Step;
+		for (int32 iIso = 1; iIso < IsoCount; iIso++)
+		{
+			Coordinate += Step;
 
-				TArray<double> Intersections;
-				FindLoopIntersectionsWithIso(IsoType, Coordinate, BoundaryApproximation, Intersections);
-				int32 IntersectionCount = Intersections.Num();
-				if (IntersectionCount == 0)
-				{
-					return;
-				}
+			TArray<double> Intersections;
+			FindLoopIntersectionsWithIso(IsoType, Coordinate, BoundaryApproximation, Intersections);
+			int32 IntersectionCount = Intersections.Num();
+			if (IntersectionCount == 0)
+			{
+				return;
+			}
 
-				FLinearBoundary CurveBounds(Intersections[0], Intersections.Last());
+			FLinearBoundary CurveBounds(Intersections[0], Intersections.Last());
 
-				Polyline.Empty();
-				Sampler.Set(IsoType, Coordinate, CurveBounds);
-				Sampler.Sample();
+			Polyline.Empty();
+			Sampler.Set(IsoType, Coordinate, CurveBounds);
+			Sampler.Sample();
 
-				if (IntersectionCount % 2 != 0)
-				{
-					TArray<FPoint> SubPolyline;
-					FLinearBoundary Boundary(Intersections[IntersectionCount - 1], CurveBounds.GetMax());
-					Polyline.GetSubPolyline(Boundary, EOrientation::Front, SubPolyline);
-					Draw(SubPolyline, EVisuProperty::YellowCurve);
-					Intersections.Pop();
-					IntersectionCount--;
-				}
+			if (IntersectionCount % 2 != 0)
+			{
+				TArray<FPoint> SubPolyline;
+				FLinearBoundary Boundary(Intersections[IntersectionCount - 1], CurveBounds.GetMax());
+				Polyline.GetSubPolyline(Boundary, EOrientation::Front, SubPolyline);
+				Draw(SubPolyline, EVisuProperty::YellowCurve);
+				Intersections.Pop();
+				IntersectionCount--;
+			}
 
-				for (int32 ISection = 0; ISection < IntersectionCount; ISection += 2)
-				{
-					TArray<FPoint> SubPolyline;
-					FLinearBoundary Boundary(Intersections[ISection], Intersections[ISection + 1]);
-					Polyline.GetSubPolyline(Boundary, EOrientation::Front, SubPolyline);
-					Draw(SubPolyline, EVisuProperty::Iso);
-				}
+			for (int32 ISection = 0; ISection < IntersectionCount; ISection += 2)
+			{
+				TArray<FPoint> SubPolyline;
+				FLinearBoundary Boundary(Intersections[ISection], Intersections[ISection + 1]);
+				Polyline.GetSubPolyline(Boundary, EOrientation::Front, SubPolyline);
+				Draw(SubPolyline, EVisuProperty::Iso);
 			}
-		};
+		}
+	};
 
-		int32 IsoUCount = FSystem::Get().GetVisu()->GetParameters()->IsoUNumber;
-		int32 IsoVCount = FSystem::Get().GetVisu()->GetParameters()->IsoVNumber;
+	int32 IsoUCount = FSystem::Get().GetVisu()->GetParameters()->IsoUNumber;
+	int32 IsoVCount = FSystem::Get().GetVisu()->GetParameters()->IsoVNumber;
 
-		DrawIsos(IsoUCount, EIso::IsoU);
-		DrawIsos(IsoVCount, EIso::IsoV);
+	DrawIsos(IsoUCount, EIso::IsoU);
+	DrawIsos(IsoVCount, EIso::IsoV);
 #endif
-	}
+}
 
-	void Display(const FTopologicalEdge& Edge, EVisuProperty Property)
-	{
+void Display(const FTopologicalEdge& Edge, EVisuProperty Property)
+{
 #ifdef CADKERNEL_DEV
-		F3DDebugSegment GraphicSegment(Edge.GetId());
-		Draw(Edge, Property);
+	F3DDebugSegment GraphicSegment(Edge.GetId());
+	Draw(Edge, Property);
 #endif
-	}
+}
 
-	void Display2D(const FTopologicalEdge& Edge, EVisuProperty Property)
-	{
+void Display2D(const FTopologicalEdge& Edge, EVisuProperty Property)
+{
 #ifdef CADKERNEL_DEV
-		F3DDebugSegment GraphicSegment(Edge.GetId());
-		TArray<FPoint2D> Polyline;
-		Edge.GetCurve()->GetDiscretizationPoints(Edge.GetBoundary(), EOrientation::Front, Polyline);
-		DisplayPolyline(Polyline, Property);
+	F3DDebugSegment GraphicSegment(Edge.GetId());
+	TArray<FPoint2D> Polyline;
+	Edge.GetCurve()->GetDiscretizationPoints(Edge.GetBoundary(), EOrientation::Front, Polyline);
+	DisplayPolyline(Polyline, Property);
 #endif
-	}
+}
 
-	void Draw(const FTopologicalEdge& Edge, EVisuProperty Property)
-	{
+void Draw(const FTopologicalEdge& Edge, EVisuProperty Property)
+{
 #ifdef CADKERNEL_DEV
-		const FLinearBoundary& Boundary = Edge.GetBoundary();
-		Draw(Boundary, *Edge.GetCurve(), Property);
+	const FLinearBoundary& Boundary = Edge.GetBoundary();
+	Draw(Boundary, *Edge.GetCurve(), Property);
 #endif
-	}
+}
 
-	void Display(const FModel& Model)
-	{
+void Display(const FModel& Model)
+{
 #ifdef CADKERNEL_DEV
-		FProgress MainProgress(2, TEXT("Display model"));
+	FProgress MainProgress(2, TEXT("Display model"));
+	{
+		const TArray<TSharedPtr<FBody>>& Bodies = Model.GetBodies();
+		FProgress BodyProgress((int32)Bodies.Num(), TEXT("Bodies"));
+		for (const TSharedPtr<FBody>& Body : Bodies)
 		{
-			const TArray<TSharedPtr<FBody>>& Bodies = Model.GetBodies();
-			FProgress BodyProgress((int32)Bodies.Num(), TEXT("Bodies"));
-			for (const TSharedPtr<FBody>& Body : Bodies)
-			{
-				Display(*Body);
-			}
+			Display(*Body);
 		}
+	}
 
+	{
+		const TArray<TSharedPtr<FTopologicalFace>>& Surfaces = Model.GetFaces();
+		FProgress SurfaceProgress((int32)Surfaces.Num(), TEXT("Surfaces"));
+		for (const TSharedPtr<FTopologicalFace>& Surface : Surfaces)
 		{
-			const TArray<TSharedPtr<FTopologicalFace>>& Surfaces = Model.GetFaces();
-			FProgress SurfaceProgress((int32)Surfaces.Num(), TEXT("Surfaces"));
-			for (const TSharedPtr<FTopologicalFace>& Surface : Surfaces)
-			{
-				Display(*Surface);
-			}
+			Display(*Surface);
 		}
-#endif
 	}
+#endif
+}
 
-	void DisplayProductTree(const FModel& Model)
-	{
+void DisplayProductTree(const FModel& Model)
+{
 #ifdef CADKERNEL_DEV
-		F3DDebugSession GraphicSession(FString::Printf(TEXT("%s %d"), Model.GetTypeName(), Model.GetId()), { Model.GetId() });
+	F3DDebugSession GraphicSession(FString::Printf(TEXT("%s %d"), Model.GetTypeName(), Model.GetId()), { Model.GetId() });
 
-		FProgress MainProgress(2, TEXT("Display Model"));
+	FProgress MainProgress(2, TEXT("Display Model"));
+	{
+		const TArray<TSharedPtr<FBody>>& Bodies = Model.GetBodies();
+		FProgress BodyProgress((int32)Bodies.Num(), TEXT("Bodies"));
+		for (const TSharedPtr<FBody>& Body : Bodies)
 		{
-			const TArray<TSharedPtr<FBody>>& Bodies = Model.GetBodies();
-			FProgress BodyProgress((int32)Bodies.Num(), TEXT("Bodies"));
-			for (const TSharedPtr<FBody>& Body : Bodies)
-			{
-				DisplayProductTree(*Body);
-			}
+			DisplayProductTree(*Body);
 		}
+	}
 
+	{
+		const TArray<TSharedPtr<FTopologicalFace>>& Faces = Model.GetFaces();
+		FProgress SurfaceProgress((int32)Faces.Num(), TEXT("Faces"));
+		for (const TSharedPtr<FTopologicalFace>& Face : Faces)
 		{
-			const TArray<TSharedPtr<FTopologicalFace>>& Faces = Model.GetFaces();
-			FProgress SurfaceProgress((int32)Faces.Num(), TEXT("Faces"));
-			for (const TSharedPtr<FTopologicalFace>& Face : Faces)
-			{
-				Display(*Face);
-			}
+			Display(*Face);
 		}
-#endif
 	}
+#endif
+}
 
-	void DisplayProductTree(const FBody& Body)
-	{
+void DisplayProductTree(const FBody& Body)
+{
 #ifdef CADKERNEL_DEV
 #ifdef CORETECHBRIDGE_DEBUG
-		F3DDebugSession GraphicSession(FString::Printf(TEXT("%s %s KioId: %d Id: %d"), Body.GetTypeName(), Body.GetName(), Body.GetKioId(), Body.GetId()), { Body.GetId() });
+	F3DDebugSession GraphicSession(FString::Printf(TEXT("%s %s KioId: %d Id: %d"), Body.GetTypeName(), Body.GetName(), Body.GetKioId(), Body.GetId()), { Body.GetId() });
 #else
-		F3DDebugSession GraphicSession(FString::Printf(TEXT("%s %s Id: %d"), Body.GetTypeName(), Body.GetName(), Body.GetId()), { Body.GetId() });
+	F3DDebugSession GraphicSession(FString::Printf(TEXT("%s %s Id: %d"), Body.GetTypeName(), Body.GetName(), Body.GetId()), { Body.GetId() });
 #endif
-		FProgress Progress(Body.GetShells().Num(), TEXT("Display Body"));
+	FProgress Progress(Body.GetShells().Num(), TEXT("Display Body"));
 
-		for (TSharedPtr<FShell> Shell : Body.GetShells())
+	for (TSharedPtr<FShell> Shell : Body.GetShells())
+	{
+		if (!Shell.IsValid())
 		{
-			if (!Shell.IsValid())
-			{
-				continue;
-			}
-			DisplayProductTree(*Shell);
+			continue;
 		}
-#endif
+		DisplayProductTree(*Shell);
 	}
+#endif
+}
 
-	void DisplayProductTree(const FShell& Shell)
-	{
+void DisplayProductTree(const FShell& Shell)
+{
 #ifdef CADKERNEL_DEV
 #ifdef CORETECHBRIDGE_DEBUG
-		F3DDebugSession GraphicSession(FString::Printf(TEXT("%s %s KioId: %d Id: %d"), Shell.GetTypeName(), Shell.GetName(), Shell.GetKioId(), Shell.GetId()), { Shell.GetId() });
+	F3DDebugSession GraphicSession(FString::Printf(TEXT("%s %s KioId: %d Id: %d"), Shell.GetTypeName(), Shell.GetName(), Shell.GetKioId(), Shell.GetId()), { Shell.GetId() });
 #else
-		F3DDebugSession GraphicSession(FString::Printf(TEXT("%s %s Id: %d"), Shell.GetTypeName(), Shell.GetName(), Shell.GetId()), { Shell.GetId() });
+	F3DDebugSession GraphicSession(FString::Printf(TEXT("%s %s Id: %d"), Shell.GetTypeName(), Shell.GetName(), Shell.GetId()), { Shell.GetId() });
 #endif
-		Draw(Shell);
+	Draw(Shell);
 #endif
-	}
+}
 
-	void DisplayProductTree(const FEntity& Entity)
-	{
+void DisplayProductTree(const FEntity& Entity)
+{
 #ifdef CADKERNEL_DEV
-		FTimePoint StartTime = FChrono::Now();
+	FTimePoint StartTime = FChrono::Now();
 
-		switch (Entity.GetEntityType())
-		{
-		case EEntity::TopologicalFace:
-			Display((const FTopologicalFace&) Entity);
-			break;
-		case EEntity::Shell:
-			DisplayProductTree((const FShell&) Entity);
-			break;
-		case EEntity::Body:
-			DisplayProductTree((const FBody&) Entity);
-			break;
-		case EEntity::Model:
-			DisplayProductTree((const FModel&) Entity);
-			break;
-		default:
-			FMessage::Printf(Log, TEXT("Unable to display Entity of type %s"), FEntity::GetTypeName(Entity.GetEntityType()));
-		}
+	switch (Entity.GetEntityType())
+	{
+	case EEntity::TopologicalFace:
+		Display((const FTopologicalFace&)Entity);
+		break;
+	case EEntity::Shell:
+		DisplayProductTree((const FShell&)Entity);
+		break;
+	case EEntity::Body:
+		DisplayProductTree((const FBody&)Entity);
+		break;
+	case EEntity::Model:
+		DisplayProductTree((const FModel&)Entity);
+		break;
+	default:
+		FMessage::Printf(Log, TEXT("Unable to display Entity of type %s"), FEntity::GetTypeName(Entity.GetEntityType()));
+	}
 
-		FDuration DisplayDuration = FChrono::Elapse(StartTime);
-		FChrono::PrintClockElapse(Log, TEXT("  "), TEXT("Display "), DisplayDuration);
+	FDuration DisplayDuration = FChrono::Elapse(StartTime);
+	FChrono::PrintClockElapse(Log, TEXT("  "), TEXT("Display "), DisplayDuration);
 #endif
-	}
+}
 
-	void Display(const FGroup& Group)
-	{
+void Display(const FGroup& Group)
+{
 #ifdef CADKERNEL_DEV
-		TArray<TSharedPtr<FEntity>> Entities;
+	TArray<TSharedPtr<FEntity>> Entities;
 
-		Group.GetValidEntities(Entities);
+	Group.GetValidEntities(Entities);
 
-		FProgress Progress(Entities.Num());
-		for (const TSharedPtr<FEntity>& Entity : Entities)
-		{
-			DisplayEntity(*Entity);
-		}
-#endif
+	FProgress Progress(Entities.Num());
+	for (const TSharedPtr<FEntity>& Entity : Entities)
+	{
+		DisplayEntity(*Entity);
 	}
+#endif
+}
 
-	void DisplayMesh(const FFaceMesh& Mesh)
-	{
+void DisplayMesh(const FFaceMesh& Mesh)
+{
 #ifdef CADKERNEL_DEV
-		TMap<int32, const FPoint*> NodeIdToCoordinates;
-		Mesh.GetNodeIdToCoordinates(NodeIdToCoordinates);
+	TMap<int32, const FPoint*> NodeIdToCoordinates;
+	Mesh.GetNodeIdToCoordinates(NodeIdToCoordinates);
+
+	TArray<FPoint> Points;
+	Points.SetNum(3);
+	double EdgeMeanLength = 0;
+	const TArray<int32>& TriangleIndices = Mesh.TrianglesVerticesIndex;
+	const TArray<int32>& VertexIndices = Mesh.VerticesGlobalIndex;
+	for (int32 Index = 0; Index < TriangleIndices.Num();)
+	{
+		Points[0] = *NodeIdToCoordinates[VertexIndices[TriangleIndices[Index++]]];
+		Points[1] = *NodeIdToCoordinates[VertexIndices[TriangleIndices[Index++]]];
+		Points[2] = *NodeIdToCoordinates[VertexIndices[TriangleIndices[Index++]]];
+		DrawElement(2, Points, EVisuProperty::Element);
+		DrawSegment(Points[0], Points[1], EVisuProperty::EdgeMesh);
+		DrawSegment(Points[1], Points[2], EVisuProperty::EdgeMesh);
+		DrawSegment(Points[2], Points[0], EVisuProperty::EdgeMesh);
+		EdgeMeanLength += Points[0].Distance(Points[1]);
+		EdgeMeanLength += Points[1].Distance(Points[2]);
+		EdgeMeanLength += Points[2].Distance(Points[0]);
+	}
+	EdgeMeanLength /= 10 * TriangleIndices.Num();
 
-		TArray<FPoint> Points;
-		Points.SetNum(3);
-		double EdgeMeanLength = 0;
-		const TArray<int32>& TriangleIndices = Mesh.TrianglesVerticesIndex;
-		const TArray<int32>& VertexIndices = Mesh.VerticesGlobalIndex;
-		for (int32 Index = 0; Index < TriangleIndices.Num();)
-		{
-			Points[0] = *NodeIdToCoordinates[VertexIndices[TriangleIndices[Index++]]];
-			Points[1] = *NodeIdToCoordinates[VertexIndices[TriangleIndices[Index++]]];
-			Points[2] = *NodeIdToCoordinates[VertexIndices[TriangleIndices[Index++]]];
-			DrawElement(2, Points, EVisuProperty::Element);
-			DrawSegment(Points[0], Points[1], EVisuProperty::EdgeMesh);
-			DrawSegment(Points[1], Points[2], EVisuProperty::EdgeMesh);
-			DrawSegment(Points[2], Points[0], EVisuProperty::EdgeMesh);
-			EdgeMeanLength += Points[0].Distance(Points[1]);
-			EdgeMeanLength += Points[1].Distance(Points[2]);
-			EdgeMeanLength += Points[2].Distance(Points[0]);
-		}
-		EdgeMeanLength /= 10 * TriangleIndices.Num();
+	for (const int32& Index : VertexIndices)
+	{
+		F3DDebugSegment GraphicSegment(Index);
+		DrawPoint(*NodeIdToCoordinates[Index], EVisuProperty::NodeMesh);
+	}
 
-		for (const int32& Index : VertexIndices)
+	bool test = FSystem::Get().GetVisu()->GetParameters()->bDisplayNormals;
+	if (FSystem::Get().GetVisu()->GetParameters()->bDisplayNormals)
+	{
+		double NormalLength = FSystem::Get().GetVisu()->GetParameters()->NormalLength;
+		const TArray<FVector>& Normals = Mesh.Normals;
+		for (int32 Index = 0; Index < VertexIndices.Num(); ++Index)
 		{
 			F3DDebugSegment GraphicSegment(Index);
-			DrawPoint(*NodeIdToCoordinates[Index], EVisuProperty::NodeMesh);
-		}
-
-		bool test = FSystem::Get().GetVisu()->GetParameters()->bDisplayNormals;
-		if (FSystem::Get().GetVisu()->GetParameters()->bDisplayNormals)
-		{
-			double NormalLength = FSystem::Get().GetVisu()->GetParameters()->NormalLength;
-			const TArray<FVector>& Normals = Mesh.Normals;
-			for (int32 Index = 0; Index < VertexIndices.Num(); ++Index)
-			{
-				F3DDebugSegment GraphicSegment(Index);
-				FVector Normal = Normals[Index];
-				Normal.Normalize();
-				Normal *= NormalLength;
-				DrawSegment(*NodeIdToCoordinates[VertexIndices[Index]], *NodeIdToCoordinates[VertexIndices[Index]] + Normal, EVisuProperty::EdgeMesh);
-			}
+			FVector Normal = Normals[Index];
+			Normal.Normalize();
+			Normal *= NormalLength;
+			DrawSegment(*NodeIdToCoordinates[VertexIndices[Index]], *NodeIdToCoordinates[VertexIndices[Index]] + Normal, EVisuProperty::EdgeMesh);
 		}
-#endif
 	}
+#endif
+}
 
-	void DisplayMesh(const FEdgeMesh& Mesh)
-	{
+void DisplayMesh(const FEdgeMesh& Mesh)
+{
 #ifdef CADKERNEL_DEV
-		const FModelMesh& MeshModel = Mesh.GetMeshModel();
+	const FModelMesh& MeshModel = Mesh.GetMeshModel();
 
-		const TArray<int32>& NodeIds = Mesh.EdgeVerticesIndex;
-		const TArray<FPoint>& NodeCoordinates = Mesh.GetNodeCoordinates();
+	const TArray<int32>& NodeIds = Mesh.EdgeVerticesIndex;
+	const TArray<FPoint>& NodeCoordinates = Mesh.GetNodeCoordinates();
 
-		int32 StartNodeId = NodeIds[0];
-		int32 LastNodeId = NodeIds.Last();
+	int32 StartNodeId = NodeIds[0];
+	int32 LastNodeId = NodeIds.Last();
 
+	{
+		F3DDebugSegment GraphicSegment(Mesh.GetGeometricEntity().GetId());
+		if (NodeCoordinates.Num() == 0)
 		{
-			F3DDebugSegment GraphicSegment(Mesh.GetGeometricEntity().GetId());
-			if (NodeCoordinates.Num() == 0)
-			{
-				DrawSegment(MeshModel.GetMeshOfVertexNodeId(StartNodeId)->GetNodeCoordinates()[0], MeshModel.GetMeshOfVertexNodeId(LastNodeId)->GetNodeCoordinates()[0], EVisuProperty::EdgeMesh);
-			}
-			else
+			DrawSegment(MeshModel.GetMeshOfVertexNodeId(StartNodeId)->GetNodeCoordinates()[0], MeshModel.GetMeshOfVertexNodeId(LastNodeId)->GetNodeCoordinates()[0], EVisuProperty::EdgeMesh);
+		}
+		else
+		{
+			DrawSegment(MeshModel.GetMeshOfVertexNodeId(StartNodeId)->GetNodeCoordinates()[0], NodeCoordinates[0], EVisuProperty::EdgeMesh);
+			if (NodeCoordinates.Num() > 1)
 			{
-				DrawSegment(MeshModel.GetMeshOfVertexNodeId(StartNodeId)->GetNodeCoordinates()[0], NodeCoordinates[0], EVisuProperty::EdgeMesh);
-				if (NodeCoordinates.Num() > 1)
+				for (int32 Index = 0; Index < NodeCoordinates.Num() - 1; Index++)
 				{
-					for (int32 Index = 0; Index < NodeCoordinates.Num() - 1; Index++)
-					{
-						DrawSegment(NodeCoordinates[Index], NodeCoordinates[Index + 1], EVisuProperty::EdgeMesh);
-					}
+					DrawSegment(NodeCoordinates[Index], NodeCoordinates[Index + 1], EVisuProperty::EdgeMesh);
 				}
-				DrawSegment(NodeCoordinates.Last(), MeshModel.GetMeshOfVertexNodeId(LastNodeId)->GetNodeCoordinates()[0], EVisuProperty::EdgeMesh);
 			}
+			DrawSegment(NodeCoordinates.Last(), MeshModel.GetMeshOfVertexNodeId(LastNodeId)->GetNodeCoordinates()[0], EVisuProperty::EdgeMesh);
 		}
+	}
 
-		{
-			F3DDebugSegment GraphicSegment(StartNodeId);
-			DrawPoint(MeshModel.GetMeshOfVertexNodeId(StartNodeId)->GetNodeCoordinates()[0], EVisuProperty::NodeMesh);
-		}
-
-		{
-			F3DDebugSegment GraphicSegment(LastNodeId);
-			DrawPoint(MeshModel.GetMeshOfVertexNodeId(LastNodeId)->GetNodeCoordinates()[0], EVisuProperty::NodeMesh);
-		}
-
-		if (NodeCoordinates.Num() > 1)
-		{
-			for (int32 Index = 0; Index < NodeCoordinates.Num(); Index++)
-			{
-				F3DDebugSegment GraphicSegment(NodeIds[Index]);
-				DrawPoint(NodeCoordinates[Index], EVisuProperty::NodeMesh);
-			}
-		}
-#endif
+	{
+		F3DDebugSegment GraphicSegment(StartNodeId);
+		DrawPoint(MeshModel.GetMeshOfVertexNodeId(StartNodeId)->GetNodeCoordinates()[0], EVisuProperty::NodeMesh);
 	}
 
-	void DisplayMesh(const FVertexMesh& Mesh)
 	{
-#ifdef CADKERNEL_DEV
-		DisplayPoint(Mesh.GetNodeCoordinates()[0], EVisuProperty::NodeMesh, Mesh.GetId());
-#endif
+		F3DDebugSegment GraphicSegment(LastNodeId);
+		DrawPoint(MeshModel.GetMeshOfVertexNodeId(LastNodeId)->GetNodeCoordinates()[0], EVisuProperty::NodeMesh);
 	}
 
-	void Display(const FModelMesh& MeshModel)
+	if (NodeCoordinates.Num() > 1)
 	{
-#ifdef CADKERNEL_DEV
-		for (const TSharedPtr<FFaceMesh>& Mesh : MeshModel.GetFaceMeshes())
+		for (int32 Index = 0; Index < NodeCoordinates.Num(); Index++)
 		{
-			if (Mesh.IsValid())
-			{
-				DisplayMesh(*Mesh);
-			}
+			F3DDebugSegment GraphicSegment(NodeIds[Index]);
+			DrawPoint(NodeCoordinates[Index], EVisuProperty::NodeMesh);
 		}
-#endif
 	}
+#endif
+}
 
-	void DisplaySegment(const FPoint& Point1, const FPoint& Point2, FIdent Ident, EVisuProperty Property)
-	{
+void DisplayMesh(const FVertexMesh& Mesh)
+{
 #ifdef CADKERNEL_DEV
-		F3DDebugSegment G(Ident);
-		DrawSegment(Point1, Point2, Property);
+	DisplayPoint(Mesh.GetNodeCoordinates()[0], EVisuProperty::NodeMesh, Mesh.GetId());
 #endif
-	}
+}
 
-	void DisplaySegment(const FPoint2D& Point1, const FPoint2D& Point2, FIdent Ident, EVisuProperty Property, bool bWithOrientation)
-	{
+void Display(const FModelMesh& MeshModel)
+{
 #ifdef CADKERNEL_DEV
-		F3DDebugSegment G(Ident);
-		if (bWithOrientation)
+	for (const TSharedPtr<FFaceMesh>& Mesh : MeshModel.GetFaceMeshes())
+	{
+		if (Mesh.IsValid())
 		{
-			DrawSegmentOrientation(Point1, Point2, Property);
+			DisplayMesh(*Mesh);
 		}
-		DrawSegment(Point1, Point2, Property);
-#endif
 	}
+#endif
+}
 
-	void DisplaySegment(const FPointH& Point1, const FPointH& Point2, FIdent Ident, EVisuProperty Property)
-	{
+void DisplaySegment(const FPoint& Point1, const FPoint& Point2, FIdent Ident, EVisuProperty Property)
+{
 #ifdef CADKERNEL_DEV
-		F3DDebugSegment G(Ident);
-		DrawSegment(Point1, Point2, Property);
+	F3DDebugSegment G(Ident);
+	DrawSegment(Point1, Point2, Property);
 #endif
-	}
+}
 
-	void DisplayLoop(const FTopologicalFace& Surface)
+void DisplaySegment(const FPoint2D& Point1, const FPoint2D& Point2, FIdent Ident, EVisuProperty Property, bool bWithOrientation)
+{
+#ifdef CADKERNEL_DEV
+	F3DDebugSegment G(Ident);
+	if (bWithOrientation)
 	{
+		DrawSegmentOrientation(Point1, Point2, Property);
+	}
+	DrawSegment(Point1, Point2, Property);
+#endif
+}
+
+void DisplaySegment(const FPointH& Point1, const FPointH& Point2, FIdent Ident, EVisuProperty Property)
+{
 #ifdef CADKERNEL_DEV
-		for (const TSharedPtr<FTopologicalLoop>& Loop : Surface.GetLoops())
-		{
-			Display(*Loop);
-		}
+	F3DDebugSegment G(Ident);
+	DrawSegment(Point1, Point2, Property);
 #endif
+}
+
+void DisplayLoop(const FTopologicalFace& Surface)
+{
+#ifdef CADKERNEL_DEV
+	for (const TSharedPtr<FTopologicalLoop>& Loop : Surface.GetLoops())
+	{
+		Display(*Loop);
 	}
+#endif
+}
 
 
-	void Display(const FTopologicalLoop& Loop)
-	{
+void Display(const FTopologicalLoop& Loop)
+{
 #ifdef CADKERNEL_DEV
-		for (const FOrientedEdge& Edge : Loop.GetEdges())
-		{
-			Display(*Edge.Entity);
-		}
-#endif
+	for (const FOrientedEdge& Edge : Loop.GetEdges())
+	{
+		Display(*Edge.Entity);
 	}
+#endif
+}
 
-	void Display2D(const FTopologicalLoop& Loop)
-	{
+void Display2D(const FTopologicalLoop& Loop)
+{
 #ifdef CADKERNEL_DEV
-		for (const FOrientedEdge& Edge : Loop.GetEdges())
-		{
-			Display2D(*Edge.Entity);
-		}
-#endif
+	for (const FOrientedEdge& Edge : Loop.GetEdges())
+	{
+		Display2D(*Edge.Entity);
 	}
+#endif
+}
+
 } // namespace CADKernel
diff --git a/Engine/Source/Runtime/Datasmith/CADKernel/Public/CADKernel/Mesh/Meshers/MesherReport.h b/Engine/Source/Runtime/Datasmith/CADKernel/Public/CADKernel/Mesh/Meshers/MesherReport.h
new file mode 100644
--- /dev/null
+++ b/Engine/Source/Runtime/Datasmith/CADKernel/Public/CADKernel/Mesh/Meshers/MesherReport.h
@@ -0,0 +1,85 @@
+// Copyright Epic Games, Inc. All Rights Reserved.
+#pragma once
+
+#include "CADKernel/Core/Chrono.h"
+#include "CADKernel/Core/Types.h"
+#include "CADKernel/UI/Message.h"
+
+
+namespace CADKernel
+{
+
+struct CADKERNEL_API FMesherLog
+{
+	int32 CountOfIntersectingLoopCorrectionFailures = 0;
+	int32 CountOfSelfIntersectingCycles = 0;
+	int32 CountOfDegeneratedGrids = 0;
+	int32 CountOfDegeneratedLoops = 0;
+
+	FMesherLog()
+	{}
+
+	void PrintBilan() const
+	{
+	}
+};
+
+struct CADKERNEL_API FMesherChronos
+{
+	FDuration GlobalDuration;
+	FDuration ApplyCriteriaDuration;
+	FDuration IsolateQuadPatchDuration;
+	FDuration GlobalMeshDuration;
+	FDuration GlobalPointCloudDuration;
+	FDuration GlobalGeneratePointCloudDuration;
+	FDuration GlobalTriangulateDuration;
+	FDuration GlobalDelaunayDuration;
+	FDuration GlobalMeshAndGetLoopNodes;
+	FDuration GlobalMeshEdges;
+	FDuration GlobalThinZones;
+
+	FDuration GlobalFindThinZones;
+	//FDuration GlobalScaleGrid;
+	FDuration GlobalMeshThinZones;
+
+	FMesherChronos()
+		: GlobalDuration(FChrono::Init())
+		, ApplyCriteriaDuration(FChrono::Init())
+		, IsolateQuadPatchDuration(FChrono::Init())
+		, GlobalMeshDuration(FChrono::Init())
+		, GlobalPointCloudDuration(FChrono::Init())
+		, GlobalGeneratePointCloudDuration(FChrono::Init())
+		, GlobalTriangulateDuration(FChrono::Init())
+		, GlobalDelaunayDuration(FChrono::Init())
+		, GlobalMeshAndGetLoopNodes(FChrono::Init())
+		, GlobalMeshEdges(FChrono::Init())
+		, GlobalThinZones(FChrono::Init())
+		, GlobalFindThinZones(FChrono::Init())
+		//, GlobalScaleGrid(FChrono::Init())
+		, GlobalMeshThinZones(FChrono::Init())
+	{}
+
+	void PrintTimeElapse() const
+	{
+		FMessage::Printf(Log, TEXT("\n\n\n"));
+		FChrono::PrintClockElapse(Log, TEXT(""), TEXT("Total"), GlobalDuration);
+		FChrono::PrintClockElapse(Log, TEXT("  |  "), TEXT("Apply Criteria"), ApplyCriteriaDuration);
+		FChrono::PrintClockElapse(Log, TEXT("  |  "), TEXT("Find Quad Surfaces"), IsolateQuadPatchDuration);
+		FChrono::PrintClockElapse(Log, TEXT("  |  "), TEXT("Mesh Time"), GlobalMeshDuration);
+		FChrono::PrintClockElapse(Log, TEXT("  |   |  "), TEXT("GeneratePoint Cloud "), GlobalGeneratePointCloudDuration);
+		FChrono::PrintClockElapse(Log, TEXT("  |   |  |  "), TEXT("Point Cloud "), GlobalPointCloudDuration);
+		FChrono::PrintClockElapse(Log, TEXT("  |   |  "), TEXT("ThinZones "), GlobalThinZones);
+		FChrono::PrintClockElapse(Log, TEXT("  |   |  "), TEXT("Mesh ThinZones "), GlobalMeshThinZones);
+		FChrono::PrintClockElapse(Log, TEXT("  |   |  "), TEXT("MeshEdges "), GlobalMeshEdges);
+		FChrono::PrintClockElapse(Log, TEXT("  |   |  "), TEXT("TriangulateDuration "), GlobalTriangulateDuration);
+		FChrono::PrintClockElapse(Log, TEXT("  |   |   |  "), TEXT("Delaunay Duration "), GlobalDelaunayDuration);
+	}
+};
+
+struct CADKERNEL_API FMesherReport
+{
+	FMesherLog Logs;
+	FMesherChronos Chronos;
+};
+
+}
diff --git a/Engine/Source/Runtime/Datasmith/CADKernel/Public/CADKernel/Mesh/Meshers/ParametricMesher.h b/Engine/Source/Runtime/Datasmith/CADKernel/Public/CADKernel/Mesh/Meshers/ParametricMesher.h
--- a/Engine/Source/Runtime/Datasmith/CADKernel/Public/CADKernel/Mesh/Meshers/ParametricMesher.h
+++ b/Engine/Source/Runtime/Datasmith/CADKernel/Public/CADKernel/Mesh/Meshers/ParametricMesher.h
@@ -3,6 +3,7 @@
 
 #include "CADKernel/Core/Chrono.h"
 #include "CADKernel/Core/Types.h"
+#include "CADKernel/Mesh/Meshers/MesherReport.h"
 #include "CADKernel/Topo/TopologicalEdge.h"
 #include "CADKernel/Topo/Shell.h"
 #include "CADKernel/UI/Progress.h"
@@ -41,58 +42,6 @@ namespace CADKernel
 		}
 	};
 
-	struct CADKERNEL_API FMesherChronos
-	{
-		FDuration GlobalDuration;
-		FDuration ApplyCriteriaDuration;
-		FDuration IsolateQuadPatchDuration;
-		FDuration GlobalMeshDuration;
-		FDuration GlobalPointCloudDuration;
-		FDuration GlobalGeneratePointCloudDuration;
-		FDuration GlobalTriangulateDuration;
-		FDuration GlobalDelaunayDuration;
-		FDuration GlobalMeshAndGetLoopNodes;
-		FDuration GlobalMeshEdges;
-		FDuration GlobalThinZones;
-
-		FDuration GlobalFindThinZones;
-		//FDuration GlobalScaleGrid;
-		FDuration GlobalMeshThinZones;
-
-		FMesherChronos()
-			: GlobalDuration(FChrono::Init())
-			, ApplyCriteriaDuration(FChrono::Init())
-			, IsolateQuadPatchDuration(FChrono::Init())
-			, GlobalMeshDuration(FChrono::Init())
-			, GlobalPointCloudDuration(FChrono::Init())
-			, GlobalGeneratePointCloudDuration(FChrono::Init())
-			, GlobalTriangulateDuration(FChrono::Init())
-			, GlobalDelaunayDuration(FChrono::Init())
-			, GlobalMeshAndGetLoopNodes(FChrono::Init())
-			, GlobalMeshEdges(FChrono::Init())
-			, GlobalThinZones(FChrono::Init())
-			, GlobalFindThinZones(FChrono::Init())
-			//, GlobalScaleGrid(FChrono::Init())
-			, GlobalMeshThinZones(FChrono::Init())
-		{}
-
-		void PrintTimeElapse() const
-		{
-			FMessage::Printf(Log, TEXT("\n\n\n"));
-			FChrono::PrintClockElapse(Log, TEXT(""), TEXT("Total"), GlobalDuration);
-			FChrono::PrintClockElapse(Log, TEXT("  |  "), TEXT("Apply Criteria"), ApplyCriteriaDuration);
-			FChrono::PrintClockElapse(Log, TEXT("  |  "), TEXT("Find Quad Surfaces"), IsolateQuadPatchDuration);
-			FChrono::PrintClockElapse(Log, TEXT("  |  "), TEXT("Mesh Time"), GlobalMeshDuration);
-			FChrono::PrintClockElapse(Log, TEXT("  |   |  "), TEXT("GeneratePoint Cloud "), GlobalGeneratePointCloudDuration);
-			FChrono::PrintClockElapse(Log, TEXT("  |   |  |  "), TEXT("Point Cloud "), GlobalPointCloudDuration);
-			FChrono::PrintClockElapse(Log, TEXT("  |   |  "), TEXT("ThinZones "), GlobalThinZones);
-			FChrono::PrintClockElapse(Log, TEXT("  |   |  "), TEXT("Mesh ThinZones "), GlobalMeshThinZones);
-			FChrono::PrintClockElapse(Log, TEXT("  |   |  "), TEXT("MeshEdges "), GlobalMeshEdges);
-			FChrono::PrintClockElapse(Log, TEXT("  |   |  "), TEXT("TriangulateDuration "), GlobalTriangulateDuration);
-			FChrono::PrintClockElapse(Log, TEXT("  |   |   |  "), TEXT("Delaunay Duration "), GlobalDelaunayDuration);
-		}
-	};
-
 	class CADKERNEL_API FParametricMesher
 	{
 	protected:
@@ -109,7 +58,7 @@ namespace CADKernel
 		TArray<TSharedPtr<FTopologicalEdge>> Edges;
 		TArray<TSharedPtr<FTopologicalVertex>> Vertices;
 
-		FMesherChronos Chronos;
+		FMesherReport MesherReport;
 
 		bool bDisplay = false;
 
diff --git a/Engine/Source/Runtime/Datasmith/CADKernel/Public/CADKernel/UI/Display.h b/Engine/Source/Runtime/Datasmith/CADKernel/Public/CADKernel/UI/Display.h
--- a/Engine/Source/Runtime/Datasmith/CADKernel/Public/CADKernel/UI/Display.h
+++ b/Engine/Source/Runtime/Datasmith/CADKernel/Public/CADKernel/UI/Display.h
@@ -10,250 +10,262 @@
 class FString;
 namespace CADKernel
 {
-	class FAABB;
-	class FAABB2D;
-	class FBody;
-	class FCurve;
-	class FRestrictionCurve;
-	class FTopologicalEdge;
-	class FEdgeMesh;
-	class FEntity;
-	class FGroup;
-	class FMesh;
-
-	class FVertexMesh;
-	class FEdgeMesh;
-	class FFaceMesh;
-	class FModelMesh;
-
-	class FModel;
-	class FShell;
-	class FSurface;
-	class FFaceMesh;
-	class FTopologicalFace;
-	class FTopologicalLoop;
-	class FTopologicalVertex;
-	class FVertexMesh;
-
-	struct FLinearBoundary;
+
+class FAABB;
+class FAABB2D;
+class FBody;
+class FCurve;
+class FRestrictionCurve;
+class FTopologicalEdge;
+class FEdgeMesh;
+class FEntity;
+class FGroup;
+class FMesh;
+
+class FVertexMesh;
+class FEdgeMesh;
+class FFaceMesh;
+class FModelMesh;
+
+class FModel;
+class FShell;
+class FSurface;
+class FFaceMesh;
+class FTopologicalFace;
+class FTopologicalLoop;
+class FTopologicalVertex;
+class FVertexMesh;
+
+struct FLinearBoundary;
 
 #ifdef CADKERNEL_DEV
-	void Wait(bool bMakeWait = true);
-	void Open3DDebugSession(FString name, const TArray<FIdent>& idList = TArray<FIdent>());
-	inline void Open3DDebugSession(bool bIsDisplayed, FString name, const TArray<FIdent>& idList = TArray<FIdent>()) 
+void Wait(bool bMakeWait = true);
+void Open3DDebugSession(FString name, const TArray<FIdent>& idList = TArray<FIdent>());
+inline void Open3DDebugSession(bool bIsDisplayed, FString name, const TArray<FIdent>& idList = TArray<FIdent>())
+{
+	if (bIsDisplayed)
 	{
-		if (bIsDisplayed)
-		{
-			Open3DDebugSession(name, idList);
-		}
-	};
-	void Close3DDebugSession(bool bIsDisplayed = true);
+		Open3DDebugSession(name, idList);
+	}
+};
+void Close3DDebugSession(bool bIsDisplayed = true);
 #else
-	inline void Wait(bool bMakeWait = true) {};
-	inline void Open3DDebugSession(bool bIsDisplayed, FString name, const TArray<FIdent>& idList = TArray<FIdent>()) {};
-	inline void Open3DDebugSession(FString name, const TArray<FIdent>& idList = TArray<FIdent>()) {};
-	inline void Close3DDebugSession(bool bIsDisplayed = true) {};
+inline void Wait(bool bMakeWait = true) {};
+inline void Open3DDebugSession(bool bIsDisplayed, FString name, const TArray<FIdent>& idList = TArray<FIdent>()) {};
+inline void Open3DDebugSession(FString name, const TArray<FIdent>& idList = TArray<FIdent>()) {};
+inline void Close3DDebugSession(bool bIsDisplayed = true) {};
 #endif
 
-	class CADKERNEL_API F3DDebugSession
+class CADKERNEL_API F3DDebugSession
+{
+private:
+	bool bDisplay = true;
+
+public:
+	F3DDebugSession(FString Name, const TArray<FIdent>& Idents = TArray<FIdent>())
 	{
-	private:
-		bool bDisplay = true;
+		Open3DDebugSession(Name, Idents);
+	}
 
-	public :
-		F3DDebugSession(FString name, const TArray<FIdent>& Idents = TArray<FIdent>())
+	F3DDebugSession(bool bInDisplay, FString Name, const TArray<FIdent>& Idents = TArray<FIdent>())
+		: bDisplay(bInDisplay)
+	{
+		if (bDisplay)
 		{
-			Open3DDebugSession(name, Idents);
+			Open3DDebugSession(Name, Idents);
 		}
+	}
 
-		F3DDebugSession(bool bInDisplay, FString name, const TArray<FIdent>& Idents = TArray<FIdent>())
-			: bDisplay(bInDisplay)
+	~F3DDebugSession()
+	{
+		if (bDisplay)
 		{
-			if (bDisplay)
-			{
-				Open3DDebugSession(name, Idents);
-			}
+			Close3DDebugSession();
 		}
+	}
+};
 
-		~F3DDebugSession()
-		{
-			if (bDisplay)
-			{
-				Close3DDebugSession();
-			}
-		}
-	};
+CADKERNEL_API void Open3DDebugSegment(FIdent Ident);
+CADKERNEL_API void Close3DDebugSegment();
 
-	CADKERNEL_API void Open3DDebugSegment(FIdent Ident);
-	CADKERNEL_API void Close3DDebugSegment();
+class CADKERNEL_API F3DDebugSegment
+{
+public:
+	F3DDebugSegment(FIdent Ident)
+	{
+		Open3DDebugSegment(Ident);
+	}
 
-	class CADKERNEL_API F3DDebugSegment
+	~F3DDebugSegment()
 	{
-	public:
-		F3DDebugSegment(FIdent Ident)
-		{
-			Open3DDebugSegment(Ident);
-		}
+		Close3DDebugSegment();
+	}
+};
 
-		~F3DDebugSegment()
-		{
-			Close3DDebugSegment();
-		}
-	};
+CADKERNEL_API void FlushVisu();
 
-	CADKERNEL_API void FlushVisu();
+template<typename TPoint>
+void DrawPoint(const TPoint& InPoint, EVisuProperty Property = EVisuProperty::BluePoint)
+{
+	FSystem::Get().GetVisu()->DrawPoint(InPoint, Property);
+}
 
-	template<typename TPoint>
-	void DrawPoint(const TPoint& InPoint, EVisuProperty Property = EVisuProperty::BluePoint)
-	{
-		FSystem::Get().GetVisu()->DrawPoint(InPoint, Property);
-	}
+/**
+ * Draw a mesh element i.e. Element of dimension 1 is an edge, of dimension 2 is a triangle or quadrangle, of dimension 3 is tetrahedron, pyramid, hexahedron, ...
+ */
+CADKERNEL_API void DrawElement(int32 Dimension, TArray<FPoint>& Points, EVisuProperty Property = EVisuProperty::Element);
 
-	/**
-	 * Draw a mesh element i.e. Element of dimension 1 is an edge, of dimension 2 is a triangle or quadrangle, of dimension 3 is tetrahedron, pyramid, hexahedron, ...
-	 */
-	CADKERNEL_API void DrawElement(int32 Dimension, TArray<FPoint>& Points, EVisuProperty Property = EVisuProperty::Element);
+template<typename TPoint>
+void Draw(const TArray<TPoint>& Points, EVisuProperty Property = EVisuProperty::BlueCurve)
+{
+	FSystem::Get().GetVisu()->DrawPolyline(Points, Property);
+}
 
-	template<typename TPoint>
-	void Draw(const TArray<TPoint>& Points, EVisuProperty Property = EVisuProperty::BlueCurve)
-	{
-		FSystem::Get().GetVisu()->DrawPolyline(Points, Property);
-	}
+CADKERNEL_API void DrawMesh(const TSharedPtr<FMesh>& mesh);
 
-	CADKERNEL_API void DrawMesh(const TSharedPtr<FMesh>& mesh);
+CADKERNEL_API void DisplayEdgeCriteriaGrid(int32 EdgeId, const TArray<FPoint>& Points3D);
 
-	CADKERNEL_API void DisplayEdgeCriteriaGrid(int32 EdgeId, const TArray<FPoint>& Points3D);
+template<typename TPoint>
+void DisplayPoint(const TPoint& Point, FIdent Ident)
+{
+	F3DDebugSegment G(Ident);
+	DrawPoint(Point);
+}
 
-	template<typename TPoint>
-	void DisplayPoint(const TPoint& Point, FIdent Ident)
-	{
-		F3DDebugSegment G(Ident);
-		DrawPoint(Point);
-	}
+template<typename TPoint>
+void DisplayPoint(const TPoint& Point, EVisuProperty Property = EVisuProperty::BluePoint)
+{
+	DrawPoint(Point, Property);
+}
 
-	template<typename TPoint>
-	void DisplayPoint(const TPoint& Point, EVisuProperty Property = EVisuProperty::BluePoint)
-	{
-		DrawPoint(Point, Property);
-	}
+template<typename TPoint>
+void DisplayPoint(const TPoint& Point, EVisuProperty Property, FIdent Ident)
+{
+	F3DDebugSegment G(Ident);
+	DrawPoint(Point, Property);
+}
 
-	template<typename TPoint>
-	void DisplayPoint(const TPoint& Point, EVisuProperty Property, FIdent Ident)
+template<typename TPoint>
+void DisplayPoints(FString Message, const TArray<TPoint>& Points, EVisuProperty Property = EVisuProperty::BluePoint, bool bDisplay = true)
+{
+	if (!bDisplay)
 	{
-		F3DDebugSegment G(Ident);
-		DrawPoint(Point, Property);
+		return;
 	}
 
-	template<typename TPoint>
-	void DisplayPoints(FString Message, const TArray<TPoint>& Points, EVisuProperty Property = EVisuProperty::BluePoint, bool bDisplay = true)
+	Open3DDebugSession(Message);
+	for (int32 Index = 0; Index < Points.Num(); ++Index)
 	{
-		if (!bDisplay)
-		{
-			return;
-		}
-
-		Open3DDebugSession(Message);
-		for (int32 Index = 0; Index < Points.Num(); ++Index)
-		{
-			DisplayPoint(Points[Index], Property);
-		}
-		Close3DDebugSession();
+		DisplayPoint(Points[Index], Property);
 	}
+	Close3DDebugSession();
+}
 
-	CADKERNEL_API void DisplayProductTree(const FEntity& RootId);
-	CADKERNEL_API void DisplayProductTree(const FModel& Model);
-	CADKERNEL_API void DisplayProductTree(const FBody& Body);
-	CADKERNEL_API void DisplayProductTree(const FShell& Shell);
+CADKERNEL_API void DisplayProductTree(const FEntity& RootId);
+CADKERNEL_API void DisplayProductTree(const FModel& Model);
+CADKERNEL_API void DisplayProductTree(const FBody& Body);
+CADKERNEL_API void DisplayProductTree(const FShell& Shell);
 
-	CADKERNEL_API void DisplayAABB(const FAABB& aabb, FIdent Ident = 0);
-	CADKERNEL_API void DisplayAABB2D(const FAABB2D& aabb, FIdent Ident = 0);
+CADKERNEL_API void DisplayAABB(const FAABB& aabb, FIdent Ident = 0);
+CADKERNEL_API void DisplayAABB2D(const FAABB2D& aabb, FIdent Ident = 0);
 
-	CADKERNEL_API void DisplayEntity(const FEntity& Entity);
-	CADKERNEL_API void DisplayEntity2D(const FEntity& Entity);
+CADKERNEL_API void DisplayEntity(const FEntity& Entity);
+CADKERNEL_API void DisplayEntity2D(const FEntity& Entity);
 
-	CADKERNEL_API void DisplayLoop(const FTopologicalFace& Entity);
-	CADKERNEL_API void DisplayIsoCurve(const FSurface& CarrierSurface, double Coordinate, EIso Type);
+CADKERNEL_API void DisplayLoop(const FTopologicalFace& Entity);
+CADKERNEL_API void DisplayIsoCurve(const FSurface& CarrierSurface, double Coordinate, EIso Type);
 
-	CADKERNEL_API void Display(const FPlane& plane, FIdent Ident = 0);
+CADKERNEL_API void Display(const FPlane& plane, FIdent Ident = 0);
 
-	CADKERNEL_API void Display(const FCurve& Curve);
-	CADKERNEL_API void Display(const FSurface& CarrierSurface);
+CADKERNEL_API void Display(const FCurve& Curve);
+CADKERNEL_API void Display(const FSurface& CarrierSurface);
 
-	CADKERNEL_API void Display(const FGroup& Group);
-	CADKERNEL_API void Display(const FModel& Model);
-	CADKERNEL_API void Display(const FBody& Body);
-	CADKERNEL_API void Display(const FShell& Shell);
-	CADKERNEL_API void Display(const FTopologicalEdge& Edge, EVisuProperty Property = EVisuProperty::BlueCurve);
-	CADKERNEL_API void Display(const FTopologicalFace& Face);
-	CADKERNEL_API void Display(const FTopologicalLoop& Loop);
-	CADKERNEL_API void Display(const FTopologicalVertex& Vertex);
+CADKERNEL_API void Display(const FGroup& Group);
+CADKERNEL_API void Display(const FModel& Model);
+CADKERNEL_API void Display(const FBody& Body);
+CADKERNEL_API void Display(const FShell& Shell);
+CADKERNEL_API void Display(const FTopologicalEdge& Edge, EVisuProperty Property = EVisuProperty::BlueCurve);
+CADKERNEL_API void Display(const FTopologicalFace& Face);
+CADKERNEL_API void Display(const FTopologicalLoop& Loop);
+CADKERNEL_API void Display(const FTopologicalVertex& Vertex);
 
-	CADKERNEL_API void Display2D(const FTopologicalEdge& Edge, EVisuProperty Property = EVisuProperty::BlueCurve);
-	CADKERNEL_API void Display2D(const FTopologicalFace& Face);
-	CADKERNEL_API void Display2D(const FTopologicalLoop& Loop);
-	CADKERNEL_API void Display2D(const FSurface& CarrierSurface);
+CADKERNEL_API void Display2D(const FTopologicalEdge& Edge, EVisuProperty Property = EVisuProperty::BlueCurve);
+CADKERNEL_API void Display2D(const FTopologicalFace& Face);
+CADKERNEL_API void Display2D(const FTopologicalLoop& Loop);
+CADKERNEL_API void Display2D(const FSurface& CarrierSurface);
 
-	CADKERNEL_API void DisplayMesh(const FFaceMesh& Mesh);
-	CADKERNEL_API void DisplayMesh(const FEdgeMesh& Mesh);
-	CADKERNEL_API void DisplayMesh(const FVertexMesh& Mesh);
+CADKERNEL_API void DisplayMesh(const FFaceMesh& Mesh);
+CADKERNEL_API void DisplayMesh(const FEdgeMesh& Mesh);
+CADKERNEL_API void DisplayMesh(const FVertexMesh& Mesh);
 
-	CADKERNEL_API void Display(const FModelMesh& MeshModel);
+CADKERNEL_API void Display(const FModelMesh& MeshModel);
 
-	CADKERNEL_API void DisplayControlPolygon(const FCurve& Entity);
-	CADKERNEL_API void DisplayControlPolygon(const FSurface& Entity);
+CADKERNEL_API void DisplayControlPolygon(const FCurve& Entity);
+CADKERNEL_API void DisplayControlPolygon(const FSurface& Entity);
 
-	template<typename TPoint>
-	void DisplaySegment(const TPoint& Point1, const TPoint& Point2, FIdent Ident = 0, EVisuProperty Property = EVisuProperty::Element, bool bWithOrientation = false)
+template<typename TPoint>
+void DisplaySegment(const TPoint& Point1, const TPoint& Point2, FIdent Ident = 0, EVisuProperty Property = EVisuProperty::Element, bool bWithOrientation = false)
+{
+	F3DDebugSegment G(Ident);
+	if (bWithOrientation)
 	{
-		F3DDebugSegment G(Ident);
-		if (bWithOrientation)
-		{
-			DrawSegmentOrientation(Point1, Point2, Property);
-		}
-		DrawSegment(Point1, Point2, Property);
-	};
+		DrawSegmentOrientation(Point1, Point2, Property);
+	}
+	DrawSegment(Point1, Point2, Property);
+};
 
-	template<typename TPoint>
-	void DisplayPolyline(const TArray<TPoint>& Points, EVisuProperty Property)
+template<typename TPoint>
+void DisplayPolyline(const TArray<TPoint>& Points, EVisuProperty Property)
+{
+	Open3DDebugSegment(0);
+	Draw(Points, Property);
+	Close3DDebugSegment();
+}
+
+template<typename TPoint>
+void DisplayOrientedPolyline(const TArray<TPoint>& Points, EVisuProperty Property)
+{
+	Open3DDebugSegment(0);
+	for (int32 Index = 1; Index < Points.Num(); ++Index)
 	{
-		Open3DDebugSegment(0);
-		Draw(Points, Property);
-		Close3DDebugSegment();
+		DisplaySegment(Points[Index-1], Points[Index], Index, Property, true);
 	}
+	Close3DDebugSegment();
+}
 
-	CADKERNEL_API void DrawQuadripode(double Height, double Base, FPoint& Centre, FPoint& Direction);
+CADKERNEL_API void DrawQuadripode(double Height, double Base, FPoint& Centre, FPoint& Direction);
 
-	CADKERNEL_API void Draw(const FTopologicalEdge& Edge, EVisuProperty Property = EVisuProperty::BlueCurve);
-	CADKERNEL_API void Draw(const FTopologicalFace& Face);
-	CADKERNEL_API void Draw2D(const FTopologicalFace& Face);
-	CADKERNEL_API void Draw(const FShell& Shell);
+CADKERNEL_API void Draw(const FTopologicalEdge& Edge, EVisuProperty Property = EVisuProperty::BlueCurve);
+CADKERNEL_API void Draw(const FTopologicalFace& Face);
+CADKERNEL_API void Draw2D(const FTopologicalFace& Face);
+CADKERNEL_API void Draw(const FShell& Shell);
 
-	CADKERNEL_API void Draw(const FCurve& Curve, EVisuProperty Property = EVisuProperty::BlueCurve);
-	CADKERNEL_API void Draw(const FCurve& Curve, const FLinearBoundary& Boundary, EVisuProperty Property = EVisuProperty::BlueCurve);
-	CADKERNEL_API void Draw(const FLinearBoundary& Boundary, const FRestrictionCurve& Curve, EVisuProperty Property = EVisuProperty::BlueCurve);
+CADKERNEL_API void Draw(const FCurve& Curve, EVisuProperty Property = EVisuProperty::BlueCurve);
+CADKERNEL_API void Draw(const FCurve& Curve, const FLinearBoundary& Boundary, EVisuProperty Property = EVisuProperty::BlueCurve);
+CADKERNEL_API void Draw(const FLinearBoundary& Boundary, const FRestrictionCurve& Curve, EVisuProperty Property = EVisuProperty::BlueCurve);
 
-	template<typename TPoint>
-	void DrawSegment(const TPoint& Point1, const TPoint& Point2, EVisuProperty Property = EVisuProperty::Element)
-	{
-		TArray<FPoint> Points;
-		Points.Add(Point1);
-		Points.Add(Point2);
-		Draw(Points, Property);
-	}
-
-	template<typename TPoint>
-	void DrawSegmentOrientation(const TPoint& Point1, const TPoint& Point2, EVisuProperty Property = EVisuProperty::Element)
-	{
-		double Length = Point1.Distance(Point2);
-		double Height = Length / 10.0;
-		double Base = Height / 2;
+template<typename TPoint>
+void DrawSegment(const TPoint& Point1, const TPoint& Point2, EVisuProperty Property = EVisuProperty::Element)
+{
+	TArray<FPoint> Points;
+	Points.Add(Point1);
+	Points.Add(Point2);
+	Draw(Points, Property);
+}
+
+template<typename TPoint>
+void DrawSegmentOrientation(const TPoint& Point1, const TPoint& Point2, EVisuProperty Property = EVisuProperty::Element)
+{
+	double Length = Point1.Distance(Point2);
+	double Height = Length / 10.0;
+	double Base = Height / 2;
 
-		FPoint Middle = (Point1 + Point2) / 2.;
-		FPoint Tangent = Point2 - Point1;
-		DrawQuadripode(Height, Base, Middle, Tangent);
-	}
+	FPoint Middle = (Point1 + Point2) / 2.;
+	FPoint Tangent = Point2 - Point1;
+	DrawQuadripode(Height, Base, Middle, Tangent);
+}
 
-	CADKERNEL_API void DrawIsoCurves(const FTopologicalFace& Face);
+CADKERNEL_API void DrawIsoCurves(const FTopologicalFace& Face);
 
 
 } // namespace CADKernel
