fix Static Analysis Win64 (PVS-Studio) - CL 18519525

diff --git a/Engine/Source/Runtime/Datasmith/CADKernel/Private/CADKernel/Mesh/Structure/LoopCleaner.cpp b/Engine/Source/Runtime/Datasmith/CADKernel/Private/CADKernel/Mesh/Structure/LoopCleaner.cpp
--- a/Engine/Source/Runtime/Datasmith/CADKernel/Private/CADKernel/Mesh/Structure/LoopCleaner.cpp
+++ b/Engine/Source/Runtime/Datasmith/CADKernel/Private/CADKernel/Mesh/Structure/LoopCleaner.cpp
@@ -51,23 +51,22 @@ bool FLoopCleaner::CleanLoops()
 #endif
 
 	// for each loop, start by the best node, find all intersections
-	bool bIsAnOuterLoop = true;
 	LoopIndex = -1;
 	NextLoopFirstSegmentIndex = 0;
+
+	GetNext = LoopCleanerImpl::GetNextNodeImpl;
+	GetFirst = LoopCleanerImpl::GetFirstNode;
+	GetSecond = LoopCleanerImpl::GetSecondNode;
+
 	for (FLoopNode* StartNode : BestStartNodeOfLoops)
 	{
 		LoopIndex++;
 		StartSegmentIndex = NextLoopFirstSegmentIndex;
 		UpdateNextLoopFirstSegmentIndex(LoopIndex);
 
-		bLoopOrientation = EOrientation::Front;
-		GetNext = (bLoopOrientation == EOrientation::Front) ? LoopCleanerImpl::GetNextNodeImpl : LoopCleanerImpl::GetPreviousNodeImpl;
-		GetFirst = (bLoopOrientation == EOrientation::Front) ? LoopCleanerImpl::GetFirstNode : LoopCleanerImpl::GetSecondNode;
-		GetSecond = (bLoopOrientation == EOrientation::Front) ? LoopCleanerImpl::GetSecondNode : LoopCleanerImpl::GetFirstNode;
-
 		Intersections.Empty(5);
 
-		GetLoopNodeStartingFrom(StartNode, bIsAnOuterLoop, NodesOfLoop);
+		GetLoopNodeStartingFrom(StartNode, NodesOfLoop);
 
 #ifdef DEBUG_CLEAN_LOOPS		
 		Grid.DisplayGridPolyline(TEXT("Loop: start"), EGridSpace::UniformScaled, NodesOfLoop, true, EVisuProperty::YellowCurve);
diff --git a/Engine/Source/Runtime/Datasmith/CADKernel/Private/CADKernel/Mesh/Structure/LoopCleaner.h b/Engine/Source/Runtime/Datasmith/CADKernel/Private/CADKernel/Mesh/Structure/LoopCleaner.h
--- a/Engine/Source/Runtime/Datasmith/CADKernel/Private/CADKernel/Mesh/Structure/LoopCleaner.h
+++ b/Engine/Source/Runtime/Datasmith/CADKernel/Private/CADKernel/Mesh/Structure/LoopCleaner.h
@@ -364,7 +364,7 @@ private:
 
 	bool Fill(LoopCleanerImpl::FPinchIntersectionContext& Context);
 
-	void GetLoopNodeStartingFrom(FLoopNode* StartNode, bool bForward, TArray<FLoopNode*>& Loop)
+	void GetLoopNodeStartingFrom(FLoopNode* StartNode, TArray<FLoopNode*>& Loop)
 	{
 		FLoopNode* Node = StartNode;
 		Loop.Empty(LoopNodes.Num());
