fix Static Analysis Win64 (PVS-Studio) - CL 18519525

diff --git a/Engine/Source/Runtime/Datasmith/CADKernel/Private/CADKernel/Mesh/Structure/LoopCleaner.cpp b/Engine/Source/Runtime/Datasmith/CADKernel/Private/CADKernel/Mesh/Structure/LoopCleaner.cpp
new file mode 100644
--- /dev/null
+++ b/Engine/Source/Runtime/Datasmith/CADKernel/Private/CADKernel/Mesh/Structure/LoopCleaner.cpp
@@ -0,0 +1,2287 @@
+// Copyright Epic Games, Inc. All Rights Reserved.
+
+#include "CADKernel/Mesh/Structure/LoopCleaner.h"
+
+#include "CADKernel/Mesh/Meshers/IsoTriangulator.h"
+#include "CADKernel/Mesh/Meshers/IsoTriangulator/IsoNode.h"
+#include "CADKernel/Mesh/Meshers/IsoTriangulator/IsoSegment.h"
+#include "CADKernel/Mesh/Structure/Grid.h"
+#include "CADKernel/Topo/TopologicalFace.h"
+
+namespace CADKernel
+{
+
+FLoopCleaner::FLoopCleaner(FIsoTriangulator& Triangulator)
+	: Grid(Triangulator.Grid)
+	, LoopNodes(Triangulator.LoopNodes)
+	, LoopSegments(Triangulator.LoopSegments)
+	, IsoSegmentFactory(Triangulator.IsoSegmentFactory)
+	, MesherReport(Triangulator.MesherReport)
+	, bDisplay(Triangulator.bDisplay)
+	, LoopSegmentsIntersectionTool(Grid)
+	, GeometricTolerance(Triangulator.GeometricTolerance)
+	, SquareGeometricTolerance(FMath::Square(GeometricTolerance))
+	, SquareGeometricTolerance2(2. * SquareGeometricTolerance)
+	, GetNext(LoopCleanerImpl::GetNextNodeImpl)
+	, GetPrevious(LoopCleanerImpl::GetPreviousNodeImpl)
+	, GetFirst(LoopCleanerImpl::GetFirstNode)
+	, GetSecond(LoopCleanerImpl::GetSecondNode)
+{
+
+}
+
+bool FLoopCleaner::CleanLoops()
+{
+
+	//{
+	//	FPoint2D Point = LoopNodes[36].Get2DPoint(EGridSpace::UniformScaled, Grid);
+	//	LoopNodes[56].Set2DPoint(EGridSpace::UniformScaled, Grid, Point);
+	//}
+
+	FindBestLoopExtremity();
+
+#ifdef DEBUG_CLEAN_LOOPS		
+	if (bDisplay)
+	{
+		Grid.DisplayIsoSegments(TEXT("Loops Orientation"), EGridSpace::UniformScaled, LoopSegments, false, true, EVisuProperty::BlueCurve);
+		Grid.DisplayIsoSegments(TEXT("Loops Init"), EGridSpace::UniformScaled, LoopSegments, true, false, EVisuProperty::BlueCurve);
+		Grid.DisplayIsoNodes(TEXT("BestStartNodeOfLoop"), EGridSpace::UniformScaled, (const TArray<FIsoNode*>&) BestStartNodeOfLoops, EVisuProperty::BluePoint);
+		Wait(bDisplay);
+	}
+#endif
+
+	// for each loop, start by the best node, find all intersections
+	LoopIndex = -1;
+	NextLoopFirstSegmentIndex = 0;
+
+	GetNext = LoopCleanerImpl::GetNextNodeImpl;
+	GetFirst = LoopCleanerImpl::GetFirstNode;
+	GetSecond = LoopCleanerImpl::GetSecondNode;
+
+	for (FLoopNode* StartNode : BestStartNodeOfLoops)
+	{
+		LoopIndex++;
+		StartSegmentIndex = NextLoopFirstSegmentIndex;
+		UpdateNextLoopFirstSegmentIndex(LoopIndex);
+
+		Intersections.Empty(5);
+
+		GetLoopNodeStartingFrom(StartNode, NodesOfLoop);
+
+#ifdef DEBUG_CLEAN_LOOPS		
+		Grid.DisplayGridPolyline(TEXT("Loop: start"), EGridSpace::UniformScaled, NodesOfLoop, true, EVisuProperty::YellowCurve);
+		//Wait(bDisplay);
+#endif
+
+		if (!RemoveLoopPicks())
+		{
+			continue;
+		}
+
+#ifdef DEBUG_CLEAN_LOOPS		
+		Grid.DisplayGridPolyline(TEXT("LoopIntersections: remove loop's picks"), EGridSpace::UniformScaled, NodesOfLoop, true, EVisuProperty::YellowCurve);
+		//Wait(bDisplay);
+#endif
+		
+		FindLoopIntersections();
+
+		if (!RemoveSelfIntersectionsOfLoop())
+		{
+			FMessage::Printf(Log, TEXT("Loop intersections of the surface %d cannot be fixed. The mesh of this surface is canceled.\n"), Grid.GetFace()->GetId());
+			return false;
+		}
+
+#ifdef DEBUG_CLEAN_LOOPS		
+		Grid.DisplayGridPolyline(TEXT("LoopIntersections: remove loop's self intersections"), EGridSpace::UniformScaled, NodesOfLoop, true, EVisuProperty::YellowCurve);
+		//Wait(bDisplay);
+#endif
+
+		if (!RemoveLoopPicks())
+		{
+			continue;
+		}
+
+		if (NodesOfLoopCount == 0)
+		{
+			continue;
+		}
+
+#ifdef DEBUG_CLEAN_LOOPS		
+		Grid.DisplayGridPolyline(TEXT("LoopIntersections: remove loop's picks (step2)"), EGridSpace::UniformScaled, NodesOfLoop, true, EVisuProperty::YellowCurve);
+		//Wait(bDisplay);
+#endif
+
+		FixLoopOrientation();
+	}
+
+	if (!CheckMainLoopConsistency())
+	{
+		return false;
+	}
+
+	return true;
+}
+
+//#define DEBUG_UNCROSS_LOOPS
+bool FLoopCleaner::UncrossLoops()
+{
+	if (Grid.GetLoopCount() == 1)
+	{
+		return true;
+	}
+
+#ifdef DEBUG_UNCROSS_LOOPS
+	int32 Iteration = 0;
+	F3DDebugSession _(bDisplay, TEXT("FixIntersectionBetweenLoops"));
+#endif
+
+	TSet<uint32> IntersectionAlreadyProceed;
+
+	LoopSegmentsIntersectionTool.Empty(LoopSegments.Num());
+	int32 Index = 0;
+	for (; Index < LoopSegments.Num(); ++Index)
+	{
+		if (((FLoopNode&)LoopSegments[Index]->GetFirstNode()).GetLoopIndex() != 0)
+		{
+			break;
+		}
+		LoopSegmentsIntersectionTool.AddSegment(*LoopSegments[Index]);
+	}
+
+	for (; Index < LoopSegments.Num(); )
+	{
+		ensureCADKernel(LoopSegments[Index]);
+
+		FIsoSegment& Segment = *LoopSegments[Index];
+		ensureCADKernel(!Segment.IsDelete());
+
+#ifdef DEBUG_UNCROSS_LOOPS
+		if (bDisplay)
+		{
+			LoopSegmentsIntersectionTool.Display(TEXT("IntersectionTool"));
+			F3DDebugSession _(*FString::Printf(TEXT("Segment to proceed %d %d"), Index, Iteration++));
+			Grid.DisplayIsoSegment(EGridSpace::UniformScaled, Segment, 0, EVisuProperty::BlueCurve);
+		}
+#endif
+
+		if (const FIsoSegment* IntersectingSegment = LoopSegmentsIntersectionTool.DoesIntersect(Segment))
+		{
+#ifdef DEBUG_UNCROSS_LOOPS
+			if (bDisplay)
+			{
+				{
+					LoopSegmentsIntersectionTool.Display(TEXT("IntersectionTool"));
+					F3DDebugSession _(*FString::Printf(TEXT("Segment to proceed %d %d"), Index, Iteration++));
+					Grid.DisplayIsoSegment(EGridSpace::UniformScaled, Segment, 0, EVisuProperty::BlueCurve);
+				}
+				{
+					F3DDebugSession _(TEXT("Intersecting Segments"));
+					Grid.DisplayIsoSegment(EGridSpace::UniformScaled, *IntersectingSegment, 0, EVisuProperty::RedCurve/*, true*/);
+				}
+				Wait(bDisplay);
+			}
+#endif
+
+			uint32 IntersectionHash = GetTypeHash(*IntersectingSegment, Segment);
+			bool bNotProceed = !IntersectionAlreadyProceed.Find(IntersectionHash);
+			IntersectionAlreadyProceed.Add(IntersectionHash);
+
+			bool bIsFixed = true;
+			bool bIsSameLoop = ((FLoopNode&)Segment.GetFirstNode()).GetLoopIndex() && ((FLoopNode&)Segment.GetFirstNode()).GetLoopIndex() == ((FLoopNode&)IntersectingSegment->GetFirstNode()).GetLoopIndex();
+
+			// Swapping segments is possible only with outer loop
+			if (bNotProceed)
+			{
+				if (!TryToRemoveIntersectionOfTwoConsecutiveIntersectingSegments(*IntersectingSegment, Segment))
+				{
+					if (!TryToRemoveIntersectionOfTwoConsecutiveIntersectingSegments(Segment, const_cast<FIsoSegment&>(*IntersectingSegment)))
+					{
+						RemoveIntersectionByMovingOutsideSegmentNodeInside(*IntersectingSegment, Segment, ((FLoopNode&)Segment.GetFirstNode()).GetLoopIndex() && bIsSameLoop);
+					}
+				}
+
+				// segment is moved, the previous segment is retested. Thanks to that, only one loop is enough.
+				if (Index > 1)
+				{
+					LoopSegmentsIntersectionTool.RemoveLast();
+				}
+			}
+			else if (bIsSameLoop)
+			{
+				SwapSegments(const_cast<FIsoSegment&>(*IntersectingSegment), Segment);
+			}
+			else
+			{
+				// if no more append remove it
+				ensureCADKernel(false);
+				LoopSegmentsIntersectionTool.AddSegment(Segment);
+			}
+
+#ifdef DEBUG_UNCROSS_LOOPS
+			if (true)
+			{
+				Grid.DisplayGridPolyline(TEXT("After fix"), EGridSpace::UniformScaled, LoopNodes, true, EVisuProperty::YellowCurve);
+				Wait(true);
+			}
+#endif
+		}
+		else
+		{
+			LoopSegmentsIntersectionTool.AddSegment(Segment);
+		}
+
+		if (!Segment.IsDelete())
+		{
+			RemovePickRecursively((FLoopNode*) &Segment.GetFirstNode(), (FLoopNode*) &Segment.GetSecondNode());
+#ifdef DEBUG_UNCROSS_LOOPS
+			if (true)
+			{
+				Grid.DisplayGridPolyline(TEXT("After RemovePickRecursively"), EGridSpace::UniformScaled, LoopNodes, true, EVisuProperty::YellowCurve);
+				Wait(true);
+			}
+#endif
+		}
+
+		Index = LoopSegmentsIntersectionTool.Count();
+	}
+
+#ifdef DEBUG_UNCROSS_LOOPS
+	if (bDisplay)
+	{
+		Grid.DisplayGridPolyline(TEXT("After UncrossLoops"), EGridSpace::UniformScaled, LoopNodes, true, EVisuProperty::YellowCurve);
+		Wait(true);
+	}
+#endif
+
+	return true;
+}
+
+bool FLoopCleaner::TryToRemoveIntersectionOfTwoConsecutiveIntersectingSegments(const FIsoSegment& IntersectingSegment, FIsoSegment& Segment)
+{
+	FLoopNode* Node = nullptr;
+	FLoopNode* PreviousNode = nullptr;
+	FLoopNode* NextNode = nullptr;
+
+	TSegment<FPoint2D> IntersectingSegment2D(IntersectingSegment.GetFirstNode().Get2DPoint(EGridSpace::UniformScaled, Grid), IntersectingSegment.GetSecondNode().Get2DPoint(EGridSpace::UniformScaled, Grid));
+	TSegment<FPoint2D> Segment2D(Segment.GetFirstNode().Get2DPoint(EGridSpace::UniformScaled, Grid), Segment.GetSecondNode().Get2DPoint(EGridSpace::UniformScaled, Grid));
+
+	double IntersectingSegmentSlop = ComputeOrientedSlope(IntersectingSegment2D.Point0, IntersectingSegment2D.Point1, 0);
+	double SegmentSlop = ComputeUnorientedSlope(Segment2D.Point1, Segment2D.Point0, IntersectingSegmentSlop);
+	if (SegmentSlop > 2)
+	{
+		SegmentSlop = 4 - SegmentSlop;
+	}
+
+	// if the segment and IntersectingSegment are parallel, segment are moved inside
+	if (SegmentSlop < 0.01)
+	{
+		double StartPointSquareDistance = SquareDistanceOfPointToSegment(Segment2D.Point0, IntersectingSegment2D.Point0, IntersectingSegment2D.Point1);
+		double EndPointSquareDistance = SquareDistanceOfPointToSegment(Segment2D.Point1, IntersectingSegment2D.Point0, IntersectingSegment2D.Point1);
+		if (StartPointSquareDistance < SquareGeometricTolerance && EndPointSquareDistance < SquareGeometricTolerance)
+		{
+			OffsetSegment(Segment, Segment2D, IntersectingSegment2D);
+			return true;
+		}
+	}
+
+	// check if the intersection is not at the extremity
+	{
+		double Coordinate = 0;
+		FindIntersectionOfSegments2D(Segment2D, IntersectingSegment2D, Coordinate);
+		if (FMath::IsNearlyZero(Coordinate))
+		{
+			// can add a test to offset the outside node and not the node a 
+			OffsetNode((FLoopNode&)Segment.GetFirstNode(), IntersectingSegment2D);
+			return true;
+		}
+		else if (FMath::IsNearlyEqual(Coordinate, 1))
+		{
+			OffsetNode((FLoopNode&)Segment.GetSecondNode(), IntersectingSegment2D);
+			return true;
+		}
+	}
+
+	double OrientedSlop = ComputeOrientedSlope(IntersectingSegment2D.Point0, Segment.GetFirstNode().Get2DPoint(EGridSpace::UniformScaled, Grid), IntersectingSegmentSlop);
+	if (OrientedSlop >= 0)
+	{
+		Node = (FLoopNode*)&Segment.GetSecondNode();
+		PreviousNode = (FLoopNode*)&Segment.GetFirstNode();
+		NextNode = (FLoopNode*)&Node->GetNextNode();
+	}
+	else
+	{
+		Node = (FLoopNode*)&Segment.GetFirstNode();
+		PreviousNode = (FLoopNode*)&Segment.GetSecondNode();
+		NextNode = (FLoopNode*)&Node->GetPreviousNode();
+	}
+
+#ifdef DEBUG_TWO_CONSECUTIVE_INTERSECTING
+	if (bDisplay)
+	{
+		F3DDebugSession _(TEXT("Intersecting Segments"));
+		Display(EGridSpace::UniformScaled, Segment, 0, EVisuProperty::BlueCurve);
+		Display(EGridSpace::UniformScaled, *Node, *NextNode, 0, EVisuProperty::BlueCurve);
+		Display(EGridSpace::UniformScaled, IntersectingSegment, 0, EVisuProperty::RedCurve);
+		Display(EGridSpace::UniformScaled, *Node, 0, EVisuProperty::RedPoint);
+		Wait(false);
+	}
+#endif
+
+	TSegment<FPoint2D> NextSegment2D(Node->Get2DPoint(EGridSpace::UniformScaled, Grid), NextNode->Get2DPoint(EGridSpace::UniformScaled, Grid));
+	if (!FastIntersectSegments2D(NextSegment2D, IntersectingSegment2D))
+	{
+		return false;
+	}
+
+	TSegment<FPoint2D>  PreviousSegment2D(Node->Get2DPoint(EGridSpace::UniformScaled, Grid), PreviousNode->Get2DPoint(EGridSpace::UniformScaled, Grid));
+
+	FPoint2D Intersection1 = FindIntersectionOfSegments2D(PreviousSegment2D, IntersectingSegment2D);
+	FPoint2D Intersection2 = FindIntersectionOfSegments2D(NextSegment2D, IntersectingSegment2D);
+
+	double Coordinate;
+	FPoint2D ProjectedPoint = ProjectPointOnSegment(Node->Get2DPoint(EGridSpace::UniformScaled, Grid), Intersection1, Intersection2, Coordinate);
+
+#ifdef DEBUG_TWO_CONSECUTIVE_INTERSECTING
+	if (bDisplay)
+	{
+		F3DDebugSession _(TEXT("ProjectedPoint"));
+		DisplayPoint(ProjectedPoint, EVisuProperty::RedPoint);
+	}
+#endif
+
+	FPoint2D SegmentTangent = IntersectingSegment2D.Point1 - IntersectingSegment2D.Point0;
+	SegmentTangent.Normalize();
+	FPoint2D MoveDirection = SegmentTangent.GetPerpendicularVector();
+
+	MoveDirection *= GeometricTolerance;
+
+	FPoint2D NewCoordinate = ProjectedPoint + MoveDirection;
+
+#ifdef DEBUG_TWO_CONSECUTIVE_INTERSECTING
+	if (bDisplay)
+	{
+		F3DDebugSession _(TEXT("NewCoordinate"));
+		DisplayPoint(NewCoordinate, EVisuProperty::BluePoint);
+	}
+#endif
+
+	if ((PreviousNode->Get2DPoint(EGridSpace::UniformScaled, Grid).SquareDistance(NewCoordinate) < SquareGeometricTolerance2) ||
+		(NextNode->Get2DPoint(EGridSpace::UniformScaled, Grid).SquareDistance(NewCoordinate) < SquareGeometricTolerance2))
+	{
+		if(RemoveNodeOfLoop(*Node))
+		{
+			return true;
+		}
+	}
+	else
+	{
+		Node->Set2DPoint(EGridSpace::UniformScaled, Grid, NewCoordinate);
+		FIsoSegment* NextSegment = Node->GetSegmentConnectedTo(NextNode);
+		LoopSegmentsIntersectionTool.Update(&Segment);
+		LoopSegmentsIntersectionTool.Update(NextSegment);
+
+#ifdef DEBUG_TWO_CONSECUTIVE_INTERSECTING
+		if (bDisplay)
+		{
+			F3DDebugSession _(TEXT("New position"));
+			DisplayPoint(Intersection1, EVisuProperty::RedPoint);
+			DisplayPoint(Intersection2, EVisuProperty::RedPoint);
+			Display(EGridSpace::UniformScaled, *Node, EVisuProperty::RedPoint);
+			Display(EGridSpace::UniformScaled, Segment, 0, EVisuProperty::GreenCurve/*, true*/);
+			Display(EGridSpace::UniformScaled, *Node, *NextNode, 0, EVisuProperty::GreenCurve/*, true*/);
+			Wait(false);
+		}
+#endif
+	}
+
+	return true;
+}
+
+void FLoopCleaner::OffsetSegment(FIsoSegment& Segment, TSegment<FPoint2D>& Segment2D, TSegment<FPoint2D>& IntersectingSegment2D)
+{
+	FPoint2D SegmentTangent = IntersectingSegment2D.Point1 - IntersectingSegment2D.Point0;
+	SegmentTangent.Normalize();
+	FPoint2D MoveDirection = SegmentTangent.GetPerpendicularVector();
+	MoveDirection *= GeometricTolerance;
+
+	FPoint2D NewPoint0 = Segment2D.Point0 + MoveDirection;
+	FPoint2D NewPoint1 = Segment2D.Point1 + MoveDirection;
+
+	Segment.GetFirstNode().Set2DPoint(EGridSpace::UniformScaled, Grid, NewPoint0);
+	Segment.GetSecondNode().Set2DPoint(EGridSpace::UniformScaled, Grid, NewPoint1);
+}
+
+void FLoopCleaner::OffsetNode(FLoopNode& Node, TSegment<FPoint2D>& IntersectingSegment2D)
+{
+	FPoint2D SegmentTangent = IntersectingSegment2D.Point1 - IntersectingSegment2D.Point0;
+	SegmentTangent.Normalize();
+	FPoint2D MoveDirection = SegmentTangent.GetPerpendicularVector();
+	MoveDirection *= GeometricTolerance;
+
+	FPoint2D NewPoint = Node.Get2DPoint(EGridSpace::UniformScaled, Grid) + MoveDirection;
+	Node.Set2DPoint(EGridSpace::UniformScaled, Grid, NewPoint);
+}
+
+void FLoopCleaner::RemoveIntersectionByMovingOutsideSegmentNodeInside(const FIsoSegment& IntersectingSegment, const FIsoSegment& Segment, bool bIsSameInnerLoop)
+{
+	const FIsoNode* Nodes[2][2] = { {nullptr, nullptr}, {nullptr, nullptr} };
+	Nodes[0][0] = &IntersectingSegment.GetFirstNode();
+	Nodes[0][1] = &IntersectingSegment.GetSecondNode();
+	Nodes[1][0] = &Segment.GetFirstNode();
+	Nodes[1][1] = &Segment.GetSecondNode();
+
+	FPoint2D IntersectingPoints[2];
+	IntersectingPoints[0] = IntersectingSegment.GetFirstNode().Get2DPoint(EGridSpace::UniformScaled, Grid);
+	IntersectingPoints[1] = IntersectingSegment.GetSecondNode().Get2DPoint(EGridSpace::UniformScaled, Grid);
+
+	FPoint2D SegmentPoints[2];
+	SegmentPoints[0] = Segment.GetFirstNode().Get2DPoint(EGridSpace::UniformScaled, Grid);
+	SegmentPoints[1] = Segment.GetSecondNode().Get2DPoint(EGridSpace::UniformScaled, Grid);
+
+	double SquareLenghtIntersectingSegment = IntersectingPoints[0].SquareDistance(IntersectingPoints[1]);
+	double SquareLenghtSegment = SegmentPoints[0].SquareDistance(SegmentPoints[1]);
+
+	if (SquareLenghtSegment > 10 * SquareLenghtIntersectingSegment)
+	{
+		return RemoveIntersectionByMovingOutsideSegmentNodeInside(Segment, const_cast<FIsoSegment&>(IntersectingSegment), bIsSameInnerLoop);
+	}
+
+	bool bFirstNodeIsOutside = false;
+	FPoint2D PointToMove = Segment.GetSecondNode().Get2DPoint(EGridSpace::UniformScaled, Grid);
+
+	// Is Second node, the outside node ?
+	{
+		double OrientedSlop = ComputeOrientedSlope(IntersectingPoints[0], IntersectingPoints[1], PointToMove);
+		if (bIsSameInnerLoop)
+		{
+			OrientedSlop *= -1;
+		}
+
+		if (OrientedSlop > 0)
+		{
+			PointToMove = Segment.GetFirstNode().Get2DPoint(EGridSpace::UniformScaled, Grid);
+			bFirstNodeIsOutside = true;
+		}
+	}
+
+#ifdef DEBUG_CLOSED_OUSIDE_POINT
+	if (bDisplay)
+	{
+		F3DDebugSession _(TEXT("Outside Point"));
+		DisplayPoint(PointToMove, bFirstNodeIsOutside ? EVisuProperty::GreenPoint : EVisuProperty::YellowPoint);
+	}
+#endif
+
+	double Coordinate;
+	FPoint2D ProjectedPoint = ProjectPointOnSegment(PointToMove, IntersectingPoints[0], IntersectingPoints[1], Coordinate);
+
+	FPoint2D MoveDirection = IntersectingPoints[1] - IntersectingPoints[0];
+	MoveDirection.Normalize();
+	MoveDirection = MoveDirection.GetPerpendicularVector();
+	MoveDirection *= bIsSameInnerLoop ? -GeometricTolerance : GeometricTolerance;
+
+	FPoint2D NewCoordinate = ProjectedPoint + MoveDirection;
+
+	FLoopNode& Node = (FLoopNode&) const_cast<FIsoNode&>(!bFirstNodeIsOutside ? Segment.GetSecondNode() : Segment.GetFirstNode());
+	FLoopNode& PreviousNode = Node.GetPreviousNode();
+	FLoopNode& NextNode = Node.GetNextNode();
+
+	if ((PreviousNode.Get2DPoint(EGridSpace::UniformScaled, Grid).SquareDistance(NewCoordinate) < SquareGeometricTolerance2) ||
+		(NextNode.Get2DPoint(EGridSpace::UniformScaled, Grid).SquareDistance(NewCoordinate) < SquareGeometricTolerance2))
+	{
+		RemoveNodeOfLoop(Node);
+		return;
+	}
+
+	Node.Set2DPoint(EGridSpace::UniformScaled, Grid, NewCoordinate);
+
+#ifdef DEBUG_CLOSED_OUSIDE_POINT
+	if (bDisplay)
+	{
+		F3DDebugSession _(TEXT("New Segs"));
+		DisplayPoint(NewCoordinate, EVisuProperty::BluePoint);
+		Display(EGridSpace::UniformScaled, Segment, 0, EVisuProperty::BlueCurve);
+		Display(EGridSpace::UniformScaled, IntersectingSegment, 0, EVisuProperty::RedCurve);
+	}
+#endif
+}
+
+
+#ifdef UNUSED
+bool FLoopCleaner::TryToRemoveSelfIntersectionByMovingTheClosedOusidePoint(const FIsoSegment& Segment0, const FIsoSegment& Segment1)
+{
+	const FIsoNode* Nodes[2][2] = { {nullptr, nullptr}, {nullptr, nullptr} };
+	Nodes[0][0] = &Segment0.GetFirstNode();
+	Nodes[0][1] = &Segment0.GetSecondNode();
+	Nodes[1][0] = &Segment1.GetFirstNode();
+	Nodes[1][1] = &Segment1.GetSecondNode();
+
+	FPoint2D Points[2][2];
+	Points[0][0] = Nodes[0][0]->Get2DPoint(EGridSpace::UniformScaled, Grid);
+	Points[0][1] = Nodes[0][1]->Get2DPoint(EGridSpace::UniformScaled, Grid);
+	Points[1][0] = Nodes[1][0]->Get2DPoint(EGridSpace::UniformScaled, Grid);
+	Points[1][1] = Nodes[1][1]->Get2DPoint(EGridSpace::UniformScaled, Grid);
+
+	int32 ProjectedNodeIndex[2][2] = { { 0, 0 }, { 0, 0 } };
+	FPoint2D ProjectedPoints[2][2];
+	double Distance[2][2] = { { HUGE_VALUE, HUGE_VALUE }, { HUGE_VALUE, HUGE_VALUE } };
+	FPoint2D ProjectedPointsOut[2][2];
+
+	//int32 OtherSegmentIndex = 0;
+	int32 SegmentIndex = 1;
+	int32 OtherSegmentIndex = 0;
+
+	TFunction<void(int32)> ProjectPoint = [&](int32 OtherSegmentNodeIndex)
+	{
+		double Coordinate;
+		ProjectedPoints[OtherSegmentIndex][OtherSegmentNodeIndex] = ProjectPointOnSegment(Points[OtherSegmentIndex][OtherSegmentNodeIndex], Points[SegmentIndex][0], Points[SegmentIndex][1], Coordinate, false);
+		bool bProjectedPoint1IsInside = (Coordinate >= -SMALL_NUMBER && Coordinate <= 1 + SMALL_NUMBER);
+		if (bProjectedPoint1IsInside)
+		{
+			Distance[OtherSegmentIndex][OtherSegmentNodeIndex] = ProjectedPoints[OtherSegmentIndex][OtherSegmentNodeIndex].SquareDistance(Points[OtherSegmentIndex][OtherSegmentNodeIndex]);
+		}
+
+#ifdef DEBUG_SELF_CLOSED_OUSIDE_POINT
+		if (bDisplay)
+		{
+			{
+				F3DDebugSession _(TEXT("Segs"));
+				DisplaySegment(Points[SegmentIndex][0], Points[SegmentIndex][1], 0, EVisuProperty::BlueCurve);
+				DisplayPoint(Points[OtherSegmentIndex][OtherSegmentNodeIndex], EVisuProperty::RedPoint);
+				DisplayPoint(ProjectedPoints[OtherSegmentIndex][OtherSegmentNodeIndex], EVisuProperty::RedPoint);
+				Wait(false);
+			}
+		}
+#endif
+	};
+
+	SegmentIndex = 0;
+	OtherSegmentIndex = 1;
+	ProjectPoint(0);
+	ProjectPoint(1);
+	SegmentIndex = 1;
+	OtherSegmentIndex = 0;
+	ProjectPoint(0);
+	ProjectPoint(1);
+
+	TFunction<void(int32, int32)> MoveNode = [&](int32 MoveSegmentIndex, int32 MovePointIndex)
+	{
+		FPoint2D MoveDirection = ProjectedPoints[MoveSegmentIndex][MovePointIndex] - Points[MoveSegmentIndex][MovePointIndex];
+		MoveDirection.Normalize();
+
+		if (MoveDirection.SquareLength() < 0.5)
+		{
+			MoveDirection = Points[MoveSegmentIndex][MovePointIndex == 0 ? 1 : 0] - ProjectedPoints[MoveSegmentIndex][MovePointIndex];
+			MoveDirection.Normalize();
+		}
+
+		MoveDirection *= GeometricTolerance;
+		FPoint2D NewCoordinate = ProjectedPoints[MoveSegmentIndex][MovePointIndex] + MoveDirection;
+
+		FLoopNode& Node = (FLoopNode&) const_cast<FIsoNode&>(*Nodes[MoveSegmentIndex][MovePointIndex]);
+		FLoopNode& PreviousNode = Node.GetPreviousNode();
+		FLoopNode& NextNode = Node.GetNextNode();
+
+		if ((PreviousNode.Get2DPoint(EGridSpace::UniformScaled, Grid).SquareDistance(NewCoordinate) < SquareGeometricTolerance2) ||
+			(NextNode.Get2DPoint(EGridSpace::UniformScaled, Grid).SquareDistance(NewCoordinate) < SquareGeometricTolerance2))
+		{
+			RemoveNodeOfLoop(Node);
+			return;
+		}
+
+		const_cast<FIsoNode&>((*Nodes[MoveSegmentIndex][MovePointIndex])).Set2DPoint(EGridSpace::UniformScaled, Grid, NewCoordinate);
+
+#ifdef DEBUG_SELF_CLOSED_OUSIDE_POINT
+		if (bDisplay)
+		{
+			F3DDebugSession _(TEXT("New Segs"));
+			DisplayPoint(NewCoordinate, EVisuProperty::RedPoint);
+			Display(EGridSpace::UniformScaled, Segment1, 0, EVisuProperty::RedCurve);
+			Display(EGridSpace::UniformScaled, Segment0, 0, EVisuProperty::BlueCurve);
+			Wait(false);
+		}
+#endif
+	};
+
+
+	int32 MoveSegment = 0;
+	int32 MovePoint = 0;
+	double MinDistance = Distance[0][0];
+
+	TFunction<void(int32, int32)> FindNodeToMove = [&](int32 MoveSegmentIndex, int32 MovePointIndex)
+	{
+		if (Distance[MoveSegmentIndex][MovePointIndex] < MinDistance)
+		{
+			MinDistance = Distance[MoveSegmentIndex][MovePointIndex];
+			MoveSegment = MoveSegmentIndex;
+			MovePoint = MovePointIndex;
+		}
+	};
+
+	FindNodeToMove(0, 1);
+	FindNodeToMove(1, 0);
+	FindNodeToMove(1, 1);
+
+	MoveNode(MoveSegment, MovePoint);
+
+	return true;
+}
+
+bool FLoopCleaner::TryToRemoveIntersectionByMovingTheClosedOusidePoint(const FIsoSegment& Segment0, const FIsoSegment& Segment1)
+{
+	// For Inner/outer segment intersecting outer loop (i.e. IntersectingSegment is Outer), the segment has to be moved inside IntersectingSegment
+	// For Inner segment intersecting its loop, the segment has to be moved inside IntersectingSegment but the orientation of the loop is counterclockwise => need to swap orientation
+	// For Inner segment intersecting another inner loop, the segment has to be moved outside the loop but the orientation of the loop is counterclockwise => don't need to swap orientation
+
+	// As the iteration of the segment start from outer loop's segments, Segment can not be from outer loop if IntersectingSegment is not from outer loop
+	int32 SegmentLoopIndex[2];
+	SegmentLoopIndex[0] = Segment0.GetFirstNode().IsALoopNode() ? ((FLoopNode&)Segment0.GetFirstNode()).GetLoopIndex() : 0;
+	SegmentLoopIndex[1] = Segment1.GetFirstNode().IsALoopNode() ? ((FLoopNode&)Segment1.GetFirstNode()).GetLoopIndex() : 0;
+	bool bIsOutLoopSegment = SegmentLoopIndex[0] == 0;
+	bool bIsSameLoop = SegmentLoopIndex[0] == SegmentLoopIndex[1];
+
+	bool bToInside = bIsOutLoopSegment || bIsSameLoop;
+
+	const FIsoNode* Nodes[2][2] = { {nullptr, nullptr}, {nullptr, nullptr} };
+	Nodes[0][0] = &Segment0.GetFirstNode();
+	Nodes[0][1] = &Segment0.GetSecondNode();
+	Nodes[1][0] = &Segment1.GetFirstNode();
+	Nodes[1][1] = &Segment1.GetSecondNode();
+
+	FPoint2D Points[2][2];
+	Points[0][0] = Nodes[0][0]->Get2DPoint(EGridSpace::UniformScaled, Grid);
+	Points[0][1] = Nodes[0][1]->Get2DPoint(EGridSpace::UniformScaled, Grid);
+	Points[1][0] = Nodes[1][0]->Get2DPoint(EGridSpace::UniformScaled, Grid);
+	Points[1][1] = Nodes[1][1]->Get2DPoint(EGridSpace::UniformScaled, Grid);
+
+	int32 ProjectedNodeIndex[2] = { 0, 0 };
+	FPoint2D ProjectedPoints[2];
+	double Distance[2] = { HUGE_VALUE, HUGE_VALUE };
+	int32 OtherSegmentIndex = 0;
+	int32 SegmentIndex = 1;
+
+	TFunction<void(int32)> ProjectInnerPoint = [&](int32 OtherSegmentNodeIndex)
+	{
+		double Coordinate;
+		ProjectedPoints[OtherSegmentIndex] = ProjectPointOnSegment(Points[OtherSegmentIndex][OtherSegmentNodeIndex], Points[SegmentIndex][0], Points[SegmentIndex][1], Coordinate, false);
+		bool bProjectedPoint1IsInside = (Coordinate >= -SMALL_NUMBER && Coordinate <= 1 + SMALL_NUMBER);
+		if (bProjectedPoint1IsInside)
+		{
+			Distance[OtherSegmentIndex] = ProjectedPoints[OtherSegmentIndex].Distance(Points[OtherSegmentIndex][OtherSegmentNodeIndex]);
+		}
+
+#ifdef DEBUG_CLOSED_OUSIDE_POINT
+		if (bDisplay)
+		{
+			{
+				F3DDebugSession _(TEXT("Segs"));
+				DisplaySegment(Points[SegmentIndex][0], Points[SegmentIndex][1], 0, EVisuProperty::BlueCurve);
+				DisplayPoint(Points[OtherSegmentIndex][OtherSegmentNodeIndex], EVisuProperty::RedPoint);
+				DisplayPoint(ProjectedPoints[OtherSegmentIndex], EVisuProperty::RedPoint);
+				Wait(false);
+			}
+		}
+#endif
+	};
+
+	TFunction<void()> FindInnerNodeAndProjectIt = [&]()
+	{
+		OtherSegmentIndex = SegmentIndex ? 0 : 1;
+
+		double OrientedSlop = ComputeOrientedSlope(Points[SegmentIndex][0], Points[SegmentIndex][1], Points[OtherSegmentIndex][0]);
+		ProjectedNodeIndex[OtherSegmentIndex] = ((SegmentLoopIndex[0] == 0) == (OrientedSlop < 0)) ? 0 : 1;
+		ProjectInnerPoint(ProjectedNodeIndex[OtherSegmentIndex]);
+	};
+
+	SegmentIndex = 0;
+	FindInnerNodeAndProjectIt();
+	SegmentIndex = 1;
+	FindInnerNodeAndProjectIt();
+
+	if (Distance[0] == HUGE_VALUE && Distance[1] == HUGE_VALUE)
+	{
+		return false;
+	}
+
+	TFunction<void(int32)> MoveNode = [&](int32 MovePointIndex)
+	{
+		FPoint2D MoveDirection = ProjectedPoints[MovePointIndex] - Points[MovePointIndex][ProjectedNodeIndex[MovePointIndex]];
+		MoveDirection.Normalize();
+
+		if (MoveDirection.SquareLength() < 0.5)
+		{
+			MoveDirection = Points[MovePointIndex][ProjectedNodeIndex[MovePointIndex] == 0 ? 1 : 0] - ProjectedPoints[MovePointIndex];
+			MoveDirection.Normalize();
+		}
+
+		MoveDirection *= GeometricTolerance;
+		FPoint2D NewCoordinate = ProjectedPoints[MovePointIndex] + MoveDirection;
+
+		FLoopNode& Node = (FLoopNode&) const_cast<FIsoNode&>(*Nodes[MovePointIndex][ProjectedNodeIndex[MovePointIndex]]);
+		FLoopNode& PreviousNode = Node.GetPreviousNode();
+		FLoopNode& NextNode = Node.GetNextNode();
+
+		if ((PreviousNode.Get2DPoint(EGridSpace::UniformScaled, Grid).SquareDistance(NewCoordinate) < SquareGeometricTolerance2) ||
+			(NextNode.Get2DPoint(EGridSpace::UniformScaled, Grid).SquareDistance(NewCoordinate) < SquareGeometricTolerance2))
+		{
+			RemoveNodeOfLoop(Node);
+			return;
+		}
+
+		const_cast<FIsoNode&>((*Nodes[MovePointIndex][ProjectedNodeIndex[MovePointIndex]])).Set2DPoint(EGridSpace::UniformScaled, Grid, NewCoordinate);
+
+#ifdef DEBUG_CLOSED_OUSIDE_POINT
+		if (bDisplay)
+		{
+			F3DDebugSession _(TEXT("New Segs"));
+			DisplayPoint(NewCoordinate, EVisuProperty::RedPoint);
+			Display(EGridSpace::UniformScaled, Segment1, 0, EVisuProperty::RedCurve);
+			Display(EGridSpace::UniformScaled, Segment0, 0, EVisuProperty::BlueCurve);
+			Wait(false);
+		}
+#endif
+	};
+
+	int32 MovePoint = (Distance[0] < Distance[1]) ? 0 : 1;
+	MoveNode(MovePoint);
+	return true;
+}
+#endif
+
+void FLoopCleaner::FixLoopOrientation()
+{
+	FLoopNode const* const* StartNode = NodesOfLoop.FindByPredicate([](const FLoopNode* Node) { return !Node->IsDelete(); });
+	if (StartNode == nullptr)
+	{
+		return;
+	}
+
+	EOrientation Orientation = GetLoopOrientation(*StartNode);
+	if (Orientation == EOrientation::Back)
+	{
+#ifdef DEBUG_LOOP_ORIENTATION
+		Grid.DisplayIsoSegments(TEXT("Before orientation"), EGridSpace::UniformScaled, LoopSegments, true, true);
+#endif
+
+		SwapSubLoopOrientation(StartSegmentIndex, NextLoopFirstSegmentIndex);
+
+#ifdef DEBUG_LOOP_ORIENTATION
+		Grid.DisplayIsoSegments(TEXT("After orientation"), EGridSpace::UniformScaled, LoopSegments, true, true);
+		Wait(bDisplay);
+#endif
+	}
+}
+
+
+
+bool FLoopCleaner::RemoveSelfIntersectionsOfLoop()
+{
+	using namespace LoopCleanerImpl;
+
+	// All pinches are removed
+	for (int32 IntersectionIndex = Intersections.Num() - 1; IntersectionIndex >= 0; --IntersectionIndex)
+	{
+		const TPair<double, double>& Intersection = Intersections[IntersectionIndex];
+
+		bool bIntersectionKeyIsExtremity = FMath::IsNearlyEqual(Intersection.Key, (int32)(Intersection.Key + 0.5));
+		bool bIntersectionValueIsExtremity = FMath::IsNearlyEqual(Intersection.Value, (int32)(Intersection.Value + 0.5));
+
+		if (bIntersectionKeyIsExtremity && bIntersectionValueIsExtremity)
+		{
+			FPinchIntersectionContext Context(Intersection);
+			if(!Fill(Context))
+			{
+				return false;
+			}
+
+			if (IsAPinch(Context))
+			{
+				DisconnectCoincidentNodes(Context);
+				Intersections.RemoveAt(IntersectionIndex);
+			}
+		}
+	}
+
+	Algo::Sort(Intersections, [&](const TPair<double, double>& Intersection1, const TPair<double, double>& Intersection2)
+		{
+			return (Intersection1.Key < Intersection2.Key);
+		});
+
+#ifdef DEBUG_REMOVE_LOOP_INTERSECTIONS		
+	Grid.DisplayGridPolyline(TEXT("RemoveSelfIntersectionsOfLoop start"), EGridSpace::UniformScaled, NodesOfLoop, true, EVisuProperty::BlueCurve);
+	Wait(bDisplay);
+#endif
+
+	for (int32 IntersectionIndex = 0; IntersectionIndex < Intersections.Num(); )
+	{
+		const TPair<double, double>& Intersection = Intersections[IntersectionIndex];
+
+#ifdef DEBUG_REMOVE_LOOP_INTERSECTIONS		
+		F3DDebugSession _(bDisplay, TEXT("Intersected Segments"));
+		DisplayIntersection(Intersection);
+#endif
+
+		bool bIntersectionForward = true;
+
+		int32 NextIntersectionIndex = IntersectionIndex + 1;
+		int32 IntersectionCount = 1;
+		for (; NextIntersectionIndex < Intersections.Num(); ++NextIntersectionIndex)
+		{
+			if (Intersections[NextIntersectionIndex].Value > Intersection.Value)
+			{
+				break;
+			}
+			++IntersectionCount;
+
+#ifdef DEBUG_REMOVE_LOOP_INTERSECTIONS		
+			DisplayIntersection(Intersections[NextIntersectionIndex]);
+#endif
+		}
+
+		if (IntersectionCount == 1)
+		{
+			NextIntersectionIndex = IntersectionIndex + 1;
+			for (; NextIntersectionIndex < Intersections.Num(); ++NextIntersectionIndex)
+			{
+				if (Intersections[NextIntersectionIndex].Key > Intersection.Value)
+				{
+					break;
+				}
+
+				bIntersectionForward = false;
+				++IntersectionCount;
+#ifdef DEBUG_REMOVE_LOOP_INTERSECTIONS		
+				DisplayIntersection(Intersections[NextIntersectionIndex]);
+#endif
+				break;
+			}
+		}
+
+		// Stating from this point 
+		// the process must not delete node after the first intersection
+		if (IntersectionCount == 1)
+		{
+			if (!RemoveIntersection(Intersections[IntersectionIndex]))
+			{
+				return false;
+			}
+			IntersectionIndex++;
+		}
+		else if (!bIntersectionForward)
+		{
+			if (!RemoveOutgoingLoop(Intersections[IntersectionIndex], Intersections[IntersectionIndex + 1]))
+			{
+				return false;
+			}
+			IntersectionIndex += 2;
+		}
+		else
+		{
+			if (!RemoveIntersectionsOfSubLoop(IntersectionIndex, IntersectionCount))
+			{
+				return false;
+			}
+			IntersectionIndex += IntersectionCount;
+		}
+
+#ifdef DEBUG_REMOVE_LOOP_INTERSECTIONS		
+		Grid.DisplayGridPolyline(TEXT("RemoveSelfIntersectionsOfLoop end"), EGridSpace::UniformScaled, NodesOfLoop, true, EVisuProperty::BlueCurve);
+#endif
+	}
+
+	return true;
+}
+
+bool FLoopCleaner::RemoveOutgoingLoop(const TPair<double, double>& Intersection, const TPair<double, double>& NextIntersection)
+{
+	const TPair<double, double> OutSideLoop(Intersection.Value, NextIntersection.Value);
+	//if (IsSubLoopBiggerThanMainLoop(NodesOfLoop, OutSideLoop, bForward))
+	//{
+	//	return false;
+	//}
+
+	TFunction<FPoint2D(double, FLoopNode*, FLoopNode*)> IntersectingPoint = [&](double Coordinate, FLoopNode* Start, FLoopNode* End) -> FPoint2D
+	{
+		const FPoint2D& StartPoint = Start->Get2DPoint(EGridSpace::UniformScaled, Grid);
+		const FPoint2D& EndPoint = End->Get2DPoint(EGridSpace::UniformScaled, Grid);
+
+		int32 StartIndex = (int32)Coordinate;
+		Coordinate -= StartIndex;
+		FPoint2D Intersection = PointOnSegment(StartPoint, EndPoint, Coordinate);
+
+		//{
+		//	F3DDebugSession _(TEXT("Intersected Point"));
+		//	DisplayPoint(StartPoint, EVisuProperty::RedPoint);
+		//	DisplayPoint(EndPoint, EVisuProperty::RedPoint);
+		//	DisplayPoint(Intersection, EVisuProperty::BluePoint);
+		//}
+
+		return Intersection;
+	};
+
+	FLoopNode* TmpNode = GetNodeAt(NextIndex((int32)Intersection.Value));
+	FLoopNode* EndNode = GetNodeAt((int32)NextIntersection.Value);
+	if (TmpNode == nullptr || TmpNode->IsDelete() || EndNode == nullptr || EndNode->IsDelete())
+	{
+		return false;
+	}
+
+	FLoopNode* StartNode = GetPrevious(TmpNode);
+	FPoint2D FirstIntersection = IntersectingPoint(Intersection.Value, StartNode, TmpNode);
+
+	TmpNode = GetNext(EndNode);
+	FPoint2D SecondIntersection = IntersectingPoint(NextIntersection.Value, EndNode, TmpNode);
+
+	FPoint2D MiddlePoint = FirstIntersection.Middle(SecondIntersection);
+
+	TmpNode = GetNext(StartNode);
+	while (TmpNode && (TmpNode != EndNode) && !TmpNode->IsDelete())
+	{
+		if (!RemoveNodeOfLoop(*TmpNode))
+		{
+			return false;
+		}
+		TmpNode = GetNext(StartNode);
+	}
+
+	FPoint2D MoveDirection = SecondIntersection - FirstIntersection;
+	double Length = MoveDirection.Length();
+	if (FMath::IsNearlyZero(Length))
+	{
+		FLoopNode* StartSegment = GetNodeAt((int32)Intersection.Key);
+		if (StartSegment == nullptr)
+		{
+			return false;
+		}
+		MoveDirection = GetNext(StartSegment)->Get2DPoint(EGridSpace::UniformScaled, Grid) - StartSegment->Get2DPoint(EGridSpace::UniformScaled, Grid);
+		Length = MoveDirection.Length();
+	}
+
+	MoveDirection /= Length;
+	MoveDirection = MoveDirection.GetPerpendicularVector();
+	MoveDirection *= GeometricTolerance;
+	MiddlePoint += MoveDirection;
+
+	if (!EndNode->IsDelete())
+	{
+		EndNode->Set2DPoint(EGridSpace::UniformScaled, Grid, MiddlePoint);
+	}
+
+	return true;
+}
+
+bool FLoopCleaner::RemoveIntersectionsOfSubLoop(int32 IntersectionIndex, int32 IntersectionCount)
+{
+	using namespace LoopCleanerImpl;
+
+	TFunction<void(FLoopNode*, const FPoint2D&, FPoint2D&)> MoveNodeToProjection = [&](FLoopNode* NodeToProject, const FPoint2D& PointToProject, FPoint2D& ProjectedPoint)
+	{
+		FPoint2D MoveDirection = ProjectedPoint - PointToProject;
+		MoveDirection.Normalize();
+		MoveDirection *= GeometricTolerance;
+		ProjectedPoint += MoveDirection;
+		NodeToProject->Set2DPoint(EGridSpace::UniformScaled, Grid, ProjectedPoint);
+	};
+
+	TFunction<void(FLoopNode*, const FPoint2D&, const FPoint2D&)> ProjectNodeOnSegment = [&](FLoopNode* NodeToProject, const FPoint2D& Point0, const FPoint2D& Point1)
+	{
+		const FPoint2D& PointToProject = NodeToProject->Get2DPoint(EGridSpace::UniformScaled, Grid);
+
+		FPoint2D ProjectedPoint;
+		double Coordinate;
+		ProjectedPoint = ProjectPointOnSegment(PointToProject, Point0, Point1, Coordinate);
+
+		MoveNodeToProjection(NodeToProject, PointToProject, ProjectedPoint);
+	};
+
+	TFunction<bool(const int32, const int32, const int32)> ProjectNodesOnSegment = [&](const int32 StartIndex, const int32 EndIndex, const int32 SegmnentEndIndex) -> bool
+	{
+		FLoopNode* Node = GetNodeAt(StartIndex);
+		FLoopNode* StopNode = GetNodeAt(EndIndex);
+		FLoopNode* EndSegment = GetNodeAt(SegmnentEndIndex);
+		if (Node == nullptr || Node->IsDelete() || StopNode == nullptr || StopNode->IsDelete() || EndSegment == nullptr)
+		{
+			return false;
+		}
+
+		StopNode = GetNext(StopNode);
+		FLoopNode* StartSegment = GetPrevious(EndSegment);
+
+		const FPoint2D& EndPoint = EndSegment->Get2DPoint(EGridSpace::UniformScaled, Grid);
+		const FPoint2D& StartPoint = StartSegment->Get2DPoint(EGridSpace::UniformScaled, Grid);
+
+		for (; Node != StopNode; Node = GetNext(Node))
+		{
+			ProjectNodeOnSegment(Node, StartPoint, EndPoint);
+		}
+
+		for (Node = GetNodeAt(StartIndex); Node != StopNode; )
+		{
+			if (Node->IsDelete())
+			{
+				return false;
+			}
+
+			FLoopNode* NodeToProcess = Node;
+			Node = GetNext(Node);
+			if (Node == nullptr || Node->IsDelete())
+			{
+				return false;
+			}
+			CheckAndRemovePick(NodeToProcess->GetPreviousNode().Get2DPoint(EGridSpace::UniformScaled, Grid), NodeToProcess->Get2DPoint(EGridSpace::UniformScaled, Grid), NodeToProcess->GetNextNode().Get2DPoint(EGridSpace::UniformScaled, Grid), *NodeToProcess);
+		}
+		return true;
+	};
+
+	TFunction<void(FLoopNode*, FLoopNode*)> MoveNodeBehindOther = [&](FLoopNode* NodeToMove, FLoopNode* Node1Side1)
+	{
+		FLoopNode* Node0Side1 = GetPrevious(Node1Side1);
+		const FPoint2D& Point0 = Node0Side1->Get2DPoint(EGridSpace::UniformScaled, Grid);
+		const FPoint2D& Point1 = Node1Side1->Get2DPoint(EGridSpace::UniformScaled, Grid);
+
+		const FPoint2D& PointToMove = NodeToMove->Get2DPoint(EGridSpace::UniformScaled, Grid);
+
+		FPoint2D MoveDirection = Point1 - Point0;
+		MoveDirection.Normalize();
+		MoveDirection = MoveDirection.GetPerpendicularVector();
+		MoveDirection *= GeometricTolerance;
+
+		FPoint2D NewCoordinate = Point1 + MoveDirection;
+		NodeToMove->Set2DPoint(EGridSpace::UniformScaled, Grid, NewCoordinate);
+	};
+
+	for (int32 Index = IntersectionCount - 1; Index >= 0; --Index)
+	{
+		int32 SecondIntersectionIndex = IntersectionIndex + Index;
+
+#ifdef DEBUG_REMOVE_INTERSECTIONS		
+		Grid.DisplayGridPolyline(FString::Printf(TEXT("RemoveIntersectionsOfSubLoop Step %d"), Index), EGridSpace::UniformScaled, NodesOfLoop, true);
+		Grid.DisplayNodes(TEXT("RemoveIntersectionsOfSubLoop Nodes"), EGridSpace::UniformScaled, (TArray<const FIsoNode*>) NodesOfLoop, EVisuProperty::BluePoint);
+#endif
+
+
+		if (Index > 0)
+		{
+#ifdef DEBUG_REMOVE_INTERSECTIONS		
+			DisplayIntersection(Intersections[SecondIntersectionIndex]);
+			DisplayIntersection(Intersections[SecondIntersectionIndex - 1]);
+			//Wait();
+#endif
+			const TPair<double, double>& SecondIntersection = Intersections[SecondIntersectionIndex];
+			const TPair<double, double>& FirstIntersection = Intersections[SecondIntersectionIndex - 1];
+
+			int32 Side0NodeCount = (int32)SecondIntersection.Key - (int32)FirstIntersection.Key;
+			int32 Side1NodeCount = (int32)FirstIntersection.Value - (int32)SecondIntersection.Value;
+
+			int32 IndexSide0 = NextIndex((int32)FirstIntersection.Key);
+			int32 IndexSide1 = NextIndex((int32)SecondIntersection.Value);
+
+			if (Side0NodeCount == 0)
+			{
+				if (!ProjectNodesOnSegment(IndexSide1, FitNodeIndex((int32)FirstIntersection.Value), IndexSide0))
+				{
+					return false;
+				}
+			}
+			else if (Side1NodeCount == 0)
+			{
+				if (!ProjectNodesOnSegment(IndexSide0, FitNodeIndex((int32)SecondIntersection.Key), IndexSide1))
+				{
+					return false;
+				}
+			}
+			else if (Side0NodeCount == 1 && Side1NodeCount == 1)
+			{
+				FLoopNode* NodeSide0 = GetNodeAt(IndexSide0);
+				FLoopNode* NodeSide1 = GetNodeAt(IndexSide1);
+				if (NodeSide0->IsDelete() || NodeSide1->IsDelete())
+				{
+					return false;
+				}
+
+
+				double SlopSide0 = ComputeUnorientedSlope(GetPrevious(NodeSide0)->Get2DPoint(EGridSpace::UniformScaled, Grid), NodeSide0->Get2DPoint(EGridSpace::UniformScaled, Grid), GetNext(NodeSide0)->Get2DPoint(EGridSpace::UniformScaled, Grid));
+				double SlopSide1 = ComputeUnorientedSlope(GetPrevious(NodeSide1)->Get2DPoint(EGridSpace::UniformScaled, Grid), NodeSide1->Get2DPoint(EGridSpace::UniformScaled, Grid), GetNext(NodeSide1)->Get2DPoint(EGridSpace::UniformScaled, Grid));
+
+				if (SlopSide0 < SlopSide1)
+				{
+					MoveNodeBehindOther(NodeSide1, NodeSide0);
+				}
+				else
+				{
+					MoveNodeBehindOther(NodeSide0, NodeSide1);
+				}
+			}
+			else
+			{
+				FLoopSection IntersectingSection;
+				IntersectingSection.Key = GetPrevious(GetNodeAt(IndexSide0));
+				IntersectingSection.Value = GetNext(GetNodeAt((int32)SecondIntersection.Key));
+				if (IntersectingSection.Key->IsDelete() || IntersectingSection.Value->IsDelete())
+				{
+					return false;
+				}
+
+				FLoopSection OppositeSection;
+				OppositeSection.Key = GetPrevious(GetNodeAt(IndexSide1));
+				OppositeSection.Value = GetNext(GetNodeAt((int32)FirstIntersection.Value));
+				if (OppositeSection.Key->IsDelete() || OppositeSection.Value->IsDelete())
+				{
+					return false;
+				}
+
+				int32 IntersectingSectionCount = (int32)FirstIntersection.Value - IndexSide1;
+				int32 OppositeSectionCount = (int32)SecondIntersection.Key - IndexSide0;
+				if (OppositeSectionCount < IntersectingSectionCount)
+				{
+					Swap(IntersectingSection, OppositeSection);
+				}
+
+				MoveIntersectingSectionBehindOppositeSection(IntersectingSection, OppositeSection);
+
+#ifdef DEBUG_REMOVE_INTERSECTIONS		
+				Grid.DisplayGridPolyline(TEXT("RemoveIntersectionsOfSubLoop start"), EGridSpace::UniformScaled, NodesOfLoop, true);
+#endif
+
+				// Check if there is no more intersection, otherwise fix last intersections by moving node of OppositeSection
+				// As we don't know the existence of intersection, the process is chesk and fix
+				{
+					LoopSegmentsIntersectionTool.Empty(NodesOfLoopCount);
+					for (FLoopNode* Node = IntersectingSection.Key; Node != IntersectingSection.Value; Node = GetNext(Node))
+					{
+						LoopSegmentsIntersectionTool.AddSegment(*Node->GetSegmentConnectedTo(GetNext(Node)));
+					}
+				}
+
+				for (FLoopNode* Node = OppositeSection.Key; Node != OppositeSection.Value; Node = GetNext(Node))
+				{
+					if (const FIsoSegment* Intersection = LoopSegmentsIntersectionTool.DoesIntersect(*Node, *GetNext(Node)))
+					{
+						MoveNodeBehindSegment(*Intersection, *GetNext(Node));
+					}
+				}
+
+#ifdef DEBUG_REMOVE_INTERSECTIONS		
+				Grid.DisplayGridPolyline(TEXT("RemoveIntersectionsOfSubLoop second step"), EGridSpace::UniformScaled, NodesOfLoop, true);
+#endif
+
+				// Remove the pick if needed
+				for (FLoopNode* Node = OppositeSection.Key; Node != OppositeSection.Value;)
+				{
+					FLoopNode* NodeToProceed = Node;
+					Node = GetNext(Node);
+					if (CheckAndRemovePick(NodeToProceed->GetPreviousNode().Get2DPoint(EGridSpace::UniformScaled, Grid), NodeToProceed->Get2DPoint(EGridSpace::UniformScaled, Grid), NodeToProceed->GetNextNode().Get2DPoint(EGridSpace::UniformScaled, Grid), *NodeToProceed))
+					{
+						Node = GetPrevious(Node);
+					}
+				}
+
+#ifdef DEBUG_REMOVE_INTERSECTIONS		
+				Grid.DisplayGridPolyline(TEXT("RemoveIntersectionsOfSubLoop after remove pick"), EGridSpace::UniformScaled, NodesOfLoop, true);
+				Wait(bDisplay);
+#endif
+			}
+			--Index;
+		}
+		else
+		{
+			if (!RemoveIntersection(Intersections[IntersectionIndex]))
+			{
+				return false;
+			}
+		}
+	}
+
+	return true;
+}
+
+bool FLoopCleaner::RemoveIntersection(TPair<double, double>& Intersection)
+{
+	int32 SubLoopSegmentCount = (int32)Intersection.Value - (int32)Intersection.Key;
+	const int32 OtherSubLoopSegmentCount = NodesOfLoopCount - (int32)Intersection.Value + (int32)Intersection.Key;
+
+	bool bIntersectionKeyIsExtremity = FMath::IsNearlyEqual(Intersection.Key, (int32)(Intersection.Key + 0.5));
+	bool bIntersectionValueIsExtremity = FMath::IsNearlyEqual(Intersection.Value, (int32)(Intersection.Value + 0.5));
+
+	if(bIntersectionKeyIsExtremity && bIntersectionValueIsExtremity)
+	{   
+		LoopCleanerImpl::FPinchIntersectionContext Context(Intersection);
+		if (!Fill(Context))
+		{
+			return false;
+		}
+
+		if (IsAPinch(Context))
+		{
+			return DisconnectCoincidentNodes(Context);
+		}
+		else
+		{
+			return DisconnectCrossingSegments(Context);
+		}
+	}
+
+	/*
+	else if(bIntersectionKeyIsExtremity || bIntersectionValueIsExtremity)
+	{
+		// Case:  ________
+		//          _/\_
+		return MovePickBehind(NodesOfLoop, Intersection, bIntersectionKeyIsExtremity);
+	}
+	*/
+
+	if (OtherSubLoopSegmentCount < SubLoopSegmentCount)
+	{
+		Swap(Intersection.Value, Intersection.Key);
+		SubLoopSegmentCount = OtherSubLoopSegmentCount;
+	}
+
+#ifdef DEBUG_REMOVE_UNIQUE_INTERSECTION
+	Grid.DisplayGridPolyline(TEXT("RemoveIntersection start"), EGridSpace::UniformScaled, NodesOfLoop, true);
+	Wait(bDisplay);
+#endif
+
+	switch (SubLoopSegmentCount)
+	{
+	case 0:
+		ensureCADKernel(false);
+		break;
+
+	case 1:
+	case 2:
+		// At least 2 nodes are behind the intersection, remove the node making the intersection
+		if (!RemoveOuterNode(Intersection))
+		{
+			return false;
+		}
+		break;
+
+	case 3:
+		// 3 nodes are behind the intersection, remove the node making the intersection
+		//    ______a  c       ______a__b
+		//           \/ \o	            \o
+		//    ______d/\b/	  ______d__c/
+		//	    			  
+		if (!SwapNodes(Intersection))
+		{
+			return false;
+		}
+		break;
+
+	default:
+		if (!TryToSwapSegmentsOrRemoveLoop(Intersection))
+		{
+			return false;
+		}
+		break;
+	}
+
+#ifdef DEBUG_REMOVE_UNIQUE_INTERSECTION
+	Grid.DisplayGridPolyline(TEXT("RemoveIntersection end"), EGridSpace::UniformScaled, NodesOfLoop, true);
+	Wait(bDisplay);
+#endif
+	return true;
+};
+
+void FLoopCleaner::MoveIntersectingSectionBehindOppositeSection(LoopCleanerImpl::FLoopSection IntersectingSection, LoopCleanerImpl::FLoopSection OppositeSection)
+{
+	using namespace LoopCleanerImpl;
+
+	FLoopNode* FirstNodeIntersectingSection = IntersectingSection.Key;
+	FLoopNode* LastNodeIntersectingSection = IntersectingSection.Value;
+
+	FLoopNode* FirstNodeOppositeSection = OppositeSection.Key;
+	FLoopNode* LastNodeOppositeSection = OppositeSection.Value;
+
+	int32 OppositeSectionCount = 1;
+	for (FLoopNode* SegmentNode = FirstNodeOppositeSection; SegmentNode != LastNodeOppositeSection; SegmentNode = GetNext(SegmentNode), ++OppositeSectionCount);
+
+	TArray<const FPoint2D*> OppositeSectionPoint;
+	OppositeSectionPoint.Reserve(OppositeSectionCount);
+
+	for (FLoopNode* SegmentNode = FirstNodeOppositeSection; SegmentNode != LastNodeOppositeSection; SegmentNode = GetNext(SegmentNode))
+	{
+		OppositeSectionPoint.Add(&SegmentNode->Get2DPoint(EGridSpace::UniformScaled, Grid));
+	}
+	OppositeSectionPoint.Add(&LastNodeOppositeSection->Get2DPoint(EGridSpace::UniformScaled, Grid));
+
+#ifdef DEBUG_MOVE_INTERSECTING_SECTION_BEHIND_OPPOSITE_SECTION
+	{
+		F3DDebugSession _(TEXT("IntersectingSection"));
+		for (FLoopNode* Node = FirstNodeIntersectingSection; Node != GetNext(LastNodeIntersectingSection); Node = GetNext(Node))
+		{
+			DisplayPoint(Node->Get2DPoint(EGridSpace::UniformScaled, Grid), EVisuProperty::PurplePoint);
+		}
+	}
+	{
+		F3DDebugSession _(TEXT("OppositeSection"));
+		for (FLoopNode* Node = FirstNodeOppositeSection; Node != GetNext(LastNodeOppositeSection); Node = GetNext(Node))
+		{
+			DisplayPoint(Node->Get2DPoint(EGridSpace::UniformScaled, Grid), EVisuProperty::YellowPoint);
+		}
+		Wait();
+	}
+#endif
+
+	double Coordinate;
+	FLoopNode* NextNode = nullptr;
+	for (FLoopNode* Node = GetNext(FirstNodeIntersectingSection); Node != LastNodeIntersectingSection;)
+	{
+		FLoopNode* NodeToProceed = Node;
+		Node = GetNext(Node);
+
+		FPoint2D  CandiatePosition;
+		double MinSquareDistance = HUGE_VALUE;
+		FPoint2D PointToProject = NodeToProceed->Get2DPoint(EGridSpace::UniformScaled, Grid);
+		for (int32 Index = 1; Index < OppositeSectionCount; ++Index)
+		{
+			FPoint2D ProjectedPoint = ProjectPointOnSegment(PointToProject, *OppositeSectionPoint[Index - 1], *OppositeSectionPoint[Index], Coordinate);
+			double SquareDistance = PointToProject.SquareDistance(ProjectedPoint);
+			if (SquareDistance < MinSquareDistance)
+			{
+				MinSquareDistance = SquareDistance;
+				CandiatePosition = ProjectedPoint;
+			}
+		}
+		MoveNode(*NodeToProceed, CandiatePosition);
+	}
+
+	for (FLoopNode* Node = GetNext(FirstNodeIntersectingSection); Node != LastNodeIntersectingSection; )
+	{
+		FLoopNode* NodeToProceed = Node;
+		Node = GetNext(Node);
+		if (CheckAndRemovePick(NodeToProceed->GetPreviousNode().Get2DPoint(EGridSpace::UniformScaled, Grid), NodeToProceed->Get2DPoint(EGridSpace::UniformScaled, Grid), NodeToProceed->GetNextNode().Get2DPoint(EGridSpace::UniformScaled, Grid), *NodeToProceed))
+		{
+			Node = GetPrevious(Node);
+		}
+	}
+}
+
+void FLoopCleaner::MoveNodeBehindSegment(const FIsoSegment& IntersectingSegment, FLoopNode& NodeToMove)
+{
+	const FIsoNode* Nodes[2] = { nullptr, nullptr };
+	Nodes[0] = &IntersectingSegment.GetFirstNode();
+	Nodes[1] = &IntersectingSegment.GetSecondNode();
+
+	FPoint2D IntersectingPoints[2];
+	IntersectingPoints[0] = IntersectingSegment.GetFirstNode().Get2DPoint(EGridSpace::UniformScaled, Grid);
+	IntersectingPoints[1] = IntersectingSegment.GetSecondNode().Get2DPoint(EGridSpace::UniformScaled, Grid);
+
+	bool bEndNodeIsOutside = true;
+	FPoint2D PointToMove = NodeToMove.Get2DPoint(EGridSpace::UniformScaled, Grid);
+
+#ifdef DEBUG_CLOSED_OUSIDE_POINT
+	if (bDisplay)
+	{
+		F3DDebugSession _(TEXT("Outside Point"));
+		DisplayPoint(PointToMove, bEndNodeIsOutside ? EVisuProperty::GreenPoint : EVisuProperty::YellowPoint);
+	}
+#endif
+
+	double Coordinate;
+	FPoint2D ProjectedPoint = ProjectPointOnSegment(PointToMove, IntersectingPoints[0], IntersectingPoints[1], Coordinate);
+	MoveNode(NodeToMove, ProjectedPoint);
+}
+
+bool FLoopCleaner::Fill(LoopCleanerImpl::FPinchIntersectionContext& Context)
+{
+	TFunction<bool(int32, const int32)> FillData = [&](int32 Index, const int32 Side) -> bool
+	{
+		Context.Points[Side].SetNum(3);
+		Context.Nodes[Side][1] = GetNodeAt(Index);
+		if (Context.Nodes[Side][1] == nullptr || Context.Nodes[Side][1]->IsDelete())
+		{
+			return false;
+		}
+
+		Context.Nodes[Side][0] = &Context.Nodes[Side][1]->GetPreviousNode();
+		Context.Nodes[Side][2] = &Context.Nodes[Side][1]->GetNextNode();
+		if (Context.Nodes[Side][0] == nullptr || Context.Nodes[Side][0]->IsDelete() || Context.Nodes[Side][2] == nullptr || Context.Nodes[Side][2]->IsDelete())
+		{
+			return false;
+		}
+
+		Context.Points[Side][0] = &Context.Nodes[Side][0]->Get2DPoint(EGridSpace::UniformScaled, Grid);
+		Context.Points[Side][1] = &Context.Nodes[Side][1]->Get2DPoint(EGridSpace::UniformScaled, Grid);
+		Context.Points[Side][2] = &Context.Nodes[Side][2]->Get2DPoint(EGridSpace::UniformScaled, Grid);
+		return true;
+	};
+
+	if (!FillData((int32)(Context.Intersection.Key + 0.5), 0))
+	{
+		return false;
+	}
+
+	if (!FillData((int32)(Context.Intersection.Value + 0.5), 1))
+	{
+		return false;
+	}
+	return true;
+}
+
+bool FLoopCleaner::IsAPinch(const LoopCleanerImpl::FPinchIntersectionContext& Context) const
+{
+	return ArePointsInsideSectorABC(*Context.Points[0][2], *Context.Points[0][1], *Context.Points[0][0], Context.Points[1]);
+}
+
+bool FLoopCleaner::DisconnectCoincidentNodes(const LoopCleanerImpl::FPinchIntersectionContext& Context)
+{
+	double Slop00 = ComputeSlope(*Context.Points[0][1], *Context.Points[0][0]);
+	double Slop02 = ComputeSlope(*Context.Points[0][1], *Context.Points[0][2]);
+	double Slop10 = ComputeSlope(*Context.Points[1][1], *Context.Points[1][0]);
+	double Slop12 = ComputeSlope(*Context.Points[1][1], *Context.Points[1][2]);
+
+	double MediumSlop0 = TransformIntoPositiveSlope((Slop00 + Slop02) * 0.5);
+	double MediumSlop1 = TransformIntoPositiveSlope((Slop10 + Slop12) * 0.5);
+
+	FPoint2D P1 = SlopeToVector(MediumSlop1);
+	FPoint2D P0 = SlopeToVector(MediumSlop0);
+	P1.Normalize()*= GeometricTolerance*30;
+	P0.Normalize()*= GeometricTolerance*30;
+
+	P0 += *Context.Points[0][1];
+	P1 += *Context.Points[1][1];
+
+	Context.Nodes[0][1]->Set2DPoint(EGridSpace::UniformScaled, Grid, P0);
+	Context.Nodes[1][1]->Set2DPoint(EGridSpace::UniformScaled, Grid, P1);
+
+#ifdef DEBUG_DISCONNECT_COINCIDENT_NODES
+	F3DDebugSession _(TEXT("DisconnectCoincidentNodes"));
+	DisplaySegment(*Context.Points[0][1], P0, 0, EVisuProperty::BlueCurve);
+	DisplaySegment(*Context.Points[1][1], P1, 0, EVisuProperty::RedCurve);
+	Grid.DisplayIsoNode(EGridSpace::UniformScaled, *Context.Nodes[0][1], 0, EVisuProperty::BluePoint);
+	Grid.DisplayIsoNode(EGridSpace::UniformScaled, *Context.Nodes[1][1], 0, EVisuProperty::RedPoint);
+	Wait();
+#endif
+
+	return true;
+}
+
+
+/**
+ * CrossingCase:
+ *           --<-o    o->--       --<-o   o--<-
+ *           |    \  /    |       |    \ /    |
+ *           |     oo     |   =>  |     o     |
+ *           |    /  \    |       |     o     |
+ *           -->-o    o-<--       |    / \    |
+ *                                -->-o   o->--
+ */
+//#define DEBUG_DISCONNECT_CROSSING_SEGMENTS
+bool FLoopCleaner::DisconnectCrossingSegments(LoopCleanerImpl::FPinchIntersectionContext& Context)
+{
+
+#ifdef DEBUG_DISCONNECT_CROSSING_SEGMENTS
+	{
+		F3DDebugSession C(bDisplay, TEXT("DisconnectCrossingSegments"));
+		if (bDisplay)
+		{
+			Grid.DisplayIsoSegments(TEXT("Loops Orientation"), EGridSpace::UniformScaled, LoopSegments, false, true, EVisuProperty::BlueCurve);
+			Wait();
+		}
+	}
+#endif
+
+	FIsoSegment* Segment01_02 = Context.Nodes[0][1]->GetSegmentConnectedTo(Context.Nodes[0][2]);
+	FIsoSegment* Segment10_11 = Context.Nodes[1][0]->GetSegmentConnectedTo(Context.Nodes[1][1]);
+
+	Context.Nodes[0][2]->DisconnectSegment(*Segment01_02);
+	Context.Nodes[1][0]->DisconnectSegment(*Segment10_11);
+
+	Context.Nodes[0][2]->ConnectSegment(*Segment10_11);
+	Context.Nodes[1][0]->ConnectSegment(*Segment01_02);
+
+	Segment10_11->ReplaceNode(*Context.Nodes[1][0], *Context.Nodes[0][2]);
+	Segment01_02->ReplaceNode(*Context.Nodes[0][2], *Context.Nodes[1][0]);
+
+	Swap(Context.Nodes[0][2], Context.Nodes[1][0]);
+	Swap(Context.Points[0][2], Context.Points[1][0]);
+
+	DisconnectCoincidentNodes(Context);
+
+#ifdef DEBUG_DISCONNECT_CROSSING_SEGMENTS
+	F3DDebugSession _(bDisplay, TEXT("DisconnectCrossingSegments"));
+	if (bDisplay)
+	{
+		{
+			F3DDebugSession _(TEXT("New Segment from 00"));
+			Grid.DisplayIsoSegment(EGridSpace::UniformScaled, *Context.Nodes[0][0], Context.Nodes[0][0]->GetNextNode(), 0, EVisuProperty::YellowCurve);
+			Grid.DisplayIsoSegment(EGridSpace::UniformScaled, Context.Nodes[0][0]->GetNextNode(), Context.Nodes[0][0]->GetNextNode().GetNextNode(), 0, EVisuProperty::YellowCurve);
+		}
+		{
+			F3DDebugSession _(TEXT("New Segment fron 10"));
+			Grid.DisplayIsoSegment(EGridSpace::UniformScaled, *Context.Nodes[1][0], Context.Nodes[1][0]->GetNextNode(), 0, EVisuProperty::YellowCurve);
+			Grid.DisplayIsoSegment(EGridSpace::UniformScaled, Context.Nodes[1][0]->GetNextNode(), Context.Nodes[1][0]->GetNextNode().GetNextNode(), 0, EVisuProperty::YellowCurve);
+		}
+		Grid.DisplayIsoSegments(TEXT("Loops Orientation"), EGridSpace::UniformScaled, LoopSegments, false, true, EVisuProperty::BlueCurve);
+		Wait();
+	}
+#endif
+
+	int32 Segment01_02Index = LoopSegments.IndexOfByKey(Segment01_02);
+	int32 Segment10_11Index = LoopSegments.IndexOfByKey(Segment10_11);
+
+#ifdef DEBUG_REMOVE_UNIQUE_INTERSECTION
+	Grid.DisplayIsoSegments(TEXT("Loops Orientation"), EGridSpace::UniformScaled, LoopSegments, false, true, EVisuProperty::BlueCurve);
+	Grid.DisplayGridPolyline(TEXT("RemoveIntersection end"), EGridSpace::UniformScaled, NodesOfLoop, true);
+	Wait(bDisplay);
+#endif
+
+	SwapSubLoopOrientation(Segment01_02Index + 1, Segment10_11Index);
+
+#ifdef DEBUG_REMOVE_UNIQUE_INTERSECTION
+	Grid.DisplayIsoSegments(TEXT("Loops Orientation"), EGridSpace::UniformScaled, LoopSegments, false, true, EVisuProperty::BlueCurve);
+	Grid.DisplayGridPolyline(TEXT("RemoveIntersection end"), EGridSpace::UniformScaled, NodesOfLoop, true);
+	Wait(bDisplay);
+#endif
+
+	return true;
+}
+
+// Case  ______a  c       ______a__b
+//              \/ \o	            \o
+//       ______d/\b/	  ______d__c/
+//	    			  
+bool FLoopCleaner::SwapNodes(const TPair<double, double>& Intersection)
+{
+	FLoopNode* Node0 = GetNodeAt(NextIndex((int32)Intersection.Key));
+	if (Node0 == nullptr || Node0->IsDelete())
+	{
+		return false;
+	}
+
+	FLoopNode* Pick = GetNext(Node0);
+	FLoopNode* Node1 = GetNext(Pick);
+
+	const FPoint2D Point0Copy = Node0->Get2DPoint(EGridSpace::UniformScaled, Grid);
+	const FPoint2D& Point1 = Node1->Get2DPoint(EGridSpace::UniformScaled, Grid);
+
+	Node0->Set2DPoint(EGridSpace::UniformScaled, Grid, Point1);
+	Node1->Set2DPoint(EGridSpace::UniformScaled, Grid, Point0Copy);
+	return true;
+};
+
+bool FLoopCleaner::TryToSwapSegmentsOrRemoveLoop(const TPair<double, double>& Intersection)
+{
+	//using namespace IsoTriangulatorImpl;
+
+	const int32 Segment0StartIndex = NextIndex((int32)Intersection.Key);
+	const int32 Segment1EndIndex = (int32)Intersection.Value;
+
+	FLoopNode* Segment0_Node1 = GetNodeAt(Segment0StartIndex);
+	FLoopNode* Segment1_Node0 = GetNodeAt(Segment1EndIndex);
+	if (Segment0_Node1 == nullptr || Segment0_Node1->IsDelete() || Segment1_Node0 == nullptr || Segment1_Node0->IsDelete())
+	{
+		return false;
+	}
+
+	FLoopNode* Segment0_Node0 = GetPrevious(Segment0_Node1);
+	FLoopNode* Segment1_Node1 = GetNext(Segment1_Node0);
+	if (Segment0_Node0 == nullptr || Segment0_Node0->IsDelete() || Segment1_Node1 == nullptr || Segment1_Node1->IsDelete())
+	{
+		return false;
+	}
+
+	const FPoint2D& Segment0_Point0 = Segment0_Node0->Get2DPoint(EGridSpace::UniformScaled, Grid);
+	const FPoint2D& Segment0_Point1 = Segment0_Node1->Get2DPoint(EGridSpace::UniformScaled, Grid);
+	const FPoint2D& Segment1_Point0 = Segment1_Node0->Get2DPoint(EGridSpace::UniformScaled, Grid);
+	const FPoint2D& Segment1_Point1 = Segment1_Node1->Get2DPoint(EGridSpace::UniformScaled, Grid);
+
+#ifdef DEBUG_SWAP_SEGMENTS_OR_REMOVE
+	if (bDisplay)
+	{
+		F3DDebugSession _(TEXT("Intersected Segments"));
+		Grid.DisplayIsoSegment(EGridSpace::UniformScaled, *Segment0_Node0, *Segment0_Node1, 0, EVisuProperty::RedCurve);
+		Grid.DisplayIsoSegment(EGridSpace::UniformScaled, *Segment1_Node0, *Segment1_Node1, 0, EVisuProperty::RedCurve);
+		Grid.DisplayIsoNode(EGridSpace::UniformScaled, *Segment0_Node0, 0, EVisuProperty::RedPoint);
+		Grid.DisplayIsoNode(EGridSpace::UniformScaled, *Segment1_Node0, 0, EVisuProperty::RedPoint);
+		Wait(false);
+	}
+#endif
+
+	double Slop = ComputeSlope(Segment0_Point0, Segment0_Point1);
+	Slop = ComputeUnorientedSlope(Segment1_Point0, Segment1_Point1, Slop);
+
+	if (Slop < 2)
+	{
+		FIsoSegment* Segment0 = Segment0_Node0->GetSegmentConnectedTo(Segment0_Node1);
+		FIsoSegment* Segment1 = Segment1_Node0->GetSegmentConnectedTo(Segment1_Node1);
+		if (Segment0 == nullptr || Segment1 == nullptr)
+		{
+			return false;
+		}
+
+		if (!LoopSegmentsIntersectionTool.DoesIntersect(Segment0->GetFirstNode(), Segment1->GetFirstNode())
+		&& !LoopSegmentsIntersectionTool.DoesIntersect(Segment0->GetSecondNode(), Segment1->GetSecondNode()))
+		{
+			SwapSegments(*Segment0, *Segment1);
+
+#ifdef DEBUG_SWAP_SEGMENTS_OR_REMOVE
+			if (bDisplay)
+			{
+				F3DDebugSession _(TEXT("New Segments"));
+				Grid.DisplayIsoSegment(EGridSpace::UniformScaled, *Segment0_Node0, Segment0_Node0->GetNextNode(), 0, EVisuProperty::BlueCurve);
+				Grid.DisplayIsoSegment(EGridSpace::UniformScaled, *Segment1_Node1, Segment1_Node1->GetPreviousNode(), 0, EVisuProperty::BlueCurve);
+				Grid.DisplayIsoNode(EGridSpace::UniformScaled, *Segment0_Node0, 0, EVisuProperty::BluePoint);
+				Grid.DisplayIsoNode(EGridSpace::UniformScaled, *Segment1_Node0, 0, EVisuProperty::BluePoint);
+			}
+#endif
+			return true;
+		}
+	}
+
+	return RemoveSubLoop(Segment0_Node0, Segment1_Node1);
+}
+
+//#define DEBUG_BY_SWAPPING_SEGMENTS
+void FLoopCleaner::SwapSegments(FIsoSegment& IntersectingSegment, FIsoSegment& Segment)
+{
+	int32 IntersectingSegmentIndex = LoopSegments.IndexOfByKey(&IntersectingSegment);
+	int32 SegmentIndex = LoopSegments.IndexOfByKey(&Segment);
+
+	IntersectingSegment.GetSecondNode().DisconnectSegment(IntersectingSegment);
+	Segment.GetFirstNode().DisconnectSegment(Segment);
+
+	FIsoNode& Node = IntersectingSegment.GetSecondNode();
+	IntersectingSegment.SetSecondNode(Segment.GetFirstNode());
+	Segment.SetFirstNode(Node);
+
+	Segment.GetFirstNode().ConnectSegment(Segment);
+	IntersectingSegment.GetSecondNode().ConnectSegment(IntersectingSegment);
+
+#ifdef DEBUG_BY_SWAPPING_SEGMENTS
+	if (bDisplay)
+	{
+		F3DDebugSession _(TEXT("New Segments"));
+		Grid.DisplayIsoSegment(EGridSpace::UniformScaled, Segment, StartSegmentIndex, EVisuProperty::RedCurve/*, true*/);
+		Grid.DisplayIsoSegment(EGridSpace::UniformScaled, IntersectingSegment, EndSegmentIndex, EVisuProperty::RedCurve/*, true*/);
+		Wait();
+	}
+#endif
+
+	SwapSubLoopOrientation(IntersectingSegmentIndex + 1, SegmentIndex);
+
+#ifdef DEBUG_BY_SWAPPING_SEGMENTS
+	Grid.DisplayIsoSegments(TEXT("After orientation"), EGridSpace::UniformScaled, LoopSegments, false, true);
+	Wait();
+#endif
+
+	LoopSegmentsIntersectionTool.Update(&IntersectingSegment);
+	LoopSegmentsIntersectionTool.Update(&Segment);
+}
+
+void FLoopCleaner::SwapSubLoopOrientation(int32 FirstSegmentIndex, int32 LastSegmentIndex)
+{
+	TArray<FIsoSegment*> Segments;
+	Segments.Reserve(LastSegmentIndex - FirstSegmentIndex);
+
+	LastSegmentIndex = FitSegmentIndex(LastSegmentIndex);
+	int32 Index = FitSegmentIndex(FirstSegmentIndex);
+	do
+	{
+		LoopSegments[Index]->SwapOrientation();
+		Segments.Add(LoopSegments[Index]);
+		Index = NextSegmentIndex(Index);
+	}
+	while(Index != LastSegmentIndex);
+
+	int32 ReverseIndex = Segments.Num() - 1;
+	Index = FitSegmentIndex(FirstSegmentIndex);
+	do
+	{
+		LoopSegments[Index] = Segments[ReverseIndex];
+		Index = NextSegmentIndex(Index);
+		--ReverseIndex;
+	}
+	while(Index != LastSegmentIndex);
+}
+
+bool FLoopCleaner::RemoveSubLoop(FLoopNode* StartNode, FLoopNode* EndNode)
+{
+	FLoopNode* Node = GetNext(StartNode);
+	while (Node && (Node != EndNode) && !Node->IsDelete())
+	{
+		if (RemoveNodeOfLoop(*Node))
+		{
+			return false;
+		}
+
+		Node = GetNext(StartNode);
+	}
+
+	//Display(EGridSpace::UniformScaled, TEXT("Loop"), LoopSegments, true);
+	//Wait();
+	return true;
+}
+
+void FLoopCleaner::MoveNode(FLoopNode& NodeToMove, FPoint2D& NewPostion)
+{
+	const FPoint2D& PointToMove = NodeToMove.Get2DPoint(EGridSpace::UniformScaled, Grid);
+
+	FPoint2D ProjectedDirection = NewPostion - PointToMove;
+	ProjectedDirection.Normalize();
+	ProjectedDirection *= GeometricTolerance;
+
+	FPoint2D NewCoordinate = NewPostion + ProjectedDirection;
+
+	FLoopNode& PreviousNode = NodeToMove.GetPreviousNode();
+	FLoopNode& NextNode = NodeToMove.GetNextNode();
+
+	if ((PreviousNode.Get2DPoint(EGridSpace::UniformScaled, Grid).SquareDistance(NewCoordinate) < SquareGeometricTolerance2) ||
+		(NextNode.Get2DPoint(EGridSpace::UniformScaled, Grid).SquareDistance(NewCoordinate) < SquareGeometricTolerance2))
+	{
+		RemoveNodeOfLoop(NodeToMove);
+	}
+	else
+	{
+		NodeToMove.Set2DPoint(EGridSpace::UniformScaled, Grid, NewCoordinate);
+	}
+
+#ifdef DEBUG_MOVE_NODE
+	if (bDisplay)
+	{
+		{
+			F3DDebugSession _(TEXT("Point To Move"));
+			DisplayPoint(PointToMove, EVisuProperty::YellowPoint);
+		}
+		{
+			F3DDebugSession _(TEXT("Projected Point"));
+			DisplayPoint(NewPostion, EVisuProperty::GreenPoint);
+		}
+		{
+			F3DDebugSession _(TEXT("New Position"));
+			DisplayPoint(NewCoordinate, EVisuProperty::BluePoint);
+			Wait(false);
+		}
+	}
+#endif
+}
+
+//#define DEBUG_FIND_LOOP_INTERSECTIONS
+void FLoopCleaner::FindLoopIntersections()
+{
+	LoopSegmentsIntersectionTool.Empty(NodesOfLoopCount);
+
+	FLoopNode* const* const StartNodePtr = NodesOfLoop.FindByPredicate([](const FLoopNode* Node) { return !Node->IsDelete(); });
+	if (*StartNodePtr == nullptr || (*StartNodePtr)->IsDelete())
+	{
+		return;
+	}
+
+	FLoopNode* StartNode = *StartNodePtr;
+
+	TArray<const FIsoSegment*> IntersectedSegments;
+
+	SegmentCount = 1;
+	int32 SegmentIndex = 1;
+	FLoopNode* Node = nullptr;
+	FLoopNode* NextNode = nullptr;
+
+	TFunction<void()> FindSegmentIntersection = [&]()
+	{
+		FIsoSegment* Segment = Node->GetSegmentConnectedTo(NextNode);
+		if (Segment == nullptr)
+		{
+			return;
+		}
+
+		if (LoopSegmentsIntersectionTool.FindIntersections(*Node, *NextNode, IntersectedSegments))
+		{
+			for (const FIsoSegment* IntersectedSegment : IntersectedSegments)
+			{
+				const FLoopNode* IntersectedSegmentFirstNode = GetFirst(IntersectedSegment);
+				if (IntersectedSegmentFirstNode == nullptr)
+				{
+					continue;
+				}
+
+				int32 IntersectionIndex = 0;
+				FLoopNode* TmpNode = StartNode;
+				while (TmpNode != IntersectedSegmentFirstNode)
+				{
+					++IntersectionIndex;
+					TmpNode = GetNext(TmpNode);
+				}
+
+				TSegment<FPoint2D> SegmentPoints(Node->Get2DPoint(EGridSpace::UniformScaled, Grid), NextNode->Get2DPoint(EGridSpace::UniformScaled, Grid));
+				TSegment<FPoint2D> IntersectedSegmentPoints(GetFirst(IntersectedSegment)->Get2DPoint(EGridSpace::UniformScaled, Grid), GetSecond(IntersectedSegment)->Get2DPoint(EGridSpace::UniformScaled, Grid));
+
+				double SegmentIntersectionCoordinate;
+				FPoint2D IntersectionPoint = FindIntersectionOfSegments2D(SegmentPoints, IntersectedSegmentPoints, SegmentIntersectionCoordinate);
+
+				double SquareLength = IntersectedSegmentPoints.SquareLength();
+				double IntersectedSegmentIntersectionCoordinate = (SquareLength > SMALL_NUMBER_SQUARE) ? FMath::Sqrt(IntersectedSegmentPoints[0].SquareDistance(IntersectionPoint) / IntersectedSegmentPoints.SquareLength()) : 0;
+				IntersectedSegmentIntersectionCoordinate += IntersectionIndex;
+
+				SegmentIntersectionCoordinate += SegmentIndex;
+				if (FMath::IsNearlyEqual(SegmentIntersectionCoordinate, SegmentCount))
+				{
+					SegmentIntersectionCoordinate = IntersectedSegmentIntersectionCoordinate;
+					IntersectedSegmentIntersectionCoordinate = SegmentCount;
+				}
+
+				if (IntersectedSegmentIntersectionCoordinate > SegmentIntersectionCoordinate)
+				{
+					SegmentIntersectionCoordinate += NodesOfLoopCount;
+				}
+
+				if ((SegmentIntersectionCoordinate - IntersectedSegmentIntersectionCoordinate) > NodesOfLoopCount * 0.5)
+				{
+					Swap(SegmentIntersectionCoordinate, IntersectedSegmentIntersectionCoordinate);
+					SegmentIntersectionCoordinate += NodesOfLoopCount;
+				}
+
+				while (IntersectedSegmentIntersectionCoordinate >= NodesOfLoopCount)
+				{
+					IntersectedSegmentIntersectionCoordinate -= NodesOfLoopCount;
+					SegmentIntersectionCoordinate -= NodesOfLoopCount;
+				}
+
+ 				if (Intersections.Num() && FMath::IsNearlyEqual(Intersections.Last().Key, IntersectedSegmentIntersectionCoordinate) && FMath::IsNearlyEqual(Intersections.Last().Value, SegmentIntersectionCoordinate))
+				{
+					continue;
+				}
+
+				Intersections.Emplace(IntersectedSegmentIntersectionCoordinate, SegmentIntersectionCoordinate);
+
+#ifdef DEBUG_FIND_LOOP_INTERSECTIONS		
+				if (bDisplay)
+				{
+					F3DDebugSession _(*FString::Printf(TEXT("Intersection %f %f"), IntersectedSegmentIntersectionCoordinate, SegmentIntersectionCoordinate));
+					Grid.DisplayIsoSegment(EGridSpace::UniformScaled, *Node, *NextNode, 0, EVisuProperty::BlueCurve);
+					Grid.DisplayIsoSegment(EGridSpace::UniformScaled, *IntersectedSegment, 0, EVisuProperty::BlueCurve);
+					DisplayPoint(IntersectionPoint, EVisuProperty::RedPoint);
+					//Wait();
+				}
+#endif
+			}
+		}
+
+		if (Segment != nullptr)
+		{
+			LoopSegmentsIntersectionTool.AddSegment(*Segment);
+		}
+	};
+
+	LoopSegmentsIntersectionTool.Empty(NodesOfLoopCount);
+
+	Node = GetNext(StartNode);
+	for (; Node != StartNode; Node = GetNext(Node), ++SegmentCount);
+
+	Node = StartNode;
+	NextNode = GetNext(Node);
+
+#ifdef DEBUG_FIND_LOOP_INTERSECTIONS		
+	if (bDisplay)
+	{
+		F3DDebugSession _(TEXT("Start Node"));
+		Grid.DisplayIsoSegment(EGridSpace::UniformScaled, *Node, *NextNode, 0, EVisuProperty::BlueCurve);
+		Grid.DisplayIsoNode(EGridSpace::UniformScaled, *Node, 0, EVisuProperty::RedPoint);
+	}
+#endif
+
+	FIsoSegment* StartToEndSegment = StartNode->GetSegmentConnectedTo(NextNode);
+	if (StartToEndSegment != nullptr)
+	{
+		LoopSegmentsIntersectionTool.AddSegment(*StartToEndSegment);
+	}
+
+	for (Node = NextNode; Node != StartNode; Node = NextNode, ++SegmentIndex)
+	{
+		if (Node == nullptr || Node->IsDelete())
+		{
+			return;
+		}
+
+		NextNode = GetNext(Node);
+		if (NextNode == nullptr || NextNode->IsDelete())
+		{
+			return;
+		}
+		FindSegmentIntersection();
+	}
+
+}
+
+bool FLoopCleaner::RemoveOuterNode(const TPair<double, double>& Intersection)
+{
+	FLoopNode* Node0 = GetNodeAt(NextIndex((int32)Intersection.Key));
+	if (Node0->IsDelete())
+	{
+		return false;
+	}
+
+	FLoopNode* PreviousNode = GetPrevious(Node0);
+	FLoopNode* Node1 = GetNext(Node0);
+	FLoopNode* NextNode = GetNext(Node1);
+
+	const FPoint2D* PreviousPoint = &PreviousNode->Get2DPoint(EGridSpace::UniformScaled, Grid);
+	const FPoint2D* Point0 = &Node0->Get2DPoint(EGridSpace::UniformScaled, Grid);
+	const FPoint2D* Point1 = &Node1->Get2DPoint(EGridSpace::UniformScaled, Grid);
+	const FPoint2D* NextPoint = &NextNode->Get2DPoint(EGridSpace::UniformScaled, Grid);
+
+	double Slop0 = ComputeUnorientedSlope(*Point0, *PreviousPoint, *Point1);
+	double Slop1 = ComputeUnorientedSlope(*Point1, *PreviousPoint, *NextPoint);
+
+	if (Slop0 < Slop1)
+	{
+		return RemoveNodeOfLoop(*Node0);
+	}
+	else
+	{
+		return RemoveNodeOfLoop(*Node1);
+	}
+}
+
+bool FLoopCleaner::RemoveLoopPicks()
+{
+	for (FLoopNode* Node : NodesOfLoop)
+	{
+		if (Node == nullptr || Node->IsDelete())
+		{
+			continue;
+		}
+
+		if (!RemovePickRecursively(Node, &Node->GetNextNode()))
+		{
+			return false;
+		}
+	}
+	return UpdateNodesOfLoop();
+}
+
+bool FLoopCleaner::RemovePickRecursively(FLoopNode* Node0, FLoopNode* Node1)
+{
+	FLoopNode* PreviousNode = &Node0->GetPreviousNode();
+	FLoopNode* NextNode = &Node1->GetNextNode();
+
+	const FPoint2D* PreviousPoint = &PreviousNode->Get2DPoint(EGridSpace::UniformScaled, Grid);
+	const FPoint2D* Point0 = &Node0->Get2DPoint(EGridSpace::UniformScaled, Grid);
+	const FPoint2D* Point1 = &Node1->Get2DPoint(EGridSpace::UniformScaled, Grid);
+	const FPoint2D* NextPoint = &NextNode->Get2DPoint(EGridSpace::UniformScaled, Grid);
+
+	while (true)
+	{
+		bool bNodeIsToDelete = (Point0->SquareDistance(*Point1) < SquareGeometricTolerance2);
+		if (bNodeIsToDelete)
+		{
+			if (!RemoveNodeOfLoop(*Node0))
+			{
+				return false;
+			}
+		}
+
+		if (bNodeIsToDelete || CheckAndRemovePick(*PreviousPoint, *Point0, *Point1, *Node0))
+		{
+			if (PreviousNode->IsDelete())
+			{
+				break;
+			}
+
+			Point0 = PreviousPoint;
+			Node0 = PreviousNode;
+			PreviousNode = &PreviousNode->GetPreviousNode();
+			PreviousPoint = &PreviousNode->Get2DPoint(EGridSpace::UniformScaled, Grid);
+			continue;
+		}
+
+		if (CheckAndRemovePick(*Point0, *Point1, *NextPoint, *Node1))
+		{
+			if (NextNode->IsDelete())
+			{
+				break;
+			}
+
+			Point1 = NextPoint;
+			Node1 = NextNode;
+			NextNode = &NextNode->GetNextNode();
+			NextPoint = &NextNode->Get2DPoint(EGridSpace::UniformScaled, Grid);
+			continue;
+		}
+		break;
+	}
+
+	return true;
+}
+
+bool FLoopCleaner::RemoveNodeOfLoop(FLoopNode& NodeToRemove)
+{
+	if (NodeToRemove.GetConnectedSegments().Num() != 2)
+	{
+		return false;
+	}
+
+	FLoopNode& PreviousNode = NodeToRemove.GetPreviousNode();
+	FLoopNode& NextNode = NodeToRemove.GetNextNode();
+
+	FIsoSegment* Segment = PreviousNode.GetSegmentConnectedTo(&NodeToRemove);
+	if (Segment == nullptr)
+	{
+		return false;
+	}
+
+	FIsoSegment* SegmentToDelete = NextNode.GetSegmentConnectedTo(&NodeToRemove);
+	if (SegmentToDelete == nullptr)
+	{
+		return false;
+	}
+
+	NextNode.DisconnectSegment(*SegmentToDelete);
+	NextNode.ConnectSegment(*Segment);
+	Segment->SetSecondNode(NextNode);
+
+	if (&NextNode.GetNextNode() == &NextNode.GetPreviousNode())
+	{
+		NextNode.DisconnectSegment(*Segment);
+		PreviousNode.DisconnectSegment(*Segment);
+		RemoveSegmentOfLoops(Segment);
+		IsoSegmentFactory.DeleteEntity(Segment);
+
+		FIsoSegment* ThirdSegment = PreviousNode.GetSegmentConnectedTo(&NextNode);
+		if (!ThirdSegment)
+		{
+			return false;
+		}
+
+		NextNode.DisconnectSegment(*ThirdSegment);
+		PreviousNode.DisconnectSegment(*ThirdSegment);
+
+		RemoveSegmentOfLoops(ThirdSegment);
+		IsoSegmentFactory.DeleteEntity(ThirdSegment);
+
+		NextNode.Delete();
+		PreviousNode.Delete();
+
+		LoopSegmentsIntersectionTool.Remove(Segment);
+		LoopSegmentsIntersectionTool.Remove(ThirdSegment);
+	}
+
+	RemoveSegmentOfLoops(SegmentToDelete);
+	IsoSegmentFactory.DeleteEntity(SegmentToDelete);
+	NodeToRemove.Delete();
+
+	LoopSegmentsIntersectionTool.Remove(SegmentToDelete);
+	if (!Segment->IsDelete())
+	{
+		LoopSegmentsIntersectionTool.Update(Segment);
+	}
+
+	return true;
+}
+
+void FLoopCleaner::FindBestLoopExtremity()
+{
+	double UMin = HUGE_VAL;
+	double UMax = -HUGE_VAL;
+	double VMin = HUGE_VAL;
+	double VMax = -HUGE_VAL;
+
+	FLoopNode* ExtremityNodes[4] = { nullptr, nullptr, nullptr, nullptr };
+
+	FLoopNode* BestNode = nullptr;
+	double OptimalSlop = 9.;
+	LoopIndex = 0;
+
+	//F3DDebugSession _(TEXT("FindBestLoopExtremity"));
+
+	BestStartNodeOfLoops.Reserve(Grid.GetLoopCount());
+
+	TFunction<void(FLoopNode*)> CompareWithSlopAt = [&](FLoopNode* Node)
+	{
+		FLoopNode& PreviousNode = Node->GetPreviousNode();
+		FLoopNode& NextNode = Node->GetNextNode();
+		double Slop = ComputePositiveSlope(Node->Get2DPoint(EGridSpace::UniformScaled, Grid), PreviousNode.Get2DPoint(EGridSpace::UniformScaled, Grid), NextNode.Get2DPoint(EGridSpace::UniformScaled, Grid));
+
+		//F3DDebugSession A(FString::Printf(TEXT("Node Slop %f"), Slop));
+		//Display(EGridSpace::UniformScaled, *Node, 0, EVisuProperty::RedPoint);
+
+		if ((Slop > OptimalSlop) == (LoopIndex == 0))
+		{
+			OptimalSlop = Slop;
+			BestNode = Node;
+		}
+	};
+
+	TFunction<void()> FindLoopExtremity = [&]()
+	{
+		BestNode = nullptr;
+		OptimalSlop = (LoopIndex == 0) ? -1 : 9.;
+
+		for (FLoopNode* Node : ExtremityNodes)
+		{
+			CompareWithSlopAt(Node);
+		}
+		BestStartNodeOfLoops.Add(BestNode);
+
+		// init for next loop
+		UMin = HUGE_VAL;
+		UMax = -HUGE_VAL;
+		VMin = HUGE_VAL;
+		VMax = -HUGE_VAL;
+
+		for (FLoopNode*& Node : ExtremityNodes)
+		{
+			Node = nullptr;
+		}
+	};
+
+	// For each loops, Find the best loop extremity
+	// The end of the loop is know when Node.GetLoopIndex() change
+	int32 Index = 0;
+	for (FLoopNode& Node : LoopNodes)
+	{
+		if (Node.GetLoopIndex() != LoopIndex)
+		{
+			FindLoopExtremity();
+			LoopIndex = Node.GetLoopIndex();
+		}
+		Index++;
+		const FPoint2D& Point = Node.Get2DPoint(EGridSpace::UniformScaled, Grid);
+
+		if (Point.U > UMax)
+		{
+			UMax = Point.U;
+			ExtremityNodes[0] = &Node;
+		}
+		if (Point.U < UMin)
+		{
+			UMin = Point.U;
+			ExtremityNodes[1] = &Node;
+		}
+
+		if (Point.V > VMax)
+		{
+			VMax = Point.V;
+			ExtremityNodes[2] = &Node;
+		}
+		if (Point.V < VMin)
+		{
+			VMin = Point.V;
+			ExtremityNodes[3] = &Node;
+		}
+
+		//if (Node.GetIndex() == 44)
+		//{
+		//	UMax = VMax = 10000;
+		//	UMin = VMin = -10000;
+
+		//	ExtremityNodes[0] = \
+		//	ExtremityNodes[1] = \
+		//	ExtremityNodes[2] = \
+		//	ExtremityNodes[3] = &Node;
+		//}
+	}
+	FindLoopExtremity();
+}
+
+bool FLoopCleaner::CheckMainLoopConsistency()
+{
+	int32 OuterNodeCount = 0;
+	for (const FLoopNode& Node : LoopNodes)
+	{
+		if (Node.GetLoopIndex() != 0)
+		{
+			break;
+		}
+
+		if (!Node.IsDelete())
+		{
+			++OuterNodeCount;
+			if (OuterNodeCount > 2)
+			{
+				return true;
+			}
+		}
+	}
+	return false;
+}
+
+EOrientation FLoopCleaner::GetLoopOrientation(const FLoopNode* StartNode)
+{
+	using namespace IsoTriangulatorImpl;
+	double UMin = HUGE_VAL;
+	double UMax = -HUGE_VAL;
+	double VMin = HUGE_VAL;
+	double VMax = -HUGE_VAL;
+
+	const FLoopNode* ExtremityNodes[4] = { nullptr, nullptr, nullptr, nullptr };
+
+	const FLoopNode* BestNode = nullptr;
+
+	LoopIndex = StartNode->GetLoopIndex();
+	double OptimalSlop = (LoopIndex == 0) ? -1 : 9.;
+
+	TFunction<void(const FLoopNode*)> CompareWithSlopAt = [&](const FLoopNode* Node)
+	{
+		FLoopNode& PreviousNode = Node->GetPreviousNode();
+		FLoopNode& NextNode = Node->GetNextNode();
+		double Slop = ComputePositiveSlope(Node->Get2DPoint(EGridSpace::UniformScaled, Grid), PreviousNode.Get2DPoint(EGridSpace::UniformScaled, Grid), NextNode.Get2DPoint(EGridSpace::UniformScaled, Grid));
+
+		if ((Slop > OptimalSlop) == (LoopIndex == 0))
+		{
+			OptimalSlop = Slop;
+			BestNode = Node;
+		}
+	};
+
+	TFunction<void()> FindLoopExtremity = [&]()
+	{
+		BestNode = nullptr;
+
+		for (const FLoopNode* Node : ExtremityNodes)
+		{
+			CompareWithSlopAt(Node);
+		}
+
+		// init for next loop
+		UMin = HUGE_VAL;
+		UMax = -HUGE_VAL;
+		VMin = HUGE_VAL;
+		VMax = -HUGE_VAL;
+
+		for (const FLoopNode*& Node : ExtremityNodes)
+		{
+			Node = nullptr;
+		}
+	};
+
+	const FLoopNode* Node = LoopCleanerImpl::GetNextConstNodeImpl(StartNode);
+	for (; Node != StartNode; Node = LoopCleanerImpl::GetNextConstNodeImpl(Node))
+	{
+		const FPoint2D& Point = Node->Get2DPoint(EGridSpace::UniformScaled, Grid);
+
+		if (Point.U > UMax)
+		{
+			UMax = Point.U;
+			ExtremityNodes[0] = Node;
+		}
+		if (Point.U < UMin)
+		{
+			UMin = Point.U;
+			ExtremityNodes[1] = Node;
+		}
+
+		if (Point.V > VMax)
+		{
+			VMax = Point.V;
+			ExtremityNodes[2] = Node;
+		}
+		if (Point.V < VMin)
+		{
+			VMin = Point.V;
+			ExtremityNodes[3] = Node;
+		}
+	}
+	FindLoopExtremity();
+
+	if (LoopIndex == 0)
+	{
+		return OptimalSlop > 4 ? EOrientation::Front : EOrientation::Back;
+	}
+	else
+	{
+		return OptimalSlop < 4 ? EOrientation::Front : EOrientation::Back;
+	}
+}
+
+}
diff --git a/Engine/Source/Runtime/Datasmith/CADKernel/Private/CADKernel/Mesh/Structure/LoopCleaner.h b/Engine/Source/Runtime/Datasmith/CADKernel/Private/CADKernel/Mesh/Structure/LoopCleaner.h
new file mode 100644
--- /dev/null
+++ b/Engine/Source/Runtime/Datasmith/CADKernel/Private/CADKernel/Mesh/Structure/LoopCleaner.h
@@ -0,0 +1,431 @@
+// Copyright Epic Games, Inc. All Rights Reserved.
+#pragma once
+
+#include "CADKernel/Core/Types.h"
+#include "CADKernel/Core/Factory.h"
+#include "CADKernel/Math/SlopeUtils.h"
+#include "CADKernel/Mesh/Meshers/IsoTriangulator/IntersectionSegmentTool.h"
+#include "CADKernel/Mesh/Meshers/IsoTriangulator/IsoNode.h"
+#include "CADKernel/Mesh/Meshers/IsoTriangulator/IsoSegment.h"
+
+namespace CADKernel
+{
+class FGrid;
+class FIsoSegment;
+class FIsoTriangulator;
+class FLoopNode;
+
+struct FMesherReport;
+
+namespace LoopCleanerImpl
+{
+typedef TFunction<double(const FPoint2D&, const FPoint2D&, double)> GetSlopMethod;
+typedef TFunction<FLoopNode* (FLoopNode*)> GetNextNodeMethod;
+typedef TFunction<const FLoopNode* (const FLoopNode*)> GetNextConstNodeMethod;
+typedef TFunction<const FLoopNode* (const FIsoSegment*)> GetSegmentToNodeMethod;
+typedef TPair<FLoopNode*, FLoopNode*> FLoopSection;
+
+struct FPinchIntersectionContext
+{
+	const TPair<double, double>& Intersection;
+	FLoopNode* Nodes[2][3];
+	TArray<const FPoint2D*> Points[2];
+
+	FPinchIntersectionContext(const TPair<double, double>& InIntersection)
+		: Intersection(InIntersection)
+	{
+	}
+};
+
+inline FLoopNode* GetNextNodeImpl(FLoopNode* Node)
+{
+	return &Node->GetNextNode();
+}
+
+inline FLoopNode* GetPreviousNodeImpl(FLoopNode* Node)
+{
+	return &Node->GetPreviousNode();
+}
+
+inline const FLoopNode* GetNextConstNodeImpl(const FLoopNode* Node)
+{
+	return &Node->GetNextNode();
+}
+
+inline const FLoopNode* GetPreviousConstNodeImpl(const FLoopNode* Node)
+{
+	return &Node->GetPreviousNode();
+}
+
+inline const FLoopNode* GetFirstNode(const FIsoSegment* Segment)
+{
+	return (const FLoopNode*)&Segment->GetFirstNode();
+};
+
+inline const FLoopNode* GetSecondNode(const FIsoSegment* Segment)
+{
+	return (const FLoopNode*)&Segment->GetSecondNode();
+};
+
+inline void RemoveDeletedNodes(TArray<FLoopNode*>& NodesOfLoop)
+{
+	int32 Index = NodesOfLoop.IndexOfByPredicate([](FLoopNode* Node) { return Node->IsDelete(); });
+	if (Index == INDEX_NONE)
+	{
+		return;
+	}
+	int32 NewIndex = Index;
+	for (; Index < NodesOfLoop.Num(); ++Index)
+	{
+		if (!NodesOfLoop[Index]->IsDelete())
+		{
+			NodesOfLoop[NewIndex++] = NodesOfLoop[Index];
+		}
+	}
+	NodesOfLoop.SetNum(NewIndex);
+}
+
+}
+
+
+class FLoopCleaner
+{
+private:
+	FGrid& Grid;
+	TArray<FLoopNode>& LoopNodes;
+	TArray<FIsoSegment*>& LoopSegments;
+	TFactory<FIsoSegment>& IsoSegmentFactory;
+	FMesherReport& MesherReport;
+
+	bool bDisplay;
+
+	TArray<FLoopNode*> BestStartNodeOfLoops;
+
+	// Fields for the processing loop 
+	/** Index of the processing loop */
+	int32 LoopIndex;  
+
+	/**
+	 * Array of the proceed loop's nodes
+	 * Warning: some nodes can be deleted -> check with Node.IsDelete
+	 */
+	TArray<FLoopNode*> NodesOfLoop;
+
+	/**
+	 * This value doesn't need to be updated because the count of nodes doesn't change
+	 * Deleted nodes are not removed of NodesOfLoop array
+	 */
+	int32 NodesOfLoopCount;
+
+	/** Index of first segment of the processing loop */
+	int32 StartSegmentIndex;
+	int32 SegmentCount;
+
+	/** 
+	 * Index of first segment of the next loop 
+	 * Warning, this index need to be updated with the deletion of a segment
+	 * Use RemoveSegmentOfLoops to delete a segment
+	 * and use UpdateNextLoopFirstSegmentIndex in case of doubt
+	 */
+	int32 NextLoopFirstSegmentIndex;
+
+	bool bLoopOrientation;
+	TArray<TPair<double, double>> Intersections;
+	FIntersectionSegmentTool LoopSegmentsIntersectionTool;
+
+
+	const double GeometricTolerance;
+	const double SquareGeometricTolerance;
+	const double SquareGeometricTolerance2;
+	//const double MeshingTolerance;
+	//const double SquareMeshingTolerance;
+
+	LoopCleanerImpl::GetNextNodeMethod GetNext;
+	LoopCleanerImpl::GetNextNodeMethod GetPrevious;
+	LoopCleanerImpl::GetSegmentToNodeMethod GetFirst;
+	LoopCleanerImpl::GetSegmentToNodeMethod GetSecond;
+
+public:
+
+	FLoopCleaner(FIsoTriangulator& Triangulator);
+
+	bool Run()
+	{
+		if (!CleanLoops())
+		{
+			return false;
+		}
+		
+		if (!UncrossLoops())
+		{
+			return false;
+		}
+
+		return true;
+	}
+
+private:
+
+	bool CleanLoops();
+	bool UncrossLoops();
+
+	/**
+	 * For each loop, find the best starting node i.e. a node well oriented
+	 * All extremity nodes have to be identified, then one of the four nodes is chosen 
+	 */
+	void FindBestLoopExtremity();
+
+	//
+	// ========= CleanLoops methods ==========
+	//
+
+	/**
+	 * @return false if the process failed => the surface cannot be meshed
+	 */
+	bool RemoveLoopPicks();
+	bool RemovePickRecursively(FLoopNode* Node0, FLoopNode* Node1);
+
+	bool CheckAndRemovePick(const FPoint2D& PreviousPoint, const FPoint2D& NodeToRemovePoint, const FPoint2D& NextPoint, FLoopNode& NodeToRemove)
+	{
+		double Slop = ComputeUnorientedSlope(NodeToRemovePoint, PreviousPoint, NextPoint);
+		bool bRemoveNode = false;
+		if (Slop < 0.03)
+		{
+			bRemoveNode = true;
+		}
+		else if (Slop < 0.1)
+		{
+			double SquareDistance1 = SquareDistanceOfPointToSegment(PreviousPoint, NodeToRemovePoint, NextPoint);
+			double SquareDistance2 = SquareDistanceOfPointToSegment(NextPoint, NodeToRemovePoint, PreviousPoint);
+			double MinSquareDistance = FMath::Min(SquareDistance1, SquareDistance2);
+			if (MinSquareDistance < SquareGeometricTolerance2)
+			{
+				bRemoveNode = true;
+			}
+		}
+
+		if (bRemoveNode)
+		{
+			return RemoveNodeOfLoop(NodeToRemove);
+		}
+		return false;
+	};
+
+	bool RemoveNodeOfLoop(FLoopNode& NodeToRemove);
+
+	void FindLoopIntersections();
+
+	/**
+	 * @return false if the process failed => the surface cannot be meshed
+	 */
+	bool RemoveSelfIntersectionsOfLoop();
+
+	bool RemoveIntersection(TPair<double, double>& Intersection);
+	bool RemoveOutgoingLoop(const TPair<double, double>& Intersection, const TPair<double, double>& NextIntersection);
+	bool RemoveIntersectionsOfSubLoop(int32 IntersectionIndex, int32 IntersectionCount);
+	bool RemoveOuterNode(const TPair<double, double>& Intersection);
+	bool SwapNodes(const TPair<double, double>& Intersection);
+
+	/**
+	 * Two cases:
+	 *    - the segments of the intersection a closed parallel and in same orientation. in this case, the sub-loop is a long pick. The sub-Loop is delete
+	 *    - the loop is an inner loop closed to the border:
+	 *       _____________                _____________
+	 *   |  /             \			  |__/             \
+	 *    \/               |     =>                     |
+	 *    /\               |		   __               |
+	 *   |  \_____________/			  |  \_____________/
+	 *
+	 */
+	bool TryToSwapSegmentsOrRemoveLoop(const TPair<double, double>& Intersection);
+	void SwapSubLoopOrientation(int32 FirstSegmentIndex, int32 LastSegmentIndex);
+	bool RemoveSubLoop(FLoopNode* StartNode, FLoopNode* EndNode);
+
+	void MoveIntersectingSectionBehindOppositeSection(LoopCleanerImpl::FLoopSection IntersectingSection, LoopCleanerImpl::FLoopSection OppositeSection);
+	void MoveNodeBehindSegment(const FIsoSegment& IntersectingSegment, FLoopNode& NodeToMove);
+	void MoveNode(FLoopNode& NodeToMove, FPoint2D& NewPosition);
+
+	void FixLoopOrientation();
+	EOrientation GetLoopOrientation(const FLoopNode* StartNode);
+
+	//
+	// ========= UncrossLoops methods ==========
+	//
+
+	/**
+	 * Segment and the segment next (or before Segment) are intersecting IntersectingSegment
+	 * The common node is moved
+	 */
+	bool TryToRemoveIntersectionOfTwoConsecutiveIntersectingSegments(const FIsoSegment& IntersectingSegment, FIsoSegment& Segment);
+	void RemoveIntersectionByMovingOutsideSegmentNodeInside(const FIsoSegment& IntersectingSegment, const FIsoSegment& Segment, bool bIsSameInnerLoop);
+
+	//bool TryToRemoveSelfIntersectionByMovingTheClosedOusidePoint(const FIsoSegment& Segment0, const FIsoSegment& Segment1);
+	//bool TryToRemoveIntersectionByMovingTheClosedOusidePoint(const FIsoSegment& Segment0, const FIsoSegment& Segment1);
+
+	void OffsetSegment(FIsoSegment& Segment, TSegment<FPoint2D>& Segment2D, TSegment<FPoint2D>& IntersectingSegment2D);
+	void OffsetNode(FLoopNode& Node, TSegment<FPoint2D>& IntersectingSegment2D);
+
+
+	//
+	// ========= Commun methods ==========
+	//
+
+	void SwapSegments(FIsoSegment& Segment0, FIsoSegment& Segment1);
+
+	bool IsAPinch(const LoopCleanerImpl::FPinchIntersectionContext& Contex) const;
+
+	/**
+	 * OppositeCase:
+	 *        __________        ___________ 
+	 *       |__     ___|      |__      ___|
+	 *          \   /             \    /
+	 *           o o        =>     o  o
+	 *        __/   \___		__/    \___
+	 *       
+	 */
+	bool DisconnectCoincidentNodes(const LoopCleanerImpl::FPinchIntersectionContext& Contex);
+
+	/**
+	 * CrossingCase:
+	 *           --<-    ->--       --<-    --<-
+	 *           |   \  /   |       |   \  /   |
+	 *           |    oo    |   =>  |    o     |
+	 *           |   /  \   |       |          |
+	 *           -->-    -<--       |     o    |
+	 * 								|   /  \   |
+	 *                              -->-    ->--
+	 */
+	bool DisconnectCrossingSegments(LoopCleanerImpl::FPinchIntersectionContext& Context);
+
+	/**
+	 * Case:  o_________o          o_________o
+	 *             o         =>
+	 *          __/ \___		        o
+	 *								 __/ \___
+	 */
+	bool MovePickBehind(const TPair<double, double>& Intersection, bool bKeyIsExtremity);
+
+
+	//
+	// ========= Other ==========
+	//
+
+	bool CheckMainLoopConsistency();
+
+	FLoopNode* GetNodeAt(int32 Index)
+	{
+		while (Index >= NodesOfLoopCount)
+		{
+			Index -= NodesOfLoopCount;
+		}
+		
+		FLoopNode* Node = NodesOfLoop[Index];
+#ifdef DEBUG_LOOP_INTERSECTION_AND_FIX_IT	
+		if (Node->IsDelete())
+		{
+			Wait();
+		}
+#endif
+#ifdef CADKERNEL_DEV
+		ensureCADKernel(!Node->IsDelete());
+#endif
+		return Node;
+	};
+
+	int32 NextSegmentIndex(int32 StartIndex)
+	{
+		++StartIndex;
+		return FitSegmentIndex(StartIndex);
+	};
+
+	int32 NextIndex(int32 StartIndex)
+	{
+		++StartIndex;
+		return FitNodeIndex(StartIndex);
+	};
+
+	int32 FitNodeIndex(int32 Index)
+	{
+		while (Index >= NodesOfLoopCount)
+		{
+			Index -= NodesOfLoopCount;
+		}
+		return Index;
+	};
+
+	int32 FitSegmentIndex(int32 Index)
+	{
+		while (Index >= NextLoopFirstSegmentIndex)
+		{
+			Index -= SegmentCount;
+		}
+		return Index;
+	};
+
+	bool Fill(LoopCleanerImpl::FPinchIntersectionContext& Context);
+
+	void GetLoopNodeStartingFrom(FLoopNode* StartNode, TArray<FLoopNode*>& Loop)
+	{
+		FLoopNode* Node = StartNode;
+		Loop.Empty(LoopNodes.Num());
+		Loop.Add(StartNode);
+		for (Node = GetNext(Node); Node != StartNode; Node = GetNext(Node))
+		{
+			Loop.Add(Node);
+		}
+	}
+
+	bool UpdateNodesOfLoop()
+	{
+		LoopCleanerImpl::RemoveDeletedNodes(NodesOfLoop);
+		NodesOfLoopCount = NodesOfLoop.Num();
+		if (NodesOfLoopCount < 3)
+		{
+			return false;
+		}
+		return true;
+	}
+
+	/*
+	 * Update NextLoopFirstSegmentIndex i.e. the Index of the first segment of the next loop (for iteration purpose) 
+	 */
+	void UpdateNextLoopFirstSegmentIndex(int32 NewLoopIndex)
+	{
+		NextLoopFirstSegmentIndex = LoopSegments.IndexOfByPredicate([&](FIsoSegment* Segment) {
+			return ((FLoopNode&)Segment->GetFirstNode()).GetLoopIndex() > NewLoopIndex;
+			});
+		if (NextLoopFirstSegmentIndex == INDEX_NONE)
+		{
+			NextLoopFirstSegmentIndex = LoopSegments.Num();
+		}
+		SegmentCount = NextLoopFirstSegmentIndex - StartSegmentIndex;
+	}
+
+	void RemoveSegmentOfLoops(FIsoSegment* Segment)
+	{
+		LoopSegments.RemoveSingle(Segment);
+		NextLoopFirstSegmentIndex--;
+		SegmentCount--;
+	}
+
+#ifdef CADKERNEL_DEV		
+	void DisplayIntersection(const TPair<double, double>& Intersection)
+	{
+		if (bDisplay)
+		{
+			F3DDebugSession _(FString::Printf(TEXT("Intersection %f %f"), Intersection.Key, Intersection.Value));
+
+			FLoopNode* Segment0End = GetNodeAt(NextIndex((int32)Intersection.Key));
+			FLoopNode* Segment1Start = GetNodeAt((int32)Intersection.Value);
+
+			Grid.DisplayIsoSegment(EGridSpace::UniformScaled, *Segment0End, *GetPrevious(Segment0End), 0, EVisuProperty::RedCurve);
+			Grid.DisplayIsoSegment(EGridSpace::UniformScaled, *Segment1Start, *GetNext(Segment1Start), 0, EVisuProperty::RedCurve);
+			Grid.DisplayIsoNode(EGridSpace::UniformScaled, *GetPrevious(Segment0End), 0, EVisuProperty::RedPoint);
+			Grid.DisplayIsoNode(EGridSpace::UniformScaled, *Segment1Start, 0, EVisuProperty::RedPoint);
+		}
+	};
+#endif
+
+};
+
+}
\ No newline at end of file
