CADKernel Update

diff --git a/Engine/Source/Runtime/Datasmith/CADKernel/Private/CADKernel/Core/Version.h b/Engine/Source/Runtime/Datasmith/CADKernel/Private/CADKernel/Core/Version.h
--- a/Engine/Source/Runtime/Datasmith/CADKernel/Private/CADKernel/Core/Version.h
+++ b/Engine/Source/Runtime/Datasmith/CADKernel/Private/CADKernel/Core/Version.h
@@ -2,8 +2,8 @@
 #pragma once
 
 //** ** Version du Toolkit** **//
-#define TOOLKIT_VERSION 0, 2, 18390356
-#define TOOLKIT_VERSION_ASCII "0.2.18390356"
+#define TOOLKIT_VERSION 0, 2, 18518151
+#define TOOLKIT_VERSION_ASCII "0.2.18518151"
 
 #define RELEASE_DATE	 __DATE__
 
diff --git a/Engine/Source/Runtime/Datasmith/CADKernel/Private/CADKernel/Mesh/Structure/GridDisplay.cpp b/Engine/Source/Runtime/Datasmith/CADKernel/Private/CADKernel/Mesh/Structure/GridDisplay.cpp
--- a/Engine/Source/Runtime/Datasmith/CADKernel/Private/CADKernel/Mesh/Structure/GridDisplay.cpp
+++ b/Engine/Source/Runtime/Datasmith/CADKernel/Private/CADKernel/Mesh/Structure/GridDisplay.cpp
@@ -163,9 +163,11 @@ void FGrid::DisplayGridPolyline(EGridSpace Space, const TArray<FLoopNode*>& Node
 		//F3DDebugSession Q(TEXT("Node"));
 		if (bDisplayNode)
 		{
+			int32 Index = 0;
 			for (const FLoopNode* Node : Nodes)
 			{
-				DisplayIsoNode(Space, *Node, Node->GetIndex(), Node->IsDelete() ? EVisuProperty::PurplePoint : (EVisuProperty)((int32)Property - 1));
+				//DisplayIsoNode(Space, *Node, Node->GetIndex(), Node->IsDelete() ? EVisuProperty::PurplePoint : (EVisuProperty)((int32)Property - 1));
+				DisplayIsoNode(Space, *Node, Index++, Node->IsDelete() ? EVisuProperty::PurplePoint : (EVisuProperty)((int32)Property - 1));
 			}
 		}
 	}
diff --git a/Engine/Source/Runtime/Datasmith/CADKernel/Private/CADKernel/Mesh/Structure/LoopCleaner.cpp b/Engine/Source/Runtime/Datasmith/CADKernel/Private/CADKernel/Mesh/Structure/LoopCleaner.cpp
--- a/Engine/Source/Runtime/Datasmith/CADKernel/Private/CADKernel/Mesh/Structure/LoopCleaner.cpp
+++ b/Engine/Source/Runtime/Datasmith/CADKernel/Private/CADKernel/Mesh/Structure/LoopCleaner.cpp
@@ -212,9 +212,7 @@ bool FLoopCleaner::UncrossLoops()
 			}
 			else
 			{
-				// if no more append remove it
-				ensureCADKernel(false);
-				LoopSegmentsIntersectionTool.AddSegment(Segment);
+				return false;
 			}
 
 #ifdef DEBUG_UNCROSS_LOOPS
@@ -816,6 +814,16 @@ bool FLoopCleaner::RemoveSelfIntersectionsOfLoop()
 	{
 		const TPair<double, double>& Intersection = Intersections[IntersectionIndex];
 
+		FLoopNode* Segment0End = GetNodeAt(NextIndex((int32)Intersection.Key));
+		FLoopNode* Segment1Start = GetNodeAt((int32)Intersection.Value);
+
+		if (!Segment0End || !Segment1Start)
+		{
+			++IntersectionIndex;
+			continue;
+		}
+
+
 #ifdef DEBUG_REMOVE_LOOP_INTERSECTIONS		
 		F3DDebugSession _(bDisplay, TEXT("Intersected Segments"));
 		DisplayIntersection(Intersection);
@@ -1115,17 +1123,32 @@ bool FLoopCleaner::RemoveIntersectionsOfSubLoop(int32 IntersectionIndex, int32 I
 			else
 			{
 				FLoopSection IntersectingSection;
-				IntersectingSection.Key = GetPrevious(GetNodeAt(IndexSide0));
-				IntersectingSection.Value = GetNext(GetNodeAt((int32)SecondIntersection.Key));
-				if (IntersectingSection.Key->IsDelete() || IntersectingSection.Value->IsDelete())
+				IntersectingSection.Key = GetNodeAt(IndexSide0);
+				IntersectingSection.Value = GetNodeAt((int32)SecondIntersection.Key);
+				if(!IntersectingSection.Key || !IntersectingSection.Value)
+				{
+					return false;
+				}
+
+				IntersectingSection.Key = GetPrevious(IntersectingSection.Key);
+				IntersectingSection.Value = GetNext(IntersectingSection.Value);
+				if (!IntersectingSection.Key || IntersectingSection.Key->IsDelete() || !IntersectingSection.Value || IntersectingSection.Value->IsDelete())
 				{
 					return false;
 				}
 
 				FLoopSection OppositeSection;
-				OppositeSection.Key = GetPrevious(GetNodeAt(IndexSide1));
-				OppositeSection.Value = GetNext(GetNodeAt((int32)FirstIntersection.Value));
-				if (OppositeSection.Key->IsDelete() || OppositeSection.Value->IsDelete())
+				OppositeSection.Key = GetNodeAt(IndexSide1);
+				OppositeSection.Value = GetNodeAt((int32)FirstIntersection.Value);
+				if (!OppositeSection.Key || !OppositeSection.Value)
+				{
+					return false;
+				}
+
+				OppositeSection.Key = GetPrevious(OppositeSection.Key);
+				OppositeSection.Value = GetNext(OppositeSection.Value);
+
+				if (!OppositeSection.Key || OppositeSection.Key->IsDelete() || !OppositeSection.Value || OppositeSection.Value->IsDelete())
 				{
 					return false;
 				}
@@ -1701,7 +1724,7 @@ bool FLoopCleaner::RemoveSubLoop(FLoopNode* StartNode, FLoopNode* EndNode)
 	FLoopNode* Node = GetNext(StartNode);
 	while (Node && (Node != EndNode) && !Node->IsDelete())
 	{
-		if (RemoveNodeOfLoop(*Node))
+		if (!RemoveNodeOfLoop(*Node))
 		{
 			return false;
 		}
@@ -1709,8 +1732,6 @@ bool FLoopCleaner::RemoveSubLoop(FLoopNode* StartNode, FLoopNode* EndNode)
 		Node = GetNext(StartNode);
 	}
 
-	//Display(EGridSpace::UniformScaled, TEXT("Loop"), LoopSegments, true);
-	//Wait();
 	return true;
 }
 
diff --git a/Engine/Source/Runtime/Datasmith/CADKernel/Private/CADKernel/Mesh/Structure/LoopCleaner.h b/Engine/Source/Runtime/Datasmith/CADKernel/Private/CADKernel/Mesh/Structure/LoopCleaner.h
--- a/Engine/Source/Runtime/Datasmith/CADKernel/Private/CADKernel/Mesh/Structure/LoopCleaner.h
+++ b/Engine/Source/Runtime/Datasmith/CADKernel/Private/CADKernel/Mesh/Structure/LoopCleaner.h
@@ -320,15 +320,14 @@ private:
 		}
 		
 		FLoopNode* Node = NodesOfLoop[Index];
-#ifdef DEBUG_LOOP_INTERSECTION_AND_FIX_IT	
 		if (Node->IsDelete())
 		{
+#ifdef DEBUG_LOOP_INTERSECTION_AND_FIX_IT	
 			Wait();
-		}
-#endif
-#ifdef CADKERNEL_DEV
-		ensureCADKernel(!Node->IsDelete());
 #endif
+			return nullptr;
+		}
+
 		return Node;
 	};
 
@@ -418,6 +417,11 @@ private:
 			FLoopNode* Segment0End = GetNodeAt(NextIndex((int32)Intersection.Key));
 			FLoopNode* Segment1Start = GetNodeAt((int32)Intersection.Value);
 
+			if (!Segment0End || !Segment1Start)
+			{
+				return;
+			}
+
 			Grid.DisplayIsoSegment(EGridSpace::UniformScaled, *Segment0End, *GetPrevious(Segment0End), 0, EVisuProperty::RedCurve);
 			Grid.DisplayIsoSegment(EGridSpace::UniformScaled, *Segment1Start, *GetNext(Segment1Start), 0, EVisuProperty::RedCurve);
 			Grid.DisplayIsoNode(EGridSpace::UniformScaled, *GetPrevious(Segment0End), 0, EVisuProperty::RedPoint);
