CADKernel Update

diff --git a/Engine/Source/Runtime/Datasmith/CADKernel/Private/CADKernel/Core/Version.h b/Engine/Source/Runtime/Datasmith/CADKernel/Private/CADKernel/Core/Version.h
--- a/Engine/Source/Runtime/Datasmith/CADKernel/Private/CADKernel/Core/Version.h
+++ b/Engine/Source/Runtime/Datasmith/CADKernel/Private/CADKernel/Core/Version.h
@@ -2,8 +2,8 @@
 #pragma once
 
 //** ** Version du Toolkit** **//
-#define TOOLKIT_VERSION 0, 2, 18390356
-#define TOOLKIT_VERSION_ASCII "0.2.18390356"
+#define TOOLKIT_VERSION 0, 2, 18518151
+#define TOOLKIT_VERSION_ASCII "0.2.18518151"
 
 #define RELEASE_DATE	 __DATE__
 
diff --git a/Engine/Source/Runtime/Datasmith/CADKernel/Private/CADKernel/Mesh/Structure/GridDisplay.cpp b/Engine/Source/Runtime/Datasmith/CADKernel/Private/CADKernel/Mesh/Structure/GridDisplay.cpp
new file mode 100644
--- /dev/null
+++ b/Engine/Source/Runtime/Datasmith/CADKernel/Private/CADKernel/Mesh/Structure/GridDisplay.cpp
@@ -0,0 +1,384 @@
+// Copyright Epic Games, Inc. All Rights Reserved.
+#ifdef CADKERNEL_DEV
+
+#include "CADKernel/Mesh/Structure/Grid.h"
+
+#include "CADKernel/Mesh/Meshers/IsoTriangulator/IsoNode.h"
+#include "CADKernel/Mesh/Meshers/IsoTriangulator/IsoSegment.h"
+#include "CADKernel/UI/Display.h"
+
+namespace CADKernel
+{
+
+void FGrid::DisplayIsoNode(EGridSpace DisplaySpace, const int32 PointIndex, FIdent Ident, EVisuProperty Property) const
+{
+	if (!bDisplay)
+	{
+		return;
+	}
+	DisplayPoint(Points2D[DisplaySpace][PointIndex], Ident);
+}
+
+void FGrid::DisplayIsoNode(EGridSpace Space, const FIsoNode& Node, FIdent Ident, EVisuProperty Property) const
+{
+	if (!bDisplay)
+	{
+		return;
+	}
+	DisplayPoint(Node.GetPoint(Space, *this), Property, Ident);
+}
+
+void FGrid::DisplayIsoNodes(EGridSpace Space, const TArray<FIsoNode*>& Nodes, EVisuProperty Property) const
+{
+	if (!bDisplay)
+	{
+		return;
+	}
+
+	for (FIsoNode* Node : Nodes)
+	{
+		DisplayIsoNode(Space, *Node, Node->GetId(), Property);
+	}
+}
+
+void FGrid::DisplayIsoSegment(EGridSpace Space, const FIsoSegment& Segment, FIdent Ident, EVisuProperty Property, bool bDisplayOrientation) const
+{
+	if (!bDisplay)
+	{
+		return;
+	}
+	DisplaySegment(Segment.GetFirstNode().GetPoint(Space, *this), Segment.GetSecondNode().GetPoint(Space, *this), Ident, Property, bDisplayOrientation);
+}
+
+void FGrid::DisplayIsoSegment(EGridSpace Space, const FIsoNode& NodeA, const FIsoNode& NodeB, FIdent Ident, EVisuProperty Property) const
+{
+	if (!bDisplay)
+	{
+		return;
+	}
+	DisplaySegment(NodeA.GetPoint(Space, *this), NodeB.GetPoint(Space, *this), Ident, Property);
+}
+
+void FGrid::DisplayTriangle(EGridSpace Space, const FIsoNode& NodeA, const FIsoNode& NodeB, const FIsoNode& NodeC) const
+{
+	TArray<FPoint> Points;
+	Points.SetNum(3);
+	Points[0] = NodeA.GetPoint(Space, *this);
+	Points[1] = NodeB.GetPoint(Space, *this);
+	Points[2] = NodeC.GetPoint(Space, *this);
+	DrawElement(2, Points, EVisuProperty::Element);
+
+	DisplaySegment(Points[0], Points[1], 0, EVisuProperty::EdgeMesh);
+	DisplaySegment(Points[1], Points[2], 0, EVisuProperty::EdgeMesh);
+	DisplaySegment(Points[2], Points[0], 0, EVisuProperty::EdgeMesh);
+}
+
+
+void FGrid::DisplayIsoSegments(EGridSpace Space, const TArray<FIsoSegment*>& InSegments, bool bDisplayNode, bool bDisplayOrientation, EVisuProperty Property) const
+{
+	if (!bDisplay)
+	{
+		return;
+	}
+	int32 Index = 0;
+	for (FIsoSegment* Segment : InSegments)
+	{
+		if (Segment)
+		{
+			DisplaySegment(Segment->GetFirstNode().GetPoint(Space, *this), Segment->GetSecondNode().GetPoint(Space, *this), Index++, Property, bDisplayOrientation);
+			if(bDisplayNode)
+			{
+				DisplayPoint(Segment->GetFirstNode().GetPoint(Space, *this), (EVisuProperty)(Property - 1), Segment->GetFirstNode().GetIndex());
+				DisplayPoint(Segment->GetSecondNode().GetPoint(Space, *this), (EVisuProperty)(Property - 1), Segment->GetSecondNode().GetIndex());
+			}
+		}
+	}
+}
+
+//void FGrid::DisplayGridSegments(const TCHAR* Message, EGridSpace Space, const TArray<FIsoSegment*>& Segments, bool bDisplayNode, bool bDisplayOrientation, EVisuProperty Property) const
+//{
+//	if (bDisplay)
+//	{
+//		F3DDebugSession _(Message);
+//		DisplayGridSegments(Space, Segments, bDisplayNode, bDisplayOrientation, Property);
+//	}
+//}
+
+void FGrid::DisplayGridPolyline(EGridSpace Space, const FLoopNode& StartNode, bool bDisplayNode, EVisuProperty Property) const
+{
+	if (!bDisplay)
+	{
+		return;
+	}
+	int32 Index = 1;
+	const FLoopNode* Node = &StartNode.GetNextNode();
+	const FLoopNode* NextNode = nullptr;
+	DisplayIsoSegment(Space, StartNode, *Node, Index++, Property);
+	for (; Node != &StartNode; Node = NextNode)
+	{
+		NextNode = &Node->GetNextNode();
+		DisplayIsoSegment(Space, *Node, *NextNode, Index++, Property);
+	}
+
+	DisplayIsoNode(Space, StartNode, StartNode.GetIndex(), Property);
+	if(bDisplayNode)
+	{
+		for (; Node != &StartNode; Node = &Node->GetNextNode())
+		{
+			DisplayIsoNode(Space, *Node, Node->GetIndex(), (EVisuProperty)(Property - 1));
+		}
+	}
+}
+
+void FGrid::DisplayGridPolyline(EGridSpace Space, const TArray<FLoopNode*>& Nodes, bool bDisplayNode, EVisuProperty Property) const
+{
+	if (!bDisplay)
+	{
+		return;
+	}
+
+	{
+		//F3DDebugSession Q(TEXT("GetNextNode"));
+		for (const FLoopNode* Node : Nodes)
+		{
+			if (!Node->IsDelete())
+			{
+				DisplayIsoSegment(Space, *Node, Node->GetNextNode(), 0, Property);
+			}
+		}
+	}
+
+	//{
+	//	F3DDebugSession Q(TEXT("GetPreviousNode"));
+	//	for (const FLoopNode* Node : Nodes)
+	//	{
+	//		if (!Node->IsDelete())
+	//		{
+	//			DisplayIsoSegment(Space, *Node, Node->GetPreviousNode(), 0, Property);
+	//		}
+	//	}
+	//}
+
+	{
+		//F3DDebugSession Q(TEXT("Node"));
+		if (bDisplayNode)
+		{
+			int32 Index = 0;
+			for (const FLoopNode* Node : Nodes)
+			{
+				//DisplayIsoNode(Space, *Node, Node->GetIndex(), Node->IsDelete() ? EVisuProperty::PurplePoint : (EVisuProperty)((int32)Property - 1));
+				DisplayIsoNode(Space, *Node, Index++, Node->IsDelete() ? EVisuProperty::PurplePoint : (EVisuProperty)((int32)Property - 1));
+			}
+		}
+	}
+}
+
+void FGrid::DisplayGridPolyline(EGridSpace Space, const TArray<FLoopNode>& Nodes, bool bDisplayNode, EVisuProperty Property) const
+{
+	if (!bDisplay)
+	{
+		return;
+	}
+
+	{
+		//F3DDebugSession Q(TEXT("GetNextNode"));
+		for (const FLoopNode& Node : Nodes)
+		{
+			if (!Node.IsDelete())
+			{
+				DisplayIsoSegment(Space, Node, Node.GetNextNode(), 0, Property);
+			}
+		}
+	}
+
+	{
+		//F3DDebugSession Q(TEXT("Node"));
+		if (bDisplayNode)
+		{
+			for (const FLoopNode& Node : Nodes)
+			{
+				DisplayIsoNode(Space, Node, Node.GetIndex(), Node.IsDelete() ? EVisuProperty::PurplePoint : (EVisuProperty)((int32)Property - 1));
+			}
+		}
+	}
+}
+
+void FGrid::DisplayGridPolyline(EGridSpace Space, const TArray<FIsoNode*>& Nodes, bool bDisplayNode, EVisuProperty Property) const
+{
+	if (!bDisplay)
+	{
+		return;
+	}
+	int32 Index = 1;
+	for (; Index < Nodes.Num(); ++Index)
+	{
+		DisplayIsoSegment(Space, *Nodes[Index - 1], *Nodes[Index], Index, Property);
+	}
+	DisplayIsoSegment(Space, *Nodes[Index - 1], *Nodes[0], Index, Property);
+
+	if (bDisplayNode)
+	{
+		for (const FIsoNode* Node : Nodes)
+		{
+			DisplayIsoNode(Space, *Node, Node->GetIndex(), (EVisuProperty)(Property - 1));
+		}
+	}
+}
+
+#ifdef TOTOT
+void FGrid::DisplayGridLoopNodesAndSegments(const TCHAR* Message, bool bOneNode, bool bSplitBySegment, bool bDisplayNext, bool bDisplayPrevious) const
+{
+	if (bDisplay)
+	{
+		F3DDebugSession _(Message);
+		{
+			{
+				F3DDebugSession A(!bOneNode, TEXT("Node Cloud"));
+				int32 Index = 0;
+				for (const FLoopNode& Node : LoopNodes)
+				{
+					if (Node.IsDelete())
+					{
+						Display(EGridSpace::UniformScaled, Node, Index++, EVisuProperty::PurplePoint);
+					}
+					Display(EGridSpace::UniformScaled, Node, Index++, EVisuProperty::BluePoint);
+				}
+			}
+
+			{
+				F3DDebugSession C(!bOneNode, TEXT("Segment"));
+				int32 Index = 0;
+				for (FIsoSegment* Segment : LoopSegments)
+				{
+					F3DDebugSession D(bSplitBySegment, TEXT("Segment"));
+					Display(EGridSpace::UniformScaled, Segment->GetFirstNode(), Segment->GetSecondNode(), Index++, EVisuProperty::BlueCurve);
+				}
+			}
+
+			if (bDisplayNext)
+			{
+				F3DDebugSession A(TEXT("Node Next"));
+				int32 Index = 0;
+				for (FIsoSegment* Segment : LoopSegments)
+				{
+					F3DDebugSession B(bSplitBySegment, TEXT("Segment"));
+					Display(EGridSpace::UniformScaled, Segment->GetFirstNode(), ((const FLoopNode&)Segment->GetFirstNode()).GetNextNode(), Index++, EVisuProperty::BlueCurve);
+				}
+			}
+
+			if (bDisplayPrevious)
+			{
+				F3DDebugSession C(TEXT("Node Forward"));
+				int32 Index = 0;
+				for (FIsoSegment* Segment : LoopSegments)
+				{
+					F3DDebugSession B(bSplitBySegment, TEXT("Segment"));
+					Display(EGridSpace::UniformScaled, Segment->GetSecondNode(), ((const FLoopNode&)Segment->GetSecondNode()).GetPreviousNode(), Index++, EVisuProperty::BlueCurve);
+				}
+			}
+
+		}
+		Wait(false);
+	}
+}
+#endif
+
+void FGrid::DisplayNodes(const TCHAR* Message, EGridSpace Space, const TArray<const FIsoNode*>& Nodes, EVisuProperty Property) const
+{
+	if (!bDisplay)
+	{
+		return;
+	}
+
+	Open3DDebugSession(Message);
+	for (int32 Index = 0; Index < Nodes.Num(); ++Index)
+	{
+		DisplayIsoNode(Space, *Nodes[Index], Nodes[Index]->GetIndex(), Property);
+	}
+	Close3DDebugSession();
+}
+
+void FGrid::DisplayGridPoints(EGridSpace DisplaySpace) const
+{
+	if (!bDisplay)
+	{
+		return;
+	}
+	int32 NbNum = 0;
+	{
+		F3DDebugSession _(TEXT("FGrid::FindInnerDomainPoints Inside Point"));
+		for (int32 Index = 0; Index < CuttingSize; ++Index)
+		{
+			if (IsInsideFace[Index])
+			{
+				DisplayPoint(Points2D[DisplaySpace][Index], EVisuProperty::BluePoint, Index);
+				NbNum++;
+			}
+		}
+	}
+	ensureCADKernel(NbNum == CountOfInnerNodes);
+
+	{
+		F3DDebugSession _(TEXT("FGrid::FindInnerDomainPoints Outside Point"));
+		for (int32 Index = 0; Index < CuttingSize; ++Index)
+		{
+			if (!IsInsideFace[Index])
+			{
+				DisplayPoint(Points2D[DisplaySpace][Index], EVisuProperty::OrangePoint, Index);
+			}
+		}
+	}
+}
+
+void FGrid::DisplayGridNormal() const
+{
+	if (!bDisplay)
+	{
+		return;
+	}
+
+	double NormalLength = FSystem::Get().GetVisu()->GetParameters()->NormalLength;
+	{
+		F3DDebugSession _(TEXT("FGrid::Inner Normal"));
+
+		for (int32 Index = 0; Index < CuttingSize; ++Index)
+		{
+			F3DDebugSegment GraphicSegment(Index);
+			FVector Normal = Normals[Index];
+			Normal.Normalize();
+			Normal *= NormalLength;
+			DrawSegment(Points3D[Index], Points3D[Index] + Normal, EVisuProperty::GreenCurve);
+		}
+	}
+
+	{
+		F3DDebugSession _(TEXT("FGrid::Loop Normal"));
+		for (int32 LoopIndex = 0; LoopIndex < FaceLoops3D.Num(); ++LoopIndex)
+		{
+			const TArray<FPoint>& LoopPoints = FaceLoops3D[LoopIndex];
+			const TArray<FVector>& LoopNormals = NormalsOfFaceLoops[LoopIndex];
+
+			for (int32 Index = 0; Index < LoopPoints.Num(); ++Index)
+			{
+				F3DDebugSegment GraphicSegment(Index);
+				FVector Normal = LoopNormals[Index];
+				Normal.Normalize();
+				Normal *= NormalLength;
+				DrawSegment(LoopPoints[Index], LoopPoints[Index] + Normal, EVisuProperty::YellowCurve);
+			}
+		}
+	}
+}
+
+void FGrid::DisplayInnerPoints(TCHAR* Message, EGridSpace DisplaySpace) const
+{
+	if (!bDisplay)
+	{
+		return;
+	}
+	DisplayInnerDomainPoints(Message, Points2D[DisplaySpace]);
+}
+
+}
+
+#endif
\ No newline at end of file
diff --git a/Engine/Source/Runtime/Datasmith/CADKernel/Private/CADKernel/Mesh/Structure/LoopCleaner.cpp b/Engine/Source/Runtime/Datasmith/CADKernel/Private/CADKernel/Mesh/Structure/LoopCleaner.cpp
--- a/Engine/Source/Runtime/Datasmith/CADKernel/Private/CADKernel/Mesh/Structure/LoopCleaner.cpp
+++ b/Engine/Source/Runtime/Datasmith/CADKernel/Private/CADKernel/Mesh/Structure/LoopCleaner.cpp
@@ -212,9 +212,7 @@ bool FLoopCleaner::UncrossLoops()
 			}
 			else
 			{
-				// if no more append remove it
-				ensureCADKernel(false);
-				LoopSegmentsIntersectionTool.AddSegment(Segment);
+				return false;
 			}
 
 #ifdef DEBUG_UNCROSS_LOOPS
@@ -816,6 +814,16 @@ bool FLoopCleaner::RemoveSelfIntersectionsOfLoop()
 	{
 		const TPair<double, double>& Intersection = Intersections[IntersectionIndex];
 
+		FLoopNode* Segment0End = GetNodeAt(NextIndex((int32)Intersection.Key));
+		FLoopNode* Segment1Start = GetNodeAt((int32)Intersection.Value);
+
+		if (!Segment0End || !Segment1Start)
+		{
+			++IntersectionIndex;
+			continue;
+		}
+
+
 #ifdef DEBUG_REMOVE_LOOP_INTERSECTIONS		
 		F3DDebugSession _(bDisplay, TEXT("Intersected Segments"));
 		DisplayIntersection(Intersection);
@@ -1115,17 +1123,32 @@ bool FLoopCleaner::RemoveIntersectionsOfSubLoop(int32 IntersectionIndex, int32 I
 			else
 			{
 				FLoopSection IntersectingSection;
-				IntersectingSection.Key = GetPrevious(GetNodeAt(IndexSide0));
-				IntersectingSection.Value = GetNext(GetNodeAt((int32)SecondIntersection.Key));
-				if (IntersectingSection.Key->IsDelete() || IntersectingSection.Value->IsDelete())
+				IntersectingSection.Key = GetNodeAt(IndexSide0);
+				IntersectingSection.Value = GetNodeAt((int32)SecondIntersection.Key);
+				if(!IntersectingSection.Key || !IntersectingSection.Value)
+				{
+					return false;
+				}
+
+				IntersectingSection.Key = GetPrevious(IntersectingSection.Key);
+				IntersectingSection.Value = GetNext(IntersectingSection.Value);
+				if (!IntersectingSection.Key || IntersectingSection.Key->IsDelete() || !IntersectingSection.Value || IntersectingSection.Value->IsDelete())
 				{
 					return false;
 				}
 
 				FLoopSection OppositeSection;
-				OppositeSection.Key = GetPrevious(GetNodeAt(IndexSide1));
-				OppositeSection.Value = GetNext(GetNodeAt((int32)FirstIntersection.Value));
-				if (OppositeSection.Key->IsDelete() || OppositeSection.Value->IsDelete())
+				OppositeSection.Key = GetNodeAt(IndexSide1);
+				OppositeSection.Value = GetNodeAt((int32)FirstIntersection.Value);
+				if (!OppositeSection.Key || !OppositeSection.Value)
+				{
+					return false;
+				}
+
+				OppositeSection.Key = GetPrevious(OppositeSection.Key);
+				OppositeSection.Value = GetNext(OppositeSection.Value);
+
+				if (!OppositeSection.Key || OppositeSection.Key->IsDelete() || !OppositeSection.Value || OppositeSection.Value->IsDelete())
 				{
 					return false;
 				}
@@ -1701,7 +1724,7 @@ bool FLoopCleaner::RemoveSubLoop(FLoopNode* StartNode, FLoopNode* EndNode)
 	FLoopNode* Node = GetNext(StartNode);
 	while (Node && (Node != EndNode) && !Node->IsDelete())
 	{
-		if (RemoveNodeOfLoop(*Node))
+		if (!RemoveNodeOfLoop(*Node))
 		{
 			return false;
 		}
@@ -1709,8 +1732,6 @@ bool FLoopCleaner::RemoveSubLoop(FLoopNode* StartNode, FLoopNode* EndNode)
 		Node = GetNext(StartNode);
 	}
 
-	//Display(EGridSpace::UniformScaled, TEXT("Loop"), LoopSegments, true);
-	//Wait();
 	return true;
 }
 
diff --git a/Engine/Source/Runtime/Datasmith/CADKernel/Private/CADKernel/Mesh/Structure/LoopCleaner.h b/Engine/Source/Runtime/Datasmith/CADKernel/Private/CADKernel/Mesh/Structure/LoopCleaner.h
--- a/Engine/Source/Runtime/Datasmith/CADKernel/Private/CADKernel/Mesh/Structure/LoopCleaner.h
+++ b/Engine/Source/Runtime/Datasmith/CADKernel/Private/CADKernel/Mesh/Structure/LoopCleaner.h
@@ -320,15 +320,14 @@ private:
 		}
 		
 		FLoopNode* Node = NodesOfLoop[Index];
-#ifdef DEBUG_LOOP_INTERSECTION_AND_FIX_IT	
 		if (Node->IsDelete())
 		{
+#ifdef DEBUG_LOOP_INTERSECTION_AND_FIX_IT	
 			Wait();
-		}
-#endif
-#ifdef CADKERNEL_DEV
-		ensureCADKernel(!Node->IsDelete());
 #endif
+			return nullptr;
+		}
+
 		return Node;
 	};
 
@@ -418,6 +417,11 @@ private:
 			FLoopNode* Segment0End = GetNodeAt(NextIndex((int32)Intersection.Key));
 			FLoopNode* Segment1Start = GetNodeAt((int32)Intersection.Value);
 
+			if (!Segment0End || !Segment1Start)
+			{
+				return;
+			}
+
 			Grid.DisplayIsoSegment(EGridSpace::UniformScaled, *Segment0End, *GetPrevious(Segment0End), 0, EVisuProperty::RedCurve);
 			Grid.DisplayIsoSegment(EGridSpace::UniformScaled, *Segment1Start, *GetNext(Segment1Start), 0, EVisuProperty::RedCurve);
 			Grid.DisplayIsoNode(EGridSpace::UniformScaled, *GetPrevious(Segment0End), 0, EVisuProperty::RedPoint);
