Adding ConnectionState member to FStreamInfo, allowing user-code to distinguish between "real" inactive streams and those that are just waiting for a response.
Also fixing crash when accessing UDirectLinkExtensionSettings CDO on engine shutdown.

Johan.Duparc

diff --git a/Engine/Plugins/Enterprise/DatasmithImporter/Source/DirectLinkExtension/Private/DirectLinkManager.cpp b/Engine/Plugins/Enterprise/DatasmithImporter/Source/DirectLinkExtension/Private/DirectLinkManager.cpp
--- a/Engine/Plugins/Enterprise/DatasmithImporter/Source/DirectLinkExtension/Private/DirectLinkManager.cpp
+++ b/Engine/Plugins/Enterprise/DatasmithImporter/Source/DirectLinkExtension/Private/DirectLinkManager.cpp
@@ -24,10 +24,20 @@ DEFINE_LOG_CATEGORY(LogDirectLinkManager);
 
 namespace UE::DatasmithImporter
 {
+	FDirectLinkAutoReconnectManager::FDirectLinkAutoReconnectManager(FDirectLinkManager& InManager)
+		: Manager(InManager)
+		, bShouldRun(true)
+	{
+		if (const UDirectLinkExtensionSettings* DefaultSettings = GetDefault<UDirectLinkExtensionSettings>())
+		{
+			bAutoReconnectEnabled = DefaultSettings->bAutoReconnect;
+			ReconnectionDelayInSeconds = DefaultSettings->ReconnectionDelayInSeconds;
+		}
+	}
+
 	bool FDirectLinkAutoReconnectManager::Start()
 	{
-		if (GetDefault<UDirectLinkExtensionSettings>()->bAutoReconnect
-			&& (!CompletedFuture.IsValid() || CompletedFuture.IsReady()))
+		if (bAutoReconnectEnabled && (!CompletedFuture.IsValid() || CompletedFuture.IsReady()))
 		{
 			bShouldRun = true;
 			CompletedFuture = Async(EAsyncExecution::ThreadPool, [this]() {	Run(); });
@@ -47,7 +57,6 @@ namespace UE::DatasmithImporter
 	{
 		const float CurrentTime = FPlatformTime::Seconds();
 		const float TimeSinceLastTry = CurrentTime - LastTryTime;
-		const float ReconnectionDelayInSeconds = GetDefault<UDirectLinkExtensionSettings>()->ReconnectionDelayInSeconds;
 
 		if (TimeSinceLastTry < ReconnectionDelayInSeconds)
 		{
@@ -263,7 +272,8 @@ namespace UE::DatasmithImporter
 		TSet<DirectLink::FDestinationHandle> ActiveStreamsSources;
 		for (const DirectLink::FRawInfo::FStreamInfo& StreamInfo : RawInfoCache.StreamsInfo)
 		{
-			if (!StreamInfo.bIsActive)
+			if (!(StreamInfo.ConnectionState == DirectLink::EStreamConnectionState::Active
+				|| StreamInfo.ConnectionState == DirectLink::EStreamConnectionState::RequestSent))
 			{
 				continue;
 			}
@@ -325,7 +335,7 @@ namespace UE::DatasmithImporter
 				// We can infer that a DirectLink source is empty (no scene synced) by looking at if its stream is planning to send any data.
 				// #ueent_todo: Ideally it would be better to not allow an AsyncLoad to take place in the first time, but we can't know a source is empty 
 				//				before actually connecting to it, so this is the best we can do at the current time.
-				const bool bStreamIsEmpty = StreamInfo.bIsActive
+				const bool bStreamIsEmpty = StreamInfo.ConnectionState == DirectLink::EStreamConnectionState::Active
 					&& !StreamInfo.CommunicationStatus.IsTransmitting()
 					&& StreamInfo.CommunicationStatus.TaskTotal == 0;
 
diff --git a/Engine/Plugins/Enterprise/DatasmithImporter/Source/DirectLinkExtension/Private/DirectLinkManager.h b/Engine/Plugins/Enterprise/DatasmithImporter/Source/DirectLinkExtension/Private/DirectLinkManager.h
--- a/Engine/Plugins/Enterprise/DatasmithImporter/Source/DirectLinkExtension/Private/DirectLinkManager.h
+++ b/Engine/Plugins/Enterprise/DatasmithImporter/Source/DirectLinkExtension/Private/DirectLinkManager.h
@@ -28,15 +28,15 @@ namespace UE::DatasmithImporter
 	class FDirectLinkAutoReconnectManager
 	{
 	public:
-		FDirectLinkAutoReconnectManager(FDirectLinkManager& InManager)
-			: Manager(InManager)
-			, bShouldRun(true)
-		{}
+		FDirectLinkAutoReconnectManager(FDirectLinkManager& InManager);
 
 		~FDirectLinkAutoReconnectManager()
 		{
-			Stop();
-			CompletedFuture.Wait();
+			if (CompletedFuture.IsValid())
+			{
+				Stop();
+				CompletedFuture.Wait();
+			}
 		}
 
 		bool Start();
@@ -53,6 +53,10 @@ namespace UE::DatasmithImporter
 		TFuture<void> CompletedFuture;
 
 		float LastTryTime = 0;
+
+		bool bAutoReconnectEnabled = false;
+
+		float ReconnectionDelayInSeconds = 1;
 	};
 
 	class FDirectLinkManager: public IDirectLinkManager, public DirectLink::IEndpointObserver
diff --git a/Engine/Plugins/Experimental/Enterprise/DatasmithRuntime/Source/Private/DirectLinkUtils.cpp b/Engine/Plugins/Experimental/Enterprise/DatasmithRuntime/Source/Private/DirectLinkUtils.cpp
--- a/Engine/Plugins/Experimental/Enterprise/DatasmithRuntime/Source/Private/DirectLinkUtils.cpp
+++ b/Engine/Plugins/Experimental/Enterprise/DatasmithRuntime/Source/Private/DirectLinkUtils.cpp
@@ -416,7 +416,8 @@ namespace DatasmithRuntime
 
 		for ( const FRawInfo::FStreamInfo& StreamInfo : RawInfo.StreamsInfo)
 		{
-			if (!StreamInfo.bIsActive)
+			if (!(StreamInfo.ConnectionState == DirectLink::EStreamConnectionState::Active 
+				|| StreamInfo.ConnectionState == DirectLink::EStreamConnectionState::RequestSent))
 			{
 				continue;
 			}
diff --git a/Engine/Source/Developer/Datasmith/DatasmithExporterUI/Private/Widgets/SDirectLinkStreamManager.cpp b/Engine/Source/Developer/Datasmith/DatasmithExporterUI/Private/Widgets/SDirectLinkStreamManager.cpp
--- a/Engine/Source/Developer/Datasmith/DatasmithExporterUI/Private/Widgets/SDirectLinkStreamManager.cpp
+++ b/Engine/Source/Developer/Datasmith/DatasmithExporterUI/Private/Widgets/SDirectLinkStreamManager.cpp
@@ -421,7 +421,9 @@ void SDirectLinkStreamManager::UpdateData(const DirectLink::FRawInfo& RawInfo)
 	for ( const FRawInfo::FStreamInfo& StreamInfo : RawInfo.StreamsInfo )
 	{
 		TSharedPtr<FDestinationData> DestinationData;
-		if ( StreamInfo.bIsActive )
+		const bool bIsActive = StreamInfo.ConnectionState == EStreamConnectionState::Active;
+
+		if ( bIsActive )
 		{
 			if ( DestinationsMap.Contains( StreamInfo.Destination ) )
 			{
@@ -436,18 +438,18 @@ void SDirectLinkStreamManager::UpdateData(const DirectLink::FRawInfo& RawInfo)
 			TSharedRef<FStreamData>& StreamData = *StreamDataPtr;
 
 			// Made with the assumption that the source and destionation data can't change
-			if ( StreamData->bIsActive != StreamInfo.bIsActive )
+			if ( StreamData->bIsActive != bIsActive )
 			{
-				StreamData->bIsActive = StreamInfo.bIsActive;
+				StreamData->bIsActive = bIsActive;
 				bDirtyStreamList = true;
 			}
 		}
-		else if ( StreamInfo.bIsActive )
+		else if ( bIsActive )
 		{
 			TSharedRef<FSourceData>* SourceData = SourcesMap.Find( StreamInfo.Source );
 			if ( DestinationData.IsValid() && SourceData != nullptr )
 			{
-				TSharedRef<FStreamData> StreamData = MakeShared<FStreamData>( StreamInfo.bIsActive, *SourceData, DestinationData.ToSharedRef() );
+				TSharedRef<FStreamData> StreamData = MakeShared<FStreamData>( bIsActive, *SourceData, DestinationData.ToSharedRef() );
 				StreamMap.AddByHash( StreamIdHash, StreamId, StreamData );
 				bDirtyStreamList = true;
 			}
diff --git a/Engine/Source/Developer/Datasmith/DatasmithFacade/Public/DirectLink/DatasmithFacadeEndpointObserver.h b/Engine/Source/Developer/Datasmith/DatasmithFacade/Public/DirectLink/DatasmithFacadeEndpointObserver.h
--- a/Engine/Source/Developer/Datasmith/DatasmithFacade/Public/DirectLink/DatasmithFacadeEndpointObserver.h
+++ b/Engine/Source/Developer/Datasmith/DatasmithFacade/Public/DirectLink/DatasmithFacadeEndpointObserver.h
@@ -91,7 +91,7 @@ struct DATASMITHFACADE_API FDatasmithFacadeRawInfo
 		uint32 GetStreamId() const { return StreamInfo.StreamId; }
 		FGuid GetSource() const { return StreamInfo.Source; }
 		FGuid GetDestination() const { return StreamInfo.Destination; }
-		bool IsActive() const { return StreamInfo.bIsActive; }
+		bool IsActive() const { return StreamInfo.ConnectionState == DirectLink::EStreamConnectionState::Active; }
 
 #ifdef SWIG_FACADE
 	private:
diff --git a/Engine/Source/Runtime/Datasmith/DirectLink/Private/DirectLinkEndpoint.cpp b/Engine/Source/Runtime/Datasmith/DirectLink/Private/DirectLinkEndpoint.cpp
--- a/Engine/Source/Runtime/Datasmith/DirectLink/Private/DirectLinkEndpoint.cpp
+++ b/Engine/Source/Runtime/Datasmith/DirectLink/Private/DirectLinkEndpoint.cpp
@@ -226,7 +226,7 @@ void FEndpoint::RemoveSource(const FSourceHandle& SourceId)
 		for (FStreamDescription& Stream : SharedState.Streams)
 		{
 			if (Stream.SourcePoint == SourceId
-			 && Stream.Status != FStreamDescription::EConnectionState::Closed)
+			 && Stream.Status != EStreamConnectionState::Closed)
 			{
 				SharedState.CloseStreamInternal(Stream, _);
 			}
@@ -307,7 +307,7 @@ void FEndpoint::RemoveDestination(const FDestinationHandle& Destination)
 		for (FStreamDescription& Stream : SharedState.Streams)
 		{
 			if (Stream.DestinationPoint == Destination
-			 && Stream.Status != FStreamDescription::EConnectionState::Closed)
+			 && Stream.Status != EStreamConnectionState::Closed)
 			{
 				SharedState.CloseStreamInternal(Stream, _);
 			}
@@ -364,8 +364,8 @@ FEndpoint::EOpenStreamResult FEndpoint::OpenStream(const FSourceHandle& SourceId
 		{
 			if (Stream.SourcePoint == SourceId && Stream.DestinationPoint == DestinationId)
 			{
-				if (Stream.Status == FStreamDescription::EConnectionState::Active
-				 || Stream.Status == FStreamDescription::EConnectionState::RequestSent)
+				if (Stream.Status == EStreamConnectionState::Active
+				 || Stream.Status == EStreamConnectionState::RequestSent)
 				{
 					// useless case, temp because of the unfinished connection policy.
 					// #ue_directlink_connexion Replace with proper policy (user driven connection map) + log if this happen
@@ -455,7 +455,7 @@ FEndpoint::EOpenStreamResult FEndpoint::OpenStream(const FSourceHandle& SourceId
 		NewStream.DestinationPoint = DestinationId;
 		NewStream.LocalStreamPort = StreamPort;
 		NewStream.RemoteAddress = RemoteAddress;
-		NewStream.Status = FStreamDescription::EConnectionState::RequestSent;
+		NewStream.Status = EStreamConnectionState::RequestSent;
 	}
 	else
 	{
@@ -473,7 +473,7 @@ void FEndpoint::CloseStream(const FSourceHandle& SourceId, const FDestinationHan
 	{
 		if (Stream.SourcePoint == SourceId
 		 && Stream.DestinationPoint == DestinationId
-		 && Stream.Status != FStreamDescription::EConnectionState::Closed)
+		 && Stream.Status != EStreamConnectionState::Closed)
 		{
 			SharedState.CloseStreamInternal(Stream, _);
 		}
@@ -539,10 +539,10 @@ FString FInternalThreadState::ToString_dbg() const
 			const TCHAR* StatusText = TEXT("?"); //Stream.stabThisIsSource ? 'S' : 'D';
 			switch (Stream.Status)
 			{
-				case FStreamDescription::EConnectionState::Uninitialized: StatusText = TEXT("Uninitialized"); break;
-				case FStreamDescription::EConnectionState::RequestSent:   StatusText = TEXT("RequestSent  "); break;
-				case FStreamDescription::EConnectionState::Active:        StatusText = TEXT("Active       "); break;
-				case FStreamDescription::EConnectionState::Closed:        StatusText = TEXT("Closed       "); break;
+				case EStreamConnectionState::Uninitialized: StatusText = TEXT("Uninitialized"); break;
+				case EStreamConnectionState::RequestSent:   StatusText = TEXT("RequestSent  "); break;
+				case EStreamConnectionState::Active:        StatusText = TEXT("Active       "); break;
+				case EStreamConnectionState::Closed:        StatusText = TEXT("Closed       "); break;
 			}
 			Out.Appendf(TEXT("-- [%s] stream: %08X:%d %s %08X:%d\n"), StatusText, LocalPoint.A, Stream.LocalStreamPort, OrientationText, RemotePoint.A, Stream.RemoteStreamPort);
 		}
@@ -567,7 +567,7 @@ void FInternalThreadState::Handle_DeltaMessage(const FDirectLinkMsg_DeltaMessage
 		}
 
 		FStreamDescription& Stream = *StreamPtr;
-		bool bIsActive = Stream.Status == FStreamDescription::EConnectionState::Active;
+		bool bIsActive = Stream.Status == EStreamConnectionState::Active;
 		bool bIsReceiver = Stream.Receiver.IsValid();
 		bool bIsExpectedSender = Stream.RemoteAddress == Context->GetSender();
 		if (!bIsActive || !bIsReceiver || !bIsExpectedSender)
@@ -595,7 +595,7 @@ void FInternalThreadState::Handle_HaveListMessage(const FDirectLinkMsg_HaveListM
 		}
 
 		FStreamDescription& Stream = *StreamPtr;
-		bool bIsActive = Stream.Status == FStreamDescription::EConnectionState::Active;
+		bool bIsActive = Stream.Status == EStreamConnectionState::Active;
 		bool bIsSender = Stream.Sender.IsValid();
 		bool bIsExpectedRemote = Stream.RemoteAddress == Context->GetSender();
 		if (!bIsActive || !bIsSender || !bIsExpectedRemote)
@@ -728,7 +728,7 @@ void FInternalThreadState::Handle_OpenStreamRequest(const FDirectLinkMsg_OpenStr
 			if (Stream.SourcePoint == Message.SourceGuid && Stream.DestinationPoint == Message.DestinationGuid)
 			{
 				// #ue_directlink_cleanup implement a robust handling of duplicated connections, reopened connections, etc...
-				if (Stream.Status == FStreamDescription::EConnectionState::Active)
+				if (Stream.Status == EStreamConnectionState::Active)
 				{
 					DenyConnection(TEXT("Connection already active"));
 					return;
@@ -768,7 +768,7 @@ void FInternalThreadState::Handle_OpenStreamRequest(const FDirectLinkMsg_OpenStr
 		NewStream.LocalStreamPort = StreamPort;
 		NewStream.Sender = MoveTemp(NewSender);
 		NewStream.Receiver = MoveTemp(NewReceiver);
-		NewStream.Status = FStreamDescription::EConnectionState::Active;
+		NewStream.Status = EStreamConnectionState::Active;
 
 		UE_LOG(LogDirectLinkNet, Log, TEXT("Endpoint '%s': Accepted connection"), *SharedState.NiceName);
 		MessageEndpoint->Send(Answer, RemoteEndpointAddress);
@@ -789,7 +789,7 @@ void FInternalThreadState::Handle_OpenStreamAnswer(const FDirectLinkMsg_OpenStre
 		if (FStreamDescription* StreamPtr = SharedState.GetStreamByLocalPort(Message.RecipientStreamPort, _))
 		{
 			FStreamDescription& Stream = *StreamPtr;
-			if (Stream.Status == FStreamDescription::EConnectionState::RequestSent)
+			if (Stream.Status == EStreamConnectionState::RequestSent)
 			{
 				if (Message.bAccepted)
 				{
@@ -804,11 +804,11 @@ void FInternalThreadState::Handle_OpenStreamAnswer(const FDirectLinkMsg_OpenStre
 					}
 
 					check(Stream.Receiver || Stream.Sender)
-					Stream.Status = FStreamDescription::EConnectionState::Active;
+					Stream.Status = EStreamConnectionState::Active;
 				}
 				else
 				{
-					Stream.Status = FStreamDescription::EConnectionState::Closed;
+					Stream.Status = EStreamConnectionState::Closed;
 					UE_LOG(LogDirectLinkNet, Warning, TEXT("stream connection refused. %s"), *Message.Error);
 				}
 			}
@@ -974,7 +974,7 @@ void FInternalThreadState::RemoveEndpoint(const FMessageAddress& RemoteEndpointA
 		for (auto& Stream : SharedState.Streams)
 		{
 			if (Stream.RemoteAddress == RemoteEndpointAddress
-				&& Stream.Status != FStreamDescription::EConnectionState::Closed)
+				&& Stream.Status != EStreamConnectionState::Closed)
 			{
 				UE_CLOG(SharedState.bDebugLog, LogDirectLinkNet, Log, TEXT("Endpoint '%s': Closed connection  (reason: remote endpoint removed)"), *SharedState.NiceName);
 				bool bNotifyRemote = false;
@@ -1127,7 +1127,7 @@ void FInternalThreadState::Run()
 			FRWScopeLock _(SharedState.StreamsLock, SLT_ReadOnly);
 			for (FStreamDescription& Stream : SharedState.Streams)
 			{
-				if (Stream.Status == FStreamDescription::EConnectionState::Active
+				if (Stream.Status == EStreamConnectionState::Active
 					&& Stream.bThisIsSource && ensure(Stream.Sender.IsValid()))
 				{
 					Stream.Sender->Tick(Now_s);
@@ -1179,8 +1179,8 @@ void FInternalThreadState::Run()
 					StreamInfo.StreamId = Stream.LocalStreamPort;
 					StreamInfo.Source = Stream.SourcePoint;
 					StreamInfo.Destination = Stream.DestinationPoint;
-					StreamInfo.bIsActive = Stream.Status == FStreamDescription::EConnectionState::Active;
-					if (StreamInfo.bIsActive)
+					StreamInfo.ConnectionState = Stream.Status;
+					if (Stream.Status == EStreamConnectionState::Active)
 					{
 						if (Stream.Sender)
 						{
@@ -1288,7 +1288,7 @@ FStreamDescription* FSharedState::GetStreamByLocalPort(FStreamPort LocalPort, co
 
 void FSharedState::CloseStreamInternal(FStreamDescription& Stream, const FRWScopeLock& _, bool bNotifyRemote)
 {
-	if (Stream.Status == FStreamDescription::EConnectionState::Closed)
+	if (Stream.Status == EStreamConnectionState::Closed)
 	{
 		return;
 	}
@@ -1304,7 +1304,7 @@ void FSharedState::CloseStreamInternal(FStreamDescription& Stream, const FRWScop
 	}
 
 	// close local stream
-	Stream.Status = FStreamDescription::EConnectionState::Closed;
+	Stream.Status = EStreamConnectionState::Closed;
 	Stream.Sender.Reset();
 	Stream.Receiver.Reset(); // #ue_directlink_cleanup notify associated scene provider
 }
diff --git a/Engine/Source/Runtime/Datasmith/DirectLink/Private/DirectLinkStreamDescription.h b/Engine/Source/Runtime/Datasmith/DirectLink/Private/DirectLinkStreamDescription.h
--- a/Engine/Source/Runtime/Datasmith/DirectLink/Private/DirectLinkStreamDescription.h
+++ b/Engine/Source/Runtime/Datasmith/DirectLink/Private/DirectLinkStreamDescription.h
@@ -14,14 +14,6 @@ namespace DirectLink
 
 struct FStreamDescription
 {
-	enum class EConnectionState
-	{
-		Uninitialized,
-		RequestSent,
-		Active,
-		Closed,
-	};
-
 	bool bThisIsSource = false;
 
 	FGuid SourcePoint;
@@ -29,7 +21,7 @@ struct FStreamDescription
 	FStreamPort LocalStreamPort = 0; // works like an Id within that endpoint
 	FMessageAddress RemoteAddress;
 	FStreamPort RemoteStreamPort = 0;
-	EConnectionState Status = EConnectionState::Uninitialized;
+	EStreamConnectionState Status = EStreamConnectionState::Uninitialized;
 
 	TUniquePtr<IStreamReceiver> Receiver;
 	TSharedPtr<IStreamSender> Sender;
diff --git a/Engine/Source/Runtime/Datasmith/DirectLink/Public/DirectLinkCommon.h b/Engine/Source/Runtime/Datasmith/DirectLink/Public/DirectLinkCommon.h
--- a/Engine/Source/Runtime/Datasmith/DirectLink/Public/DirectLinkCommon.h
+++ b/Engine/Source/Runtime/Datasmith/DirectLink/Public/DirectLinkCommon.h
@@ -61,6 +61,13 @@ enum class EVisibility
 	Private,   // The connection point is not expected to be contacted from a remote
 };
 
+enum class EStreamConnectionState
+{
+	Uninitialized,
+	RequestSent,
+	Active,
+	Closed,
+};
 
 struct FCommunicationStatus
 {
diff --git a/Engine/Source/Runtime/Datasmith/DirectLink/Public/DirectLinkEndpoint.h b/Engine/Source/Runtime/Datasmith/DirectLink/Public/DirectLinkEndpoint.h
--- a/Engine/Source/Runtime/Datasmith/DirectLink/Public/DirectLinkEndpoint.h
+++ b/Engine/Source/Runtime/Datasmith/DirectLink/Public/DirectLinkEndpoint.h
@@ -63,7 +63,7 @@ struct FRawInfo
 		FStreamPort StreamId = InvalidStreamPort;
 		FGuid Source;
 		FGuid Destination;
-		bool bIsActive = false;
+		EStreamConnectionState ConnectionState;
 		FCommunicationStatus CommunicationStatus;
 	};
 	FMessageAddress ThisEndpointAddress;
