Datasmith DirectLink - Adding missing DirectLink serialization for IES Texture path.

JeanLuc.Conrenthin

diff --git a/Engine/Source/Runtime/Datasmith/DatasmithCore/Private/DatasmithSceneElementsImpl.h b/Engine/Source/Runtime/Datasmith/DatasmithCore/Private/DatasmithSceneElementsImpl.h
--- a/Engine/Source/Runtime/Datasmith/DatasmithCore/Private/DatasmithSceneElementsImpl.h
+++ b/Engine/Source/Runtime/Datasmith/DatasmithCore/Private/DatasmithSceneElementsImpl.h
@@ -631,6 +631,8 @@ FDatasmithLightActorElementImpl<InterfaceType>::FDatasmithLightActorElementImpl(
 	Store.RegisterParameter(Color,              "Color"               );
 	Store.RegisterParameter(LightFLags,         "LightFLags"          );
 	Store.RegisterParameter(Temperature,        "Temperature"         );
+	Store.RegisterParameter(IesFile,            "IesFile"             );
+	Store.RegisterParameter(IesTexturePathName, "IesTexturePathName"  );
 	Store.RegisterParameter(IesBrightnessScale, "IesBrightnessScale"  );
 	Store.RegisterParameter(IesRotation,        "IesRotation"         );
 	Store.RegisterParameter(IesTexturePathName,  "IesTexturePathName" );
