DatasmithDirectLink - Remove duplicated IesTexturePathName serialization, and reorder the directlink serialize for easier comparison with the class members.

diff --git a/Engine/Source/Runtime/Datasmith/DatasmithCore/Private/DatasmithSceneElementsImpl.h b/Engine/Source/Runtime/Datasmith/DatasmithCore/Private/DatasmithSceneElementsImpl.h
--- a/Engine/Source/Runtime/Datasmith/DatasmithCore/Private/DatasmithSceneElementsImpl.h
+++ b/Engine/Source/Runtime/Datasmith/DatasmithCore/Private/DatasmithSceneElementsImpl.h
@@ -629,13 +629,12 @@ FDatasmithLightActorElementImpl<InterfaceType>::FDatasmithLightActorElementImpl(
 
 	Store.RegisterParameter(Intensity,          "Intensity"           );
 	Store.RegisterParameter(Color,              "Color"               );
-	Store.RegisterParameter(LightFLags,         "LightFLags"          );
 	Store.RegisterParameter(Temperature,        "Temperature"         );
 	Store.RegisterParameter(IesFile,            "IesFile"             );
 	Store.RegisterParameter(IesTexturePathName, "IesTexturePathName"  );
 	Store.RegisterParameter(IesBrightnessScale, "IesBrightnessScale"  );
 	Store.RegisterParameter(IesRotation,        "IesRotation"         );
-	Store.RegisterParameter(IesTexturePathName,  "IesTexturePathName" );
+	Store.RegisterParameter(LightFLags,         "LightFLags"          );
 }
 
 template< typename InterfaceType = IDatasmithPointLightElement >
