Update CADKernel

diff --git a/Engine/Plugins/Enterprise/DatasmithCADImporter/Source/CADInterfaces/Private/CADFileReport.cpp b/Engine/Plugins/Enterprise/DatasmithCADImporter/Source/CADInterfaces/Private/CADFileReport.cpp
new file mode 100644
--- /dev/null
+++ b/Engine/Plugins/Enterprise/DatasmithCADImporter/Source/CADInterfaces/Private/CADFileReport.cpp
@@ -0,0 +1,81 @@
+//CopyrightEpicGames,Inc.AllRightsReserved.
+#pragma once
+
+#include "CADFileReport.h"
+
+#include "CADKernel/Core/Chrono.h"
+#include "CADKernel/Core/Types.h"
+#include "CADKernel/UI/Message.h"
+
+namespace CADKernel
+{
+
+void FCADFileReport::Print()
+{
+	FMessage::FillReportFile(TEXT("Body"), BodyCount);   
+	FMessage::FillReportFile(TEXT("Shell"), ShellCount);
+	FMessage::FillReportFile(TEXT("Face"), FaceCount);
+	FMessage::FillReportFile(TEXT("Failed Face"), FailedFaceCount);
+	FMessage::FillReportFile(TEXT("Loop"), LoopCount);
+	FMessage::FillReportFile(TEXT("Deg Loop"), DegeneratedLoopCount);
+	FMessage::FillReportFile(TEXT("Doubt Ori"), DoubtfulLoopOrientationCount);
+
+	FMessage::FillReportFile(TEXT("Edge"), EdgeCount);
+	FMessage::FillReportFile(TEXT("Deg Edge"), DegeneratedEdgeCount);
+	
+	FMessage::FillReportFile(TEXT(""), TEXT(""));
+	FMessage::FillReportFile(TEXT("Surface"), SurfaceCount);
+	FMessage::FillReportFile(TEXT("Deg Surf"), DegeneratedSurfaceCount);
+
+	FMessage::FillReportFile(TEXT(""), TEXT(""));
+	FMessage::FillReportFile(TEXT("Cone"), ConeSurfaceCount);
+	FMessage::FillReportFile(TEXT("Cylin"), CylinderSurfaceCount);
+	FMessage::FillReportFile(TEXT("LTrans"), LinearTransfoSurfaceCount);
+	FMessage::FillReportFile(TEXT("Nurbs"), NurbsSurfaceCount);
+	FMessage::FillReportFile(TEXT("Offset"), OffsetSurfaceCount);
+	FMessage::FillReportFile(TEXT("Plane"), PlaneSurfaceCount);
+	FMessage::FillReportFile(TEXT("Revol"), RevolutionSurfaceCount);
+	FMessage::FillReportFile(TEXT("Ruled"), RuledSurfaceCount);
+	FMessage::FillReportFile(TEXT("Sphere"), SphereSurfaceCount);
+	FMessage::FillReportFile(TEXT("Torus"), TorusSurfaceCount);
+
+	FMessage::FillReportFile(TEXT("Blend01"), Blend01SurfaceCount);
+	FMessage::FillReportFile(TEXT("Blend02"), Blend02SurfaceCount);
+	FMessage::FillReportFile(TEXT("Blend03"), Blend03SurfaceCount);
+	FMessage::FillReportFile(TEXT("Cylin"), CylindricalSurfaceCount);
+	FMessage::FillReportFile(TEXT("PipeS"), PipeSurfaceCount);
+	FMessage::FillReportFile(TEXT("Extru"), ExtrusionSurfaceCount);
+	FMessage::FillReportFile(TEXT("SurfFCur"), SurfaceFromCurvesCount);
+	FMessage::FillReportFile(TEXT("SurfTrans"), TransformSurfaceCount);
+
+	FMessage::FillReportFile(TEXT("SurfAsNurbs"), SurfaceAsNurbsCount);
+	FMessage::FillReportFile(TEXT(""), TEXT(""));
+
+	FMessage::FillReportFile(TEXT("Restric"), RestrictionCurveCount);
+
+	FMessage::FillReportFile(TEXT("Curve"), CurveCount);
+	FMessage::FillReportFile(TEXT("Circle"), CurveCircleCount);
+	FMessage::FillReportFile(TEXT("Composite"), CurveCompositeCount);
+	FMessage::FillReportFile(TEXT("Ellipse"), CurveEllipseCount);
+	FMessage::FillReportFile(TEXT("Helix"), CurveHelixCount);
+	FMessage::FillReportFile(TEXT("Hyperbola"), CurveHyperbolaCount);
+	FMessage::FillReportFile(TEXT("Line"), CurveLineCount);
+	FMessage::FillReportFile(TEXT("Nurbs"), CurveNurbsCount);
+	FMessage::FillReportFile(TEXT("Parabola"), CurveParabolaCount);
+	FMessage::FillReportFile(TEXT("PolyLine"), CurvePolyLineCount);
+	FMessage::FillReportFile(TEXT("Bld02Bound"), CurveBlend02BoundaryCount);
+	FMessage::FillReportFile(TEXT("Equation"), CurveEquationCount);
+	FMessage::FillReportFile(TEXT("Intersection"), CurveIntersectionCount);
+	FMessage::FillReportFile(TEXT("Offset"), CurveOffsetCount);
+	FMessage::FillReportFile(TEXT("CrvOnSurf"), CurveOnSurfCount);
+	FMessage::FillReportFile(TEXT("CrvTrans"), CurveTransformCount);
+	FMessage::FillReportFile(TEXT("CrvAsNurbs"), CurveAsNurbsCount);
+
+	FMessage::FillReportFile(TEXT(""), TEXT(""));
+	FMessage::FillReportFile(TEXT("Load"), LoadFileDuration);
+	FMessage::FillReportFile(TEXT("Pars"), ParsingDuration);
+
+	FMessage::FillReportFile(TEXT(""), TEXT(""));
+}
+
+}
diff --git a/Engine/Plugins/Enterprise/DatasmithCADImporter/Source/CADInterfaces/Private/CoreTechBridge.cpp b/Engine/Plugins/Enterprise/DatasmithCADImporter/Source/CADInterfaces/Private/CoreTechBridge.cpp
--- a/Engine/Plugins/Enterprise/DatasmithCADImporter/Source/CADInterfaces/Private/CoreTechBridge.cpp
+++ b/Engine/Plugins/Enterprise/DatasmithCADImporter/Source/CADInterfaces/Private/CoreTechBridge.cpp
@@ -83,11 +83,12 @@
 namespace CADKernel
 {
 
-FCoreTechBridge::FCoreTechBridge(FSession& InSession)
+FCoreTechBridge::FCoreTechBridge(FSession& InSession, FCADFileReport& InReport)
 	: Session(InSession)
 	, GeometricTolerance(Session.GetGeometricTolerance())
 	, SquareGeometricTolerance(FMath::Square(Session.GetGeometricTolerance()))
 	, SquareJoiningVertexTolerance(SquareGeometricTolerance * 2)
+	, Report(InReport)
 {
 }
 
@@ -104,6 +105,7 @@ TSharedRef<FBody> FCoreTechBridge::AddBody(CT_OBJECT_ID CTBodyId)
 	}
 
 	TSharedRef<FBody> Body = FEntity::MakeShared<FBody>();
+	Report.BodyCount++;
 
 	AddMetadata(CTBodyId, Body);
 
@@ -140,6 +142,7 @@ TSharedRef<FBody> FCoreTechBridge::AddBody(CT_OBJECT_ID CTBodyId)
 	{
 		TSharedRef<FShell> Shell = FEntity::MakeShared<FShell>();
 		Body->AddShell(Shell);
+		Report.ShellCount++;
 
 		AddMetadata(CTShellId, Shell);
 #ifdef CORETECHBRIDGE_DEBUG
@@ -162,6 +165,8 @@ TSharedRef<FBody> FCoreTechBridge::AddBody(CT_OBJECT_ID CTBodyId)
 
 void FCoreTechBridge::AddFace(CT_OBJECT_ID CTFaceId, TSharedRef<FShell>& Shell)
 {
+	Report.FaceCount++;
+
 	FSurfacicBoundary FaceBoundary;
 	Get2DCurvesRange(CTFaceId, FaceBoundary);
 
@@ -171,6 +176,7 @@ void FCoreTechBridge::AddFace(CT_OBJECT_ID CTFaceId, TSharedRef<FShell>& Shell)
 
 	if (Return != IO_OK || CTSurfaceId <= 0)
 	{
+		Report.FailedFaceCount++;
 		FMessage::Printf(Log, TEXT("The CTFace %d has invalid carrier surface, this face is ignored\n"), CTFaceId);
 		return;
 	}
@@ -178,6 +184,7 @@ void FCoreTechBridge::AddFace(CT_OBJECT_ID CTFaceId, TSharedRef<FShell>& Shell)
 	TSharedPtr<FSurface> SurfacePtr = AddSurface(CTSurfaceId, FaceBoundary);
 	if (!SurfacePtr.IsValid())
 	{
+		Report.FailedFaceCount++;
 		FMessage::Printf(Log, TEXT("The CTFace %d has invalid carrier surface, this face is ignored\n"), CTFaceId);
 		return;
 	}
@@ -188,6 +195,8 @@ void FCoreTechBridge::AddFace(CT_OBJECT_ID CTFaceId, TSharedRef<FShell>& Shell)
 	const FSurfacicTolerance& Tolerance = Surface->GetIsoTolerances();
 	if (Tolerance[0] < SMALL_NUMBER || Tolerance[1] < SMALL_NUMBER)
 	{
+		Report.DegeneratedSurfaceCount++;
+		Report.FailedFaceCount++;
 		FMessage::Printf(Log, TEXT("KIO Surface %d (of Body) has surfacic tolerance probleme: %f %f\n"), CTFaceId, Tolerance[0], Tolerance[1]);
 		return;
 	}
@@ -197,6 +206,8 @@ void FCoreTechBridge::AddFace(CT_OBJECT_ID CTFaceId, TSharedRef<FShell>& Shell)
 	const FSurfacicBoundary& SurfaceBounds = Surface->GetBoundary();
 	if (SurfaceBounds.IsDegenerated())
 	{
+		Report.DegeneratedSurfaceCount++;
+		Report.FailedFaceCount++;
 		FMessage::Printf(Log, TEXT("The CTFace %d has degenerated carrier surface, this face is ignored ([%f, %f], [%f, %f])\n"), CTFaceId, SurfaceBounds[EIso::IsoU].Min, SurfaceBounds[EIso::IsoU].Max, SurfaceBounds[EIso::IsoV].Min, SurfaceBounds[EIso::IsoV].Max);
 		return;
 	}
@@ -205,6 +216,7 @@ void FCoreTechBridge::AddFace(CT_OBJECT_ID CTFaceId, TSharedRef<FShell>& Shell)
 	Return = CT_FACE_IO::AskLoops(CTFaceId, CTLoopIds);
 	if (Return != IO_OK)
 	{
+		Report.FailedFaceCount++;
 		FMessage::Printf(Log, TEXT("The CTFace %d has problem to get its loops, this face is ignored\n"), CTFaceId);
 		return;
 	}
@@ -249,6 +261,7 @@ void FCoreTechBridge::AddFace(CT_OBJECT_ID CTFaceId, TSharedRef<FShell>& Shell)
 
 	if (Loops.Num() == 0)
 	{
+		Report.FailedFaceCount++;
 		FMessage::Printf(Log, TEXT("The CTFace %d is degenerate, this face is ignored\n"), CTFaceId);
 		return;
 	}
@@ -261,7 +274,7 @@ void FCoreTechBridge::AddFace(CT_OBJECT_ID CTFaceId, TSharedRef<FShell>& Shell)
 	Face->CtKioId = (FIdent)CTFaceId;
 #endif
 
-	Face->AddLoops(Loops);
+	Face->AddLoops(Loops, Report.DoubtfulLoopOrientationCount);
 
 	EOrientation Orientation = CTOrientation == CT_FORWARD ? EOrientation::Front : EOrientation::Back;
 
@@ -275,6 +288,8 @@ void FCoreTechBridge::Get2DCurvesRange(CT_OBJECT_ID CTFaceId, FSurfacicBoundary&
 
 TSharedPtr<FTopologicalLoop> FCoreTechBridge::AddLoop(CT_OBJECT_ID CTLoopId, TSharedRef<FSurface>& Surface)
 {
+	Report.LoopCount++;
+
 	CT_LIST_IO CTCoedgeIds;
 	CT_LOOP_IO::AskCoedges(CTLoopId, CTCoedgeIds);
 
@@ -302,6 +317,7 @@ TSharedPtr<FTopologicalLoop> FCoreTechBridge::AddLoop(CT_OBJECT_ID CTLoopId, TSh
 
 	if (Edges.Num() == 0)
 	{
+		Report.DegeneratedLoopCount++;
 		return TSharedPtr<FTopologicalLoop>();
 	}
 
@@ -349,6 +365,8 @@ void FCoreTechBridge::LinkEdgesLoop(CT_OBJECT_ID CTLoopId, FTopologicalLoop& Loo
 
 TSharedPtr<FTopologicalEdge> FCoreTechBridge::AddEdge(CT_OBJECT_ID CTCoedgeId, TSharedRef<FSurface>& Surface)
 {
+	Report.EdgeCount++;
+
 	// build carrier curve
 	CT_UINT32 Order;
 	CT_UINT32 PoleDim;
@@ -408,6 +426,7 @@ TSharedPtr<FTopologicalEdge> FCoreTechBridge::AddEdge(CT_OBJECT_ID CTCoedgeId, T
 
 	if (!Curve.IsValid())
 	{
+		Report.DegeneratedEdgeCount++;
 		return TSharedPtr<FTopologicalEdge>();
 	}
 
@@ -417,6 +436,7 @@ TSharedPtr<FTopologicalEdge> FCoreTechBridge::AddEdge(CT_OBJECT_ID CTCoedgeId, T
 	{
 		RestrictionCurve->Delete();
 		Curve->Delete();
+		Report.DegeneratedEdgeCount++;
 		return TSharedPtr<FTopologicalEdge>();
 	}
 
@@ -431,6 +451,8 @@ TSharedPtr<FTopologicalEdge> FCoreTechBridge::AddEdge(CT_OBJECT_ID CTCoedgeId, T
 
 TSharedPtr<FSurface> FCoreTechBridge::AddSurface(CT_OBJECT_ID CTSurfaceId, const FSurfacicBoundary& Boundary)
 {
+	Report.SurfaceCount++;
+
 	TSharedPtr<FEntity>* SurfacePtr = CTIdToEntity.Find((uint32)CTSurfaceId);
 	if (SurfacePtr != nullptr)
 	{
@@ -444,33 +466,43 @@ TSharedPtr<FSurface> FCoreTechBridge::AddSurface(CT_OBJECT_ID CTSurfaceId, const
 	switch (CTSurfaceType)
 	{
 	case CT_PLANE_TYPE:
+		Report.PlaneSurfaceCount++;
 		Surface = AddPlaneSurface(CTSurfaceId, Boundary);
 		break;
 	case CT_S_NURBS_TYPE:
+		Report.NurbsSurfaceCount++;
 		Surface = AddNurbsSurface(CTSurfaceId);
 		break;
 	case CT_S_REVOL_TYPE:
+		Report.RevolutionSurfaceCount++;
 		Surface = AddRevolutionSurface(CTSurfaceId, Boundary);
 		break;
 	case CT_S_OFFSET_TYPE:
+		Report.OffsetSurfaceCount++;
 		Surface = AddOffsetSurface(CTSurfaceId, Boundary);
 		break;
 	case CT_CYLINDER_TYPE:
+		Report.CylinderSurfaceCount++;
 		Surface = AddCylinderSurface(CTSurfaceId, Boundary);
 		break;
 	case CT_CONE_TYPE:
+		Report.ConeSurfaceCount++;
 		Surface = AddConeSurface(CTSurfaceId, Boundary);
 		break;
 	case CT_SPHERE_TYPE:
+		Report.SphereSurfaceCount++;
 		Surface = AddSphereSurface(CTSurfaceId, Boundary);
 		break;
 	case CT_TORUS_TYPE:
+		Report.TorusSurfaceCount++;
 		Surface = AddTorusSurface(CTSurfaceId, Boundary);
 		break;
 	case CT_S_RULED_TYPE:
+		Report.RuledSurfaceCount++;
 		Surface = AddRuledSurface(CTSurfaceId);
 		break;
 	case CT_S_LINEARTRANSFO_TYPE:
+		Report.LinearTransfoSurfaceCount++;
 		Surface = AddLinearTransfoSurface(CTSurfaceId);
 		break;
 	default:
@@ -493,6 +525,8 @@ TSharedPtr<FCurve> FCoreTechBridge::AddCurve(CT_OBJECT_ID CTCurveId, CT_OBJECT_I
 		return StaticCastSharedPtr<FCurve>(*CurvePtr);
 	}
 
+	Report.CurveCount++;
+
 	CT_OBJECT_TYPE CurveType;
 	CT_OBJECT_IO::AskType(CTCurveId, CurveType);
 
@@ -506,27 +540,35 @@ TSharedPtr<FCurve> FCoreTechBridge::AddCurve(CT_OBJECT_ID CTCurveId, CT_OBJECT_I
 	switch (CurveType)
 	{
 	case CT_C_NURBS_TYPE:
+		Report.CurveNurbsCount++;
 		Curve = AddNurbsCurve(CTCurveId);
 		break;
 	case CT_LINE_TYPE:
+		Report.CurveLineCount++;
 		Curve = AddLineCurve(CTCurveId, CTSurfaceId);
 		break;
 	case CT_C_COMPO_TYPE:
+		Report.CurveCompositeCount++;
 		Curve = AddCompositeCurve(CTCurveId);
 		break;
 	case CT_CIRCLE_TYPE:
+		Report.CurveCircleCount++;
 		Curve = AddCircleCurve(CTCurveId);
 		break;
 	case CT_PARABOLA_TYPE:
+		Report.CurveParabolaCount++;
 		Curve = AddParabolaCurve(CTCurveId);
 		break;
 	case CT_HYPERBOLA_TYPE:
+		Report.CurveHyperbolaCount++;
 		Curve = AddHyperbolaCurve(CTCurveId);
 		break;
 	case CT_ELLIPSE_TYPE:
+		Report.CurveEllipseCount++;
 		Curve = AddEllipseCurve(CTCurveId);
 		break;
 	case CT_CURVE_ON_SURFACE_TYPE:
+		Report.CurveOnSurfCount++;
 		Curve = AddCurveOnSurface(CTCurveId);
 		break;
 	default:
@@ -633,13 +675,6 @@ TSharedPtr<FSurface> FCoreTechBridge::AddLinearTransfoSurface(CT_OBJECT_ID CTSur
 	TSharedPtr<FSurface> BaseSurface = AddSurface(CTBaseSurfaceId, FSurfacicBoundary());
 	TSharedPtr<FSurface> Surface = StaticCastSharedPtr<FSurface>(BaseSurface->ApplyMatrix(Matrix));
 
-#ifdef CADKERNEL_DEV 
-	const FMatrixH& BaseSurfaceMatrix = GetParamSpaceTransform(BaseSurface);
-	if (!BaseSurfaceMatrix.IsId())
-	{
-		SetParamSpaceTransorm(Surface, BaseSurfaceMatrix);
-	}
-#endif
 	return Surface;
 }
 
@@ -964,41 +999,12 @@ TSharedPtr<FCurve> FCoreTechBridge::AddCurveOnSurface(CT_OBJECT_ID CTCurveId)
 		return TSharedPtr<FCurve>();
 	}
 
-#ifdef CADKERNEL_DEV
-	// Curves on cylinders have UV param. and Z axis inversed:
-	// Z'=-Z and {UV}' = {VU}
-	if (CTSurfaceType == CT_CYLINDER_TYPE)
-	{
-		ensureCADKernel(false);
-		FMatrixH MatrixTransformZandUV;
-		MatrixTransformZandUV.SetIdentity();
-		MatrixTransformZandUV[0] = 0;
-		MatrixTransformZandUV[1] = -1;
-		MatrixTransformZandUV[4] = 1;
-		MatrixTransformZandUV[5] = 0;
-		SetParamSpaceTransorm(Surface, MatrixTransformZandUV);
-	}
-#endif
-
 	TSharedPtr<FCurve> Curve2D = AddCurve(CTParametricCurveId, CTBaseSurfaceId);
 	if (!Curve2D.IsValid())
 	{
 		return TSharedPtr<FCurve>();
 	}
 
-#ifdef CADKERNEL_DEV
-	FMatrixH Matrix = GetParamSpaceTransform(Surface);
-	if (!Matrix.IsId())
-	{
-		Curve2D = StaticCastSharedPtr<FCurve>(Curve2D->ApplyMatrix(Matrix));
-	}
-
-	if (!Curve2D.IsValid())
-	{
-		return TSharedPtr<FCurve>();
-	}
-#endif
-
 	return FEntity::MakeShared<FSurfacicCurve>(Curve2D.ToSharedRef(), Surface.ToSharedRef());
 }
 
diff --git a/Engine/Plugins/Enterprise/DatasmithCADImporter/Source/CADInterfaces/Private/CoreTechBridge.h b/Engine/Plugins/Enterprise/DatasmithCADImporter/Source/CADInterfaces/Private/CoreTechBridge.h
--- a/Engine/Plugins/Enterprise/DatasmithCADImporter/Source/CADInterfaces/Private/CoreTechBridge.h
+++ b/Engine/Plugins/Enterprise/DatasmithCADImporter/Source/CADInterfaces/Private/CoreTechBridge.h
@@ -4,9 +4,8 @@
 
 #ifdef USE_KERNEL_IO_SDK 
 
-#ifdef CADKERNEL_DEV
-#include "CADKernel/Geo/Surfaces/Surface.h"
-#endif
+#include "CADFileReport.h"
+
 #include "CADKernel/Math/MatrixH.h"
 
 #pragma warning(push)
@@ -18,7 +17,6 @@
 
 namespace CADKernel
 {
-
 class FSurfacicBoundary;
 
 class FBody;
@@ -48,8 +46,10 @@ private:
 
 	TArray<TSharedPtr<FCriterion>> Criteria;
 
+	FCADFileReport& Report;
+
 public:
-	FCoreTechBridge(FSession& InSession);
+	FCoreTechBridge(FSession& InSession, FCADFileReport& InReport);
 
 	TSharedRef<FBody> AddBody(CT_OBJECT_ID CTBodyId);
 
@@ -105,22 +105,6 @@ private:
 	{
 		AddMetadata(NodeId, (TSharedRef<FMetadataDictionary>) Entity);
 	}
-
-#ifdef CADKERNEL_DEV
-	TMap<const uint32, FMatrixH> SurfaceToMatrix;
-
-	const FMatrixH& GetParamSpaceTransform(TSharedPtr<FSurface>& Surface)
-	{
-		FMatrixH* Matrix = SurfaceToMatrix.Find(Surface->GetId());
-		return Matrix == nullptr ? FMatrixH::Identity : *Matrix;
-	}
-
-	void SetParamSpaceTransorm(TSharedPtr<FSurface>& Surface, const FMatrixH Matrix)
-	{
-		SurfaceToMatrix.Add(Surface->GetId(), Matrix);
-	}
-#endif
-
 };
 
 }
diff --git a/Engine/Plugins/Enterprise/DatasmithCADImporter/Source/CADInterfaces/Private/CoreTechFileParserImpl.cpp b/Engine/Plugins/Enterprise/DatasmithCADImporter/Source/CADInterfaces/Private/CoreTechFileParserImpl.cpp
--- a/Engine/Plugins/Enterprise/DatasmithCADImporter/Source/CADInterfaces/Private/CoreTechFileParserImpl.cpp
+++ b/Engine/Plugins/Enterprise/DatasmithCADImporter/Source/CADInterfaces/Private/CoreTechFileParserImpl.cpp
@@ -3,6 +3,7 @@
 #include "CoreTechFileParser.h"
 
 #include "CoreTechTypes.h"
+#include "CADFileReport.h"
 
 #ifdef USE_KERNEL_IO_SDK
 
@@ -721,7 +722,8 @@ namespace CADLibrary
 
 			TSharedRef<CADKernel::FModel> CADKernelModel = CADKernelSession.GetModel();
 
-			CADKernel::FCoreTechBridge CoreTechBridge(CADKernelSession);
+			CADKernel::FCADFileReport Report;
+			CADKernel::FCoreTechBridge CoreTechBridge(CADKernelSession, Report);
 
 			TSharedRef<CADKernel::FBody> CADKernelBody = CoreTechBridge.AddBody(BodyId);
 			CADKernelModel->Add(CADKernelBody);
@@ -822,7 +824,8 @@ namespace CADLibrary
 		CADKernelSession.SetFirstNewHostId(LastHostIdUsed);
 		TSharedRef<CADKernel::FModel> CADKernelModel = CADKernelSession.GetModel();
 
-		CADKernel::FCoreTechBridge CoreTechBridge(CADKernelSession);
+		CADKernel::FCADFileReport Report;
+		CADKernel::FCoreTechBridge CoreTechBridge(CADKernelSession, Report);
 
 		for(CT_OBJECT_ID BodyId : Bodies)
 		{
diff --git a/Engine/Plugins/Enterprise/DatasmithCADImporter/Source/CADInterfaces/Public/CADFileReport.h b/Engine/Plugins/Enterprise/DatasmithCADImporter/Source/CADInterfaces/Public/CADFileReport.h
new file mode 100644
--- /dev/null
+++ b/Engine/Plugins/Enterprise/DatasmithCADImporter/Source/CADInterfaces/Public/CADFileReport.h
@@ -0,0 +1,76 @@
+//CopyrightEpicGames,Inc.AllRightsReserved.
+#pragma once
+
+#include "CADKernel/Core/Chrono.h"
+#include "CADKernel/Core/Types.h"
+#include "CADKernel/UI/Message.h"
+
+namespace CADKernel
+{
+
+struct FCADFileReport
+{
+	int32 BodyCount = 0;
+	int32 ShellCount = 0;
+	int32 FaceCount = 0;
+	int32 LoopCount = 0;
+	int32 EdgeCount = 0;
+
+	int32 FailedFaceCount = 0;
+	int32 DegeneratedLoopCount = 0;
+	int32 DoubtfulLoopOrientationCount = 0;
+	int32 DegeneratedEdgeCount = 0;
+
+	int32 SurfaceCount = 0;
+	int32 DegeneratedSurfaceCount = 0;
+	int32 ConeSurfaceCount = 0;
+	int32 CylinderSurfaceCount = 0;
+	int32 LinearTransfoSurfaceCount = 0;
+	int32 NurbsSurfaceCount = 0;
+	int32 OffsetSurfaceCount = 0;
+	int32 PlaneSurfaceCount = 0;
+	int32 RevolutionSurfaceCount = 0;
+	int32 RuledSurfaceCount = 0;
+	int32 SphereSurfaceCount = 0;
+	int32 TorusSurfaceCount = 0;
+
+	int32 Blend01SurfaceCount = 0;
+	int32 Blend02SurfaceCount = 0;
+	int32 Blend03SurfaceCount = 0;
+	int32 CylindricalSurfaceCount = 0;
+	int32 PipeSurfaceCount = 0;
+	int32 ExtrusionSurfaceCount = 0;
+	int32 SurfaceFromCurvesCount = 0;
+	int32 TransformSurfaceCount = 0;
+
+	int32 SurfaceAsNurbsCount = 0;
+	int32 SurfaceNurbsCount = 0;
+
+	int32 RestrictionCurveCount = 0;
+
+	int32 CurveCount = 0;
+	int32 CurveCircleCount = 0;
+	int32 CurveCompositeCount = 0;
+	int32 CurveEllipseCount = 0;
+	int32 CurveHelixCount = 0;
+	int32 CurveHyperbolaCount = 0;
+	int32 CurveLineCount = 0;
+	int32 CurveNurbsCount = 0;
+	int32 CurveParabolaCount = 0;
+	int32 CurvePolyLineCount = 0;
+	int32 CurveBlend02BoundaryCount = 0;
+	int32 CurveEquationCount = 0;
+	int32 CurveIntersectionCount = 0;
+	int32 CurveOffsetCount = 0;
+	int32 CurveOnSurfCount = 0;
+	int32 CurveTransformCount = 0;
+
+	int32 CurveAsNurbsCount = 0;
+
+	FDuration LoadFileDuration;
+	FDuration ParsingDuration;
+
+	void Print();
+};
+
+}
diff --git a/Engine/Plugins/Enterprise/DatasmithCADImporter/Source/DatasmithWireTranslator/Private/AliasModelToCADKernelConverter.cpp b/Engine/Plugins/Enterprise/DatasmithCADImporter/Source/DatasmithWireTranslator/Private/AliasModelToCADKernelConverter.cpp
--- a/Engine/Plugins/Enterprise/DatasmithCADImporter/Source/DatasmithWireTranslator/Private/AliasModelToCADKernelConverter.cpp
+++ b/Engine/Plugins/Enterprise/DatasmithCADImporter/Source/DatasmithWireTranslator/Private/AliasModelToCADKernelConverter.cpp
@@ -222,7 +222,8 @@ TSharedPtr<CADKernel::FTopologicalFace> FAliasModelToCADKernelConverter::AddTrim
 	TSharedRef<CADKernel::FTopologicalFace> Face = CADKernel::FEntity::MakeShared<CADKernel::FTopologicalFace>(Surface);
 	Face->SetPatchId(LastFaceId++);
 
-	Face->AddLoops(Loops);
+	int32 DoubtfulLoopOrientationCount = 0;
+	Face->AddLoops(Loops, DoubtfulLoopOrientationCount);
 	return Face;
 }
 
diff --git a/Engine/Source/Runtime/Datasmith/CADKernel/Private/CADKernel/Core/System.cpp b/Engine/Source/Runtime/Datasmith/CADKernel/Private/CADKernel/Core/System.cpp
--- a/Engine/Source/Runtime/Datasmith/CADKernel/Private/CADKernel/Core/System.cpp
+++ b/Engine/Source/Runtime/Datasmith/CADKernel/Private/CADKernel/Core/System.cpp
@@ -104,34 +104,30 @@ void FSystem::DefineSpyFile(const FString& InSpyFilePath)
 
 
 #if defined(CADKERNEL_DEV) || defined(CADKERNEL_STDA)
-void FSystem::DefineQaDataFile(const FString& InQaDataFilePath)
+void FSystem::DefineReportFile(const FString& InReportFilePath)
 {
-	if (QaDataFile.IsValid())
+	if (ReportFile.IsValid())
 	{
-		QaDataFile->Close();
-		QaDataFile.Reset();
+		ReportFile->Close();
+		ReportFile.Reset();
 
-		if (QaHeaderFile.IsValid())
+		if (ReportHeaderFile.IsValid())
 		{
-			QaHeaderFile->Close();
-			QaHeaderFile.Reset();
+			ReportHeaderFile->Close();
+			ReportHeaderFile.Reset();
 		}
 	}
 
-	FString QualifPath = FPaths::GetPath(InQaDataFilePath);
-	if (!FPaths::DirectoryExists(QualifPath))
+	FString ReportHeaderPath = FPaths::GetPath(InReportFilePath);
+	if (!FPaths::DirectoryExists(ReportHeaderPath))
 	{
-		IFileManager::Get().MakeDirectory(*QualifPath, true);
+		IFileManager::Get().MakeDirectory(*ReportHeaderPath, true);
 	}
 
-	QaDataFile = MakeShareable<FArchive>(IFileManager::Get().CreateFileWriter(*InQaDataFilePath, IO_WRITE));
+	ReportFile = MakeShareable<FArchive>(IFileManager::Get().CreateFileWriter(*InReportFilePath, IO_WRITE));
 
-	QualifPath = FPaths::ConvertRelativePathToFull(FPaths::Combine(FPaths::GetPath(QualifPath), TEXT("QualifHeader.txt")));
-
-	if (!IFileManager::Get().FileExists(*QualifPath))
-	{
-		QaHeaderFile = MakeShareable<FArchive>(IFileManager::Get().CreateFileWriter(*QualifPath, IO_WRITE));
-	}
+	ReportHeaderPath = FPaths::ConvertRelativePathToFull(FPaths::Combine(ReportHeaderPath, TEXT("ReportHeader.csv")));
+	ReportHeaderFile = MakeShareable<FArchive>(IFileManager::Get().CreateFileWriter(*ReportHeaderPath, IO_WRITE));
 }
 #endif
 
diff --git a/Engine/Source/Runtime/Datasmith/CADKernel/Private/CADKernel/Core/Version.h b/Engine/Source/Runtime/Datasmith/CADKernel/Private/CADKernel/Core/Version.h
--- a/Engine/Source/Runtime/Datasmith/CADKernel/Private/CADKernel/Core/Version.h
+++ b/Engine/Source/Runtime/Datasmith/CADKernel/Private/CADKernel/Core/Version.h
@@ -2,8 +2,8 @@
 #pragma once
 
 //** ** Version du Toolkit** **//
-#define TOOLKIT_VERSION 0, 2, 18518151
-#define TOOLKIT_VERSION_ASCII "0.2.18518151"
+#define TOOLKIT_VERSION 5, 0, 18632741
+#define TOOLKIT_VERSION_ASCII "5.0.18632741"
 
 #define RELEASE_DATE	 __DATE__
 
diff --git a/Engine/Source/Runtime/Datasmith/CADKernel/Private/CADKernel/Math/SlopeUtils.h b/Engine/Source/Runtime/Datasmith/CADKernel/Private/CADKernel/Math/SlopeUtils.h
--- a/Engine/Source/Runtime/Datasmith/CADKernel/Private/CADKernel/Math/SlopeUtils.h
+++ b/Engine/Source/Runtime/Datasmith/CADKernel/Private/CADKernel/Math/SlopeUtils.h
@@ -14,7 +14,7 @@ namespace CADKernel
 
 	/**
 	 * Transform a positive slope into an oriented slope [-4, 4] i.e. an equivalent angle between [-Pi, Pi]
-	 * @return a slop between [-4, 4]
+	 * @return a slope between [-4, 4]
 	 */
 	inline double TransformIntoOrientedSlope(double Slope)
 	{
@@ -25,7 +25,7 @@ namespace CADKernel
 
 	/**
 	 * Transform a slope into a positive slope [0, *] i.e. an equivalent angle between [0, 2.Pi]
-	 * @return a slop between [0, 8]
+	 * @return a slope between [0, 8]
 	 */
 	inline double TransformIntoPositiveSlope(double Slope)
 	{
@@ -37,7 +37,7 @@ namespace CADKernel
 	 * Fast angle approximation by a "slope" :
 	 * This method is used to compute an approximation of the angle between the input segment defined by two points and [0, u) axis.
 	 * The return value is a real in the interval [0, 8] for an angle in the interval [0, 2Pi]
-	 * Warning, it's only an approximation... The conversion is not linear but the error is small near the integer value of slop (0, 1, 2, 3, ...8)
+	 * Warning, it's only an approximation... The conversion is not linear but the error is small near the integer value of slope (0, 1, 2, 3, ...8)
 	 * 
 	 * This approximation is very good when only comparison of angles is needed.
 	 * This method is not adapted to compute angle
@@ -46,7 +46,7 @@ namespace CADKernel
 	 * 
 	 * [0 - 2Pi] is divide into 8 angular sector i.e. [0, Pi/4] = [0,1], [Pi/4, Pi/2] = [1,2], ...
 	 * 
-	 * @return a slop between [0, 8] i.e. an equivalent angle between [0, 2Pi]
+	 * @return a slope between [0, 8] i.e. an equivalent angle between [0, 2Pi]
 	 *
 	 * Angle (Degree) to Slop
 	 *		  0   = 0
@@ -131,7 +131,7 @@ namespace CADKernel
 	/**
 	 * Compute the oriented slope of a segment according to a reference slope
 	 * This method is used to compute an approximation of the angle between two segments in 2D.
-	 * return a slop between [0, 8] i.e. an equivalent angle between [0, 2Pi]
+	 * return a slope between [0, 8] i.e. an equivalent angle between [0, 2Pi]
  	 */
 	inline double ComputePositiveSlope(const FPoint2D& StartPoint, const FPoint2D& EndPoint, double ReferenceSlope)
 	{
@@ -143,7 +143,7 @@ namespace CADKernel
 
 	/**
 	 * Compute the positive slope between the segments [StartPoint, EndPoint1] and [StartPoint, EndPoint2]
-	 * @return a slop between [0, 8] i.e. an equivalent angle between [0, 2Pi]
+	 * @return a slope between [0, 8] i.e. an equivalent angle between [0, 2Pi]
 	 */
 	inline double ComputePositiveSlope(const FPoint2D& StartPoint, const FPoint2D& EndPoint1, const FPoint2D& EndPoint2)
 	{
@@ -165,7 +165,7 @@ namespace CADKernel
 
 	/**
 	 * Compute the oriented slope of a segment according to a reference slope
-	 * @return a slop between [-4, 4] i.e. an equivalent angle between [-Pi, Pi]
+	 * @return a slope between [-4, 4] i.e. an equivalent angle between [-Pi, Pi]
 	 */
 	inline double ComputeOrientedSlope(const FPoint2D& StartPoint, const FPoint2D& EndPoint, double ReferenceSlope)
 	{
@@ -174,7 +174,7 @@ namespace CADKernel
 
 	/**
 	 * Compute the positive slope between the segments [StartPoint, EndPoint1] and [StartPoint, EndPoint2]
-	 * @return a slop between [-4, 4] i.e. an equivalent angle between [-Pi, Pi]
+	 * @return a slope between [-4, 4] i.e. an equivalent angle between [-Pi, Pi]
 	 */
 	inline double ComputeOrientedSlope(const FPoint2D& StartPoint, const FPoint2D& EndPoint1, const FPoint2D& EndPoint2)
 	{
@@ -182,7 +182,7 @@ namespace CADKernel
 	}
 
 	/**
-	 * return a slop between [0, 4] i.e. an angle between [0, Pi]
+	 * return a slope between [0, 4] i.e. an angle between [0, Pi]
 	 */
 	inline double ComputeUnorientedSlope(const FPoint2D& StartPoint, const FPoint2D& EndPoint, double ReferenceSlope)
 	{
@@ -190,7 +190,48 @@ namespace CADKernel
 	}
 
 	/**
-	 * return a slop between [0, 4] i.e. an angle between [0, Pi]
+	 * return a slope between [0, 1] relative to the nearest axis between horizontal or vertical axis i.e. 
+	 * ComputeUnorientedSlope => 0.5 return 0.5
+	 * ComputeUnorientedSlope => 2.3 return 0.3
+	 * ComputeUnorientedSlope => 3.6 return 0.4
+	 */
+	inline double ComputeSlopeRelativeToNearestAxis(const FPoint2D& StartPoint, const FPoint2D& EndPoint)
+	{
+		double Slope = FMath::Abs(TransformIntoOrientedSlope(ComputeSlope(StartPoint, EndPoint)));
+		if (Slope > 2)
+		{
+			Slope = 4 - Slope;
+		}
+
+		// if slope close to 2 means segment close to IsoU, otherwise segment close to IsoV
+		// Wants a slope between 0 and 1 to manage either IsoU and IsoV
+		// Close to 0 means close to IsoU or IsoV
+		if (Slope > 1)
+		{
+			Slope = 2 - Slope;
+		}
+
+		return Slope;
+	}
+
+	/**
+	 * return a slope between [0, 2] relative to reference Axis i.e.
+	 * ComputeUnorientedSlope => 0.5 return 0.5
+	 * ComputeUnorientedSlope => 2.3 return 1.7
+	 */
+	inline double ComputeSlopeRelativeToReferenceAxis(const FPoint2D& StartPoint, const FPoint2D& EndPoint, double ReferenceAxisSlope)
+	{
+		double Slope = ComputeUnorientedSlope(StartPoint, EndPoint, ReferenceAxisSlope);
+		if (Slope > 2)
+		{
+			Slope = 4 - Slope;
+		}
+
+		return Slope;
+	}
+
+	/**
+	 * return a slope between [0, 4] i.e. an angle between [0, Pi]
 	 */
 	inline double ComputeUnorientedSlope(const FPoint2D& StartPoint, const FPoint2D& EndPoint1, const FPoint2D& EndPoint2)
 	{
@@ -259,51 +300,51 @@ namespace CADKernel
 			});
 	}
 
-	inline FPoint2D SlopeToVector(const double Slop)
+	inline FPoint2D SlopeToVector(const double Slope)
 	{
-		int32 SlopStep = (int32)(Slop);
+		int32 SlopeStep = (int32) (Slope);
 		FPoint2D Vector;
-		switch (SlopStep)
+		switch (SlopeStep)
 		{
 		case 0:
 			// Delta = DeltaV / DeltaU;
 			Vector[0] = 1.;
-			Vector[1] = Slop;
+			Vector[1] = Slope;
 			break;
 		case 1:
 			// 2 - DeltaU / DeltaV;
-			Vector[0] = 2. - Slop;
+			Vector[0] = 2. - Slope;
 			Vector[1] = 1.;
 			break;
 		case 2:
 			// 2 - DeltaU / DeltaV;
-			Vector[0] = 2. - Slop;
+			Vector[0] = 2. - Slope;
 			Vector[1] = 1.;
 			break;
 		case 3:
 			// 4 + DeltaV / DeltaU;
 			Vector[0] = -1.;
-			Vector[1] = 4. - Slop;
+			Vector[1] = 4. - Slope;
 			break;
 		case 4:
 			// 4 + DeltaV / DeltaU;
 			Vector[0] = -1.;
-			Vector[1] = 4. - Slop;
+			Vector[1] = 4. - Slope;
 			break;
 		case 5:
 			// 6 - DeltaU / DeltaV;
-			Vector[0] = Slop - 6.;
+			Vector[0] = Slope - 6.;
 			Vector[1] = -1.;
 			break;
 		case 6:
 			// 6 - DeltaU / DeltaV // deltaU/deltaV <0
-			Vector[0] = Slop - 6.;
+			Vector[0] = Slope - 6.;
 			Vector[1] = -1;
 			break;
 		case 7:
 			// 8 + DeltaV / DeltaU; // deltaU/deltaV <0
 			Vector[0] = 1.;
-			Vector[1] = Slop - 8.;
+			Vector[1] = Slope - 8.;
 			break;
 		default:
 			break;
diff --git a/Engine/Source/Runtime/Datasmith/CADKernel/Private/CADKernel/Mesh/Meshers/IsoTriangulator.cpp b/Engine/Source/Runtime/Datasmith/CADKernel/Private/CADKernel/Mesh/Meshers/IsoTriangulator.cpp
--- a/Engine/Source/Runtime/Datasmith/CADKernel/Private/CADKernel/Mesh/Meshers/IsoTriangulator.cpp
+++ b/Engine/Source/Runtime/Datasmith/CADKernel/Private/CADKernel/Mesh/Meshers/IsoTriangulator.cpp
@@ -51,7 +51,7 @@ FIsoTriangulator::FIsoTriangulator(FGrid& InGrid, TSharedRef<FFaceMesh> EntityMe
 	}
 
 #ifdef DEBUG_BOWYERWATSON
-bool FBowyerWatsonTriangulator::bDisplay = false;
+	bool FBowyerWatsonTriangulator::bDisplay = false;
 #endif
 
 #endif
@@ -86,6 +86,7 @@ bool FIsoTriangulator::Triangulate()
 	FLoopCleaner LoopCleaner(*this);
 	if (!LoopCleaner.Run())
 	{
+		MesherReport.Logs.AddDegeneratedLoop();
 		FMessage::Printf(EVerboseLevel::Log, TEXT("The meshing of the surface %d failed due to a degenerated loop\n"), Grid.GetFace()->GetId());
 		return false;
 	}
@@ -150,75 +151,6 @@ bool FIsoTriangulator::Triangulate()
 	Grid.DisplayIsoSegments(TEXT("FIsoTrianguler::Final Iso ToLink Inner To Loops"), DisplaySpace, FinalToLoops, false, false, EVisuProperty::YellowCurve);
 #endif
 
-
-	/*
-		return true;
-
-		// =============================================================================================================
-		// 	   Find first segments of the final mesh
-		// =============================================================================================================
-
-		//FindCandidateSegmentsToLinkInnerAndLoop();
-
-		SortLoopNodes();
-		// use of SortedLoopNodes but specific order
-		FindIsoSegmentToLinkLoopToLoop();
-	#ifdef CADKERNEL_DEV
-		Display(DisplaySpace, TEXT("FIsoTrianguler::Final Iso ToLink Loop To Loops"), FinalToLoops, false);
-		//Wait(bDisplay);
-	#endif
-
-	#ifdef NEED_TO_CHECK_USEFULNESS
-		CompleteIsoSegmentLoopToLoop();
-		Display(DisplaySpace, TEXT("FIsoTrianguler::Final Iso ToLink Loop To Loops with Complete"), FinalToLoops, false);
-	#endif
-
-		FindIsoSegmentToLinkInnerToLoop();
-	#ifdef CADKERNEL_DEV
-		Display(DisplaySpace, TEXT("FIsoTrianguler::Final Iso ToLink Inner To Loops"), FinalToLoops, false);
-	#endif
-
-		// =============================================================================================================
-		// 	   Find candidate segments for the  mesh and select the best set
-		// =============================================================================================================
-
-		// Delaunay
-		if (Grid.GetLoops3D().Num())
-		{
-			//ConnectCellLoopsByNeighborhood();
-	#ifdef CADKERNEL_DEV
-			Display(DisplaySpace, TEXT("FIsoTrianguler::Candidate Segments by Delaunay"), CandidateLoopToLoopSegments, false);
-	#endif
-		}
-		//Wait(bDisplay);
-
-	#ifdef CADKERNEL_DEV
-		Chronos.TriangulateDuration3 = FChrono::Elapse(StartTime);
-	#endif
-
-		// use of SortedloopNodes
-		FindCandidateSegmentsToLinkInnerToLoop();
-	#ifdef CADKERNEL_DEV
-		Display(DisplaySpace, TEXT("FIsoTrianguler::Candidate Segments ToLink Inner To Loops"), CandidateInnerToLoopSegments, false);
-	#endif
-
-		SelectSegmentInCandidateSegments();
-	#ifdef CADKERNEL_DEV
-		Display(DisplaySpace, TEXT("FIsoTrianguler::Final Segment To Loop"), FinalToLoops, false);
-		Display(DisplaySpace, TEXT("FIsoTrianguler::Inner Segments Limit"), FinalInnerSegments, false);
-		//Wait(bDisplay);
-	#endif
-
-		// Complete the final segments with segment connecting unconnected Inner point or inner segment
-		ConnectUnconnectedInnerSegments();
-	#ifdef CADKERNEL_DEV
-		Display(DisplaySpace, TEXT("FIsoTrianguler::Final Segment To Loop"), FinalToLoops, false);
-	#endif
-
-	#ifdef CADKERNEL_DEV
-		Chronos.TriangulateDuration4 = FChrono::Elapse(StartTime);
-	#endif
-	*/
 	// =============================================================================================================
 	// Make the final tessellation 
 	// =============================================================================================================
@@ -508,6 +440,11 @@ void FIsoTriangulator::BuildInnerSegments()
 {
 
 #ifdef DEBUG_BUILDINNERSEGMENTS
+	if (bDisplay)
+	{
+		Grid.DisplayInnerPoints(TEXT("BuildInnerSegments::Points"), EGridSpace::UniformScaled);
+	}
+
 	F3DDebugSession _(bDisplay, TEXT("BuildInnerSegments"));
 #endif
 
@@ -520,24 +457,36 @@ void FIsoTriangulator::BuildInnerSegments()
 
 	LoopSegmentsIntersectionTool.Reserve(InnerSegmentsIntersectionTool.Count());
 
-	TFunction<bool(const FPoint2D&, const FPoint2D&)> AlmostHitsLoop = [&](const FPoint2D& Node1, const FPoint2D& Node2) -> bool
+	TFunction<bool(const FPoint2D&, const FPoint2D&, const ESegmentType)> AlmostHitsLoop = [&](const FPoint2D& Node1, const FPoint2D& Node2, const ESegmentType InType) -> bool
 	{
-		for (const TArray<FPoint2D>& Loop : Grid.GetLoops2D(EGridSpace::UniformScaled))
+		if (InType == ESegmentType::IsoU)
 		{
-			for (const FPoint2D& LoopPoint : Loop)
+			for (const TArray<FPoint2D>& Loop : Grid.GetLoops2D(EGridSpace::UniformScaled))
 			{
-				if (FMath::IsNearlyEqual(LoopPoint.U, Node1.U, GeometricTolerance))
+				for (const FPoint2D& LoopPoint : Loop)
 				{
-					if (Node1.V - SMALL_NUMBER < LoopPoint.V && LoopPoint.V < Node2.V + SMALL_NUMBER)
+					if (FMath::IsNearlyEqual(LoopPoint.V, Node1.V, GeometricTolerance))
 					{
-						return true;
+						if (Node1.U - SMALL_NUMBER < LoopPoint.U && LoopPoint.U < Node2.U + SMALL_NUMBER)
+						{
+							return true;
+						}
 					}
 				}
-				else if (FMath::IsNearlyEqual(LoopPoint.V, Node1.V, GeometricTolerance))
+			}
+		}
+		else
+		{
+			for (const TArray<FPoint2D>& Loop : Grid.GetLoops2D(EGridSpace::UniformScaled))
+			{
+				for (const FPoint2D& LoopPoint : Loop)
 				{
-					if (Node1.U - SMALL_NUMBER < LoopPoint.U && LoopPoint.U < Node2.U + SMALL_NUMBER)
+					if (FMath::IsNearlyEqual(LoopPoint.U, Node1.U, GeometricTolerance))
 					{
-						return true;
+						if (Node1.V - SMALL_NUMBER < LoopPoint.V && LoopPoint.V < Node2.V + SMALL_NUMBER)
+						{
+							return true;
+						}
 					}
 				}
 			}
@@ -549,26 +498,49 @@ void FIsoTriangulator::BuildInnerSegments()
 	{
 		if (!Grid.IsNodeInsideFace(IndexNode1) || !Grid.IsNodeInsideFace(IndexNode2))
 		{
+#ifdef DEBUG_BUILDINNERSEGMENTS
+			if (bDisplay)
+			{
+				DisplaySegment(Grid.GetInner2DPoint(EGridSpace::UniformScaled, IndexNode1) * Grid.DisplayScale, Grid.GetInner2DPoint(EGridSpace::UniformScaled, IndexNode2) * Grid.DisplayScale, IndexNode1, EVisuProperty::BlueCurve);
+			}
+#endif
 			InnerToOuterSegmentsIntersectionTool.AddSegment(Grid.GetInner2DPoint(EGridSpace::UniformScaled, IndexNode1), Grid.GetInner2DPoint(EGridSpace::UniformScaled, IndexNode2));
 			return;
 		}
 
 		if (Grid.IsNodeCloseToLoop(IndexNode1) && Grid.IsNodeCloseToLoop(IndexNode2))
 		{
-#ifdef DEBUG_BUILDINNERSEGMENTS
-			if (bDisplay)
-			{
-				DisplaySegment(Grid.GetInner2DPoint(EGridSpace::UniformScaled, IndexNode1), Grid.GetInner2DPoint(EGridSpace::UniformScaled, IndexNode2));
-				Wait();
-			}
-#endif
 
 			if (LoopSegmentsIntersectionTool.DoesIntersect(Grid.GetInner2DPoint(EGridSpace::UniformScaled, IndexNode1), Grid.GetInner2DPoint(EGridSpace::UniformScaled, IndexNode2))
-				|| AlmostHitsLoop(Grid.GetInner2DPoint(EGridSpace::UniformScaled, IndexNode1), Grid.GetInner2DPoint(EGridSpace::UniformScaled, IndexNode2)))
+				|| AlmostHitsLoop(Grid.GetInner2DPoint(EGridSpace::UniformScaled, IndexNode1), Grid.GetInner2DPoint(EGridSpace::UniformScaled, IndexNode2), InType))
 			{
 				InnerToOuterSegmentsIntersectionTool.AddSegment(Grid.GetInner2DPoint(EGridSpace::UniformScaled, IndexNode1), Grid.GetInner2DPoint(EGridSpace::UniformScaled, IndexNode2));
+#ifdef DEBUG_BUILDINNERSEGMENTS
+				if (bDisplay)
+				{
+					DisplaySegment(Grid.GetInner2DPoint(EGridSpace::UniformScaled, IndexNode1) * Grid.DisplayScale, Grid.GetInner2DPoint(EGridSpace::UniformScaled, IndexNode2) * Grid.DisplayScale, IndexNode1, EVisuProperty::YellowCurve);
+				}
+#endif
 				return;
 			}
+			else
+			{
+#ifdef DEBUG_BUILDINNERSEGMENTS
+				if (bDisplay)
+				{
+					DisplaySegment(Grid.GetInner2DPoint(EGridSpace::UniformScaled, IndexNode1) * Grid.DisplayScale, Grid.GetInner2DPoint(EGridSpace::UniformScaled, IndexNode2) * Grid.DisplayScale, IndexNode1, EVisuProperty::GreenCurve);
+				}
+#endif
+			}
+		}
+		else
+		{
+#ifdef DEBUG_BUILDINNERSEGMENTS
+			if (bDisplay)
+			{
+				DisplaySegment(Grid.GetInner2DPoint(EGridSpace::UniformScaled, IndexNode1) * Grid.DisplayScale, Grid.GetInner2DPoint(EGridSpace::UniformScaled, IndexNode2) * Grid.DisplayScale, IndexNode1, EVisuProperty::RedCurve);
+			}
+#endif
 		}
 
 		FIsoInnerNode& Node1 = *GlobalIndexToIsoInnerNodes[IndexNode1];
@@ -635,6 +607,43 @@ void FIsoTriangulator::BuildInnerSegmentsIntersectionTool()
 	TArray<uint8> Pixel;
 	Pixel.Init(0, Grid.GetTotalCuttingCount());
 
+#ifdef DEBUG_FINDBOUNDARYSEGMENTS
+	if (bDisplay)
+	{
+		F3DDebugSession _(TEXT("FIsoTrianguler::InnerSegment"));
+		for (int32 IndexV = 0, Index = 0; IndexV < NumV; ++IndexV)
+		{
+			for (int32 IndexU = 0; IndexU < NumU; ++IndexU, ++Index)
+			{
+				if (!Grid.IsNodeInsideFace(Index))
+				{
+					continue;
+				}
+
+				if (GlobalIndexToIsoInnerNodes[Index]->IsLinkedToNextU())
+				{
+					Grid.DisplayIsoSegment(EGridSpace::UniformScaled, *GlobalIndexToIsoInnerNodes[Index], *GlobalIndexToIsoInnerNodes[Index + 1], Index, EVisuProperty::YellowCurve);
+				}
+
+				if (GlobalIndexToIsoInnerNodes[Index]->IsLinkedToPreviousU())
+				{
+					Grid.DisplayIsoSegment(EGridSpace::UniformScaled, *GlobalIndexToIsoInnerNodes[Index], *GlobalIndexToIsoInnerNodes[Index - 1], Index, EVisuProperty::YellowCurve);
+				}
+
+				if (GlobalIndexToIsoInnerNodes[Index]->IsLinkedToNextV())
+				{
+					Grid.DisplayIsoSegment(EGridSpace::UniformScaled, *GlobalIndexToIsoInnerNodes[Index], *GlobalIndexToIsoInnerNodes[Index + NumU], Index, EVisuProperty::YellowCurve);
+				}
+
+				if (GlobalIndexToIsoInnerNodes[Index]->IsLinkedToPreviousV())
+				{
+					Grid.DisplayIsoSegment(EGridSpace::UniformScaled, *GlobalIndexToIsoInnerNodes[Index], *GlobalIndexToIsoInnerNodes[Index - NumU], Index, EVisuProperty::YellowCurve);
+				}
+			}
+		}
+	}
+#endif
+
 	// A pixel is True if one of its boundary segment does not exist 
 	for (int32 IndexV = 0, Index = 0; IndexV < NumV; ++IndexV)
 	{
@@ -705,8 +714,9 @@ void FIsoTriangulator::BuildInnerSegmentsIntersectionTool()
 	Chronos.FindLoopSegmentOfInnerTriangulationDuration += FChrono::Elapse(StartTime);
 #endif
 
-#ifdef DEBUG_FINDBOUNDARYSEGMENTS 
+#ifdef DEBUG_FINDBOUNDARYSEGMENTS
 	DisplayPixels(Pixel);
+	//Wait(bDisplay);
 #endif
 }
 
@@ -940,9 +950,6 @@ void FIsoTriangulator::FindIsoSegmentToLinkLoopToLoop()
 
 void FIsoTriangulator::FindSegmentToLinkOuterLoopNodes(FCell& Cell)
 {
-	TArray<FLoopNode*>& OuterLoop = Cell.SubLoops[0];
-	int32 NodeCount = OuterLoop.Num();
-
 	int SubdivisionCount = Cell.OuterLoopSubdivision.Num();
 	for (int32 Andex = 0; Andex < SubdivisionCount - 1; ++Andex)
 	{
@@ -956,227 +963,192 @@ void FIsoTriangulator::FindSegmentToLinkOuterLoopNodes(FCell& Cell)
 	Cell.SelectSegmentInCandidateSegments(IsoSegmentFactory);
 }
 
-void FIsoTriangulator::AddSementToLinkOuterLoopExtremities(FCell& Cell)
-{
-	//F3DDebugSession _(TEXT("AddSementToLinkOuterLoopExtremities"));
-
-	if (!Cell.bHasOuterLoop)
-	{
-		return;
-	}
-	FLoopNode* NodeA = Cell.SubLoops[0][0];
-	const FPoint2D& ACoordinates = NodeA->Get2DPoint(EGridSpace::UniformScaled, Grid);
-	FLoopNode* NodeB = Cell.SubLoops[0].Last();
-	const FPoint2D& BCoordinates = NodeB->Get2DPoint(EGridSpace::UniformScaled, Grid);
-	TryToCreateSegment(Cell, NodeA, ACoordinates, NodeB, BCoordinates, 0.1);
-	Cell.SelectSegmentInCandidateSegments(IsoSegmentFactory);
-}
-
-//#define DEBUG_FIND_ISO_SEGMENT_TO_LINK_OUTER_LOOP_NODES
-void FIsoTriangulator::FindIsoSegmentToLinkOuterLoopNodes(FCell& Cell)
+void FIsoTriangulator::FindSegmentToLinkOuterToInnerLoopNodes(FCell& Cell)
 {
-#ifdef DEBUG_FIND_ISO_SEGMENT_TO_LINK_OUTER_LOOP_NODES
-	F3DDebugSession _(TEXT("FindIsoSegmentToLinkOuterLoopNodes"));
-#endif
 	TArray<FLoopNode*>& OuterLoop = Cell.SubLoops[0];
-	int32 NodeCount = OuterLoop.Num();
 
-	int SubdivisionCount = Cell.OuterLoopSubdivision.Num();
-	for (int32 Andex = 0; Andex < SubdivisionCount - 1; ++Andex)
+	TFunction<void(FLoopNode&, FLoopNode&)> CreateCandidateSegment = [&](FLoopNode& Node, FLoopNode& CandidatNode)
 	{
-		TArray<FLoopNode*>& SubLoopA = Cell.OuterLoopSubdivision[Andex];
-		for (int32 Bndex = Andex + 1; Bndex < SubdivisionCount; ++Bndex)
+		FIsoSegment& Segment = IsoSegmentFactory.New();
+		Segment.Init(Node, CandidatNode, ESegmentType::LoopToLoop);
+		Segment.SetCandidate();
+		Cell.CandidateSegments.Add(&Segment);
+
+#ifdef DEBUG_FIND_SEGMENT_TO_LINK_OUTER_TO_INNER_LOOP
+		if (bDisplay)
 		{
-			TArray<FLoopNode*>& SubLoopB = Cell.OuterLoopSubdivision[Bndex];
-			TryToConnectTwoSubLoopsWithTheMostIsoSegment(Cell, SubLoopA, SubLoopB);
+			{
+				F3DDebugSession _(FString::Printf(TEXT("Candidate Segment")));
+				Grid.DisplayIsoSegment(EGridSpace::UniformScaled, Node, CandidatNode, 0, EVisuProperty::RedCurve);
+			}
 		}
-	}
-
-	for (int32 Index = 0; Index < SubdivisionCount; ++Index)
-	{
-		TArray<FLoopNode*>& SubLoop = Cell.OuterLoopSubdivision[Index];
-		TryToConnectVertexSubLoopWithTheMostIsoSegment(Cell, SubLoop);
-	}
-	Cell.SelectSegmentInCandidateSegments(IsoSegmentFactory);
-}
+#endif 
+	};
 
-#ifdef UNUSED_TO_DELETE_
-void FIsoTriangulator::LastChanceToCreateSegmentInCell(FCell& Cell)
-{
-	if (Cell.LoopNodeCount > 5 && Cell.CandidateSegments.Num() > Cell.SubLoops.Num())
+	const SlopMethod GetSlopAtStartNode = CounterClockwiseSlop;
+	const SlopMethod GetSlopAtEndNode = ClockwiseSlop;
+	TFunction<void(FIsoSegment*)> FindBestTriangle = [&](FIsoSegment* Segment)
 	{
-		return;
-	}
+		FLoopNode& StartNode = (FLoopNode&)Segment->GetFirstNode();
+		FLoopNode& EndNode = (FLoopNode&)Segment->GetSecondNode();
 
-	//F3DDebugSession _(TEXT("LastChanceToCreateSegmentInCell"));
+		if ((StartNode.GetConnectedSegments().Num() > 2) || (EndNode.GetConnectedSegments().Num() > 2))
+		{
+			return;
+		}
 
-	TArray<FLoopNode*> Nodes;
-	Nodes.Reserve(Cell.LoopNodeCount);
+		FIsoSegment& PreviousSegment = StartNode.GetPreviousSegment();
+		FIsoSegment& NextSegment = EndNode.GetNextSegment();
 
-	for (const TArray<FLoopNode*>& Loop : Cell.SubLoops)
-	{
-		Nodes.Append(Loop);
-	}
+		FIsoNode& PreviousNode = StartNode.GetPreviousNode();
+		FIsoNode& NextNode = EndNode.GetNextNode();
 
-	for (int32 Index = 0; Index < Cell.LoopNodeCount - 1; ++Index)
-	{
-		FLoopNode* StartNode = Nodes[Index];
-		const FPoint2D& StartPoint = StartNode->Get2DPoint(EGridSpace::UniformScaled, Grid);
-		for (int32 Endex = Index + 1; Endex < Cell.LoopNodeCount; ++Endex)
+#ifdef DEBUG_FIND_SEGMENT_TO_LINK_OUTER_TO_INNER_LOOP
+		static int32 SegmentIndex = 0;
+		if (bDisplay)
 		{
-			FLoopNode* EndNode = Nodes[Endex];
-			if (&StartNode->GetNextNode() == EndNode || &StartNode->GetPreviousNode() == EndNode)
 			{
-				continue;
+				SegmentIndex++;
+				F3DDebugSession _(FString::Printf(TEXT("Start Segment %d"), SegmentIndex));
+				Grid.DisplayIsoSegment(EGridSpace::UniformScaled, *Segment, SegmentIndex, EVisuProperty::BlueCurve);
+				Grid.DisplayIsoSegment(EGridSpace::UniformScaled, PreviousSegment, SegmentIndex, EVisuProperty::GreenCurve);
+				Grid.DisplayIsoSegment(EGridSpace::UniformScaled, NextSegment, SegmentIndex, EVisuProperty::GreenCurve);
 			}
-			const FPoint2D& EndPoint = EndNode->Get2DPoint(EGridSpace::UniformScaled, Grid);
-			TryToCreateSegment(Cell, StartNode, StartPoint, EndNode, EndPoint, 1);
 		}
-	}
-}
+#endif
 
-void FIsoTriangulator::FindIsoSegmentToLinkInnerToLoop()
-{
-	FTimePoint StartTime = FChrono::Now();
+		const FPoint2D& StartPoint2D = StartNode.Get2DPoint(EGridSpace::UniformScaled, Grid);
+		const FPoint2D& EndPoint2D = EndNode.Get2DPoint(EGridSpace::UniformScaled, Grid);
 
-	int32 IndexU = 0;
-	int32 IndexV = 0;
+		const FPoint& StartPoint3D = StartNode.Get3DPoint(Grid);
+		const FPoint& EndPoint3D = EndNode.Get3DPoint(Grid);
 
-	const TArray<double>& IsoUCoordinates = Grid.GetUniformCuttingCoordinatesAlongIso(EIso::IsoU);
-	const TArray<double>& IsoVCoordinates = Grid.GetUniformCuttingCoordinatesAlongIso(EIso::IsoV);
+		// StartMaxSlope and EndMaxSlope are at most equal to 4, because if the slop with candidate node is biggest to 4, the nez triangle will be inverted
+		double StartReferenceSlope = ComputePositiveSlope(StartPoint2D, EndPoint2D, 0);
 
-	int32 InitNum = InnerSegmentsIntersectionTool.Count() + LoopSegmentsIntersectionTool.Count();
-	FinalToLoops.Reserve(InitNum);
-	InnerToLoopSegmentsIntersectionTool.Reserve(InitNum);
+		double StartMaxSlope = GetSlopAtStartNode(StartPoint2D, PreviousNode.Get2DPoint(EGridSpace::UniformScaled, Grid), StartReferenceSlope);
+		StartMaxSlope /= 2.;
+		StartMaxSlope = FMath::Min(StartMaxSlope, 4.);
 
-	TArray<double> IsoUTolerances;
-	TArray<double> IsoVTolerances;
+		double EndReferenceSlope = StartReferenceSlope < 4 ? StartReferenceSlope + 4 : StartReferenceSlope - 4;
+		double EndMaxSlope = GetSlopAtEndNode(EndPoint2D, NextNode.Get2DPoint(EGridSpace::UniformScaled, Grid), EndReferenceSlope);
+		EndMaxSlope /= 2.;
+		EndMaxSlope = FMath::Min(EndMaxSlope, 4.);
 
-	// As the gap between two IsoPoints is not constant, the tolerance of closeness is define as 1/10 of the coordinate variation between two consecutive iso coordiantes 
-	TFunction<void(const TArray<double>&, TArray<double>&)> SetIsoTolerance = [](const TArray<double>& Isos, TArray<double>& IsoTolerances)
-	{
-		IsoTolerances.SetNum(Isos.Num() - 1);
-		for (int32 Index = 0; Index < Isos.Num() - 1; ++Index)
-		{
-			IsoTolerances[Index] = (Isos[Index + 1] - Isos[Index]) / 10.;
-		}
-	};
-	SetIsoTolerance(IsoUCoordinates, IsoUTolerances);
-	SetIsoTolerance(IsoVCoordinates, IsoVTolerances);
+		const double MinSlopToNotBeAligned = 0.01;
 
-	// We check if a loop point is aligned with the inner grid i.e. U or V points coordinate is closed to one Iso coordinates. The "closed" criteria is define by the local tolerance compute by SetIsoTolerance.
-	// OutIndex Value defined the segment in which the point is localized i.e. PointCoordinate is in [Iso[Index], Iso[Index+1]
-	// is OutPointIndex != 0, OutPointIndex = Index if PointCoordinate is closed to Iso[OutIndex] or OutPointIndex = OutIndex + 1 if PointCoordinate is closed to Iso[OutIndex + 1]
-	TFunction<void(const TArray<double>&, const TArray<double>&, const double&, int32&, int32&)> GetIsoIndexOfPoint = [](const TArray<double>& Iso, const TArray<double>& IsoTolerance, const double& PointCoord, int32& OutIndex, int32& OutPointIndex)
-	{
-		while (PointCoord < Iso[OutIndex])
-		{
-			OutIndex--;
-		}
+		FLoopNode* CandidatNode = nullptr;
+		double MinCriteria = HUGE_VALUE;
+		double CandidateSlopeAtStartNode = 8.;
+		double CandidateSlopeAtEndNode = 8.;
 
-		for (; OutIndex < Iso.Num() - 2; ++OutIndex)
+		for (int32 Index = 1; Index < Cell.SubLoops.Num(); ++Index)
 		{
-			if (Iso[OutIndex] <= PointCoord && PointCoord <= Iso[OutIndex + 1])
+			for (FLoopNode* Node : Cell.SubLoops[Index])
 			{
-				break;
+				// Check if the node is inside the sector (X) or outside (Z)
+				const FPoint2D& NodePoint2D = Node->Get2DPoint(EGridSpace::UniformScaled, Grid);
+				double SlopeAtStartNode = GetSlopAtStartNode(StartPoint2D, NodePoint2D, StartReferenceSlope);
+				double SlopeAtEndNode = GetSlopAtEndNode(EndPoint2D, NodePoint2D, EndReferenceSlope);
+
+				if (SlopeAtStartNode <= 0 || SlopeAtStartNode >= StartMaxSlope)
+				{
+					continue;
+				}
+
+				if (SlopeAtEndNode <= 0 || SlopeAtEndNode >= EndMaxSlope)
+				{
+					continue;
+				}
+
+				double PointCriteria = IsoTriangulatorImpl::IsoscelesCriteria(StartPoint2D, EndPoint2D, NodePoint2D);
+
+				if (
+					// the candidate triangle is inside the current candidate triangle
+					((SlopeAtStartNode < (CandidateSlopeAtStartNode + MinSlopToNotBeAligned)) && (SlopeAtEndNode < (CandidateSlopeAtEndNode + MinSlopToNotBeAligned))) ||
+					// or the candidate triangle is better the current candidate triangle and doesn't contain the current candidate triangle
+					((PointCriteria < MinCriteria) && ((SlopeAtStartNode > CandidateSlopeAtStartNode) ^ (SlopeAtEndNode > CandidateSlopeAtEndNode))))
+				{
+
+					if (LoopSegmentsIntersectionTool.DoesIntersect(StartNode, *Node))
+					{
+						continue;
+					}
+
+					if (LoopSegmentsIntersectionTool.DoesIntersect(EndNode, *Node))
+					{
+						continue;
+					}
+
+					MinCriteria = PointCriteria;
+					CandidatNode = Node;
+					CandidateSlopeAtStartNode = SlopeAtStartNode;
+					CandidateSlopeAtEndNode = SlopeAtEndNode;
+				}
 			}
 		}
 
-		if (PointCoord < Iso[OutIndex] + IsoTolerance[OutIndex])
-		{
-			OutPointIndex = OutIndex;
-		}
-		else if (PointCoord > Iso[OutIndex + 1] - IsoTolerance[OutIndex])
+		if (CandidatNode)
 		{
-			OutPointIndex = OutIndex + 1;
+			CreateCandidateSegment(StartNode, *CandidatNode);
+			CreateCandidateSegment(EndNode, *CandidatNode);
 		}
 	};
 
-	TFunction<const int32(FLoopNode&, const FPoint2D&, const int32, const int32, ESegmentType)> TryToBuildIsoSegment = [&](FLoopNode& IsoLoopNode, const FPoint2D& Point, int32 GlobalIndex1, int32 GlobalIndex2, ESegmentType Type)
+	for (int32 Index = 0; Index < OuterLoop.Num() - 1; ++Index)
 	{
-		int32 IsoGlobalIndex = -1;
-		if (GlobalIndexToIsoInnerNodes[GlobalIndex1] != nullptr && GlobalIndexToIsoInnerNodes[GlobalIndex2] == nullptr)
-		{
-			IsoGlobalIndex = GlobalIndex1;
-		}
-		else if (GlobalIndexToIsoInnerNodes[GlobalIndex1] == nullptr && GlobalIndexToIsoInnerNodes[GlobalIndex2] != nullptr)
-		{
-			IsoGlobalIndex = GlobalIndex2;
-		}
-		if (IsoGlobalIndex < 0)
+		FLoopNode* Node = OuterLoop[Index];
+		if (&Node->GetNextNode() != OuterLoop[Index + 1])
 		{
-			return IsoGlobalIndex;
+			continue;
 		}
 
-		if (LoopSegmentsIntersectionTool.DoesIntersect(IsoLoopNode, Grid.GetInner2DPoint(EGridSpace::UniformScaled, IsoGlobalIndex)))
+		FIsoSegment* Segment = Node->GetSegmentConnectedTo(OuterLoop[Index + 1]);
+		if (Segment == nullptr)
 		{
-			return IsoGlobalIndex;
+			continue;
 		}
+		FindBestTriangle(Segment);
+	}
 
-		FIsoInnerNode* Node = GlobalIndexToIsoInnerNodes[IsoGlobalIndex];
-		FIsoSegment& Segment = IsoSegmentFactory.New();
-		Segment.Init(*Node, IsoLoopNode, Type);
-		Segment.ConnectToNode();
-		FinalToLoops.Add(&Segment);
-		InnerToLoopSegmentsIntersectionTool.AddSegment(Segment);
-		return IsoGlobalIndex;
-	};
+	Cell.SelectSegmentInCandidateSegments(IsoSegmentFactory);
+}
 
-	const TArray<TArray<FPoint2D>>& Loops = Grid.GetLoops2D(EGridSpace::UniformScaled);
+//#define DEBUG_FIND_ISO_SEGMENT_TO_LINK_OUTER_LOOP_NODES
+void FIsoTriangulator::FindIsoSegmentToLinkOuterLoopNodes(FCell& Cell)
+{
+	TArray<FLoopNode*>& OuterLoop = Cell.SubLoops[0];
+	int32 NodeCount = OuterLoop.Num();
 
-	// Segment smaller than the size of the grid is build
-	for (FLoopNode& LoopNode : LoopNodes)
 	{
-		const FPoint2D& Point = Loops[LoopNode.GetLoopIndex()][LoopNode.GetIndex()];
-		int32 PointIndexU = -1;
-		int32 PointIndexV = -1;
-
-		GetIsoIndexOfPoint(IsoUCoordinates, IsoUTolerances, Point.U, IndexU, PointIndexU);
-		GetIsoIndexOfPoint(IsoVCoordinates, IsoVTolerances, Point.V, IndexV, PointIndexV);
-
-		// Point is in the square [IndexU, IndexV], ... [IndexU+1, IndexV+1]]
-		// If PointIndexU (resp PointIndexV) != -1, Point is closed to IsoU [PointIndexU]
-		// If GlobalIndexToIsoInnerNodes[IndexU+1, IndexV] == null => it is outside the domaine
-		// So if GlobalIndexToIsoInnerNodes[IndexU, IndexV] != null, and PointIndexV = IndexV, the segment [InnerNodes[IndexU, IndexV], LoopNode] could be build
-		if (PointIndexU == IndexU || PointIndexU == IndexU + 1)
-		{
-			const int32 GlobalIndex1 = Grid.GobalIndex(PointIndexU, IndexV);
-			const int32 GlobalIndex2 = GlobalIndex1 + Grid.GetCuttingCount(EIso::IsoU);
-			const int32 IsoGlobalIndex = TryToBuildIsoSegment(LoopNode, Point, GlobalIndex1, GlobalIndex2, ESegmentType::InnerToLoopV);
-			if (IsoGlobalIndex == GlobalIndex1)
-			{
-				GlobalIndexToIsoInnerNodes[IsoGlobalIndex]->SetLinkedToIso(EIsoLink::IsoUNext);
-			}
-			else if (IsoGlobalIndex == GlobalIndex2)
-			{
-				GlobalIndexToIsoInnerNodes[IsoGlobalIndex]->SetLinkedToIso(EIsoLink::IsoUPrevious);
-			}
-		}
-
-		if (PointIndexV == IndexV || PointIndexV == IndexV + 1)
+#ifdef DEBUG_FIND_ISO_SEGMENT_TO_LINK_OUTER_LOOP_NODES
+		F3DDebugSession _(bDisplay, ("FindIsoSegmentToLinkOuterLoopNodes 1"));
+#endif
+		int SubdivisionCount = Cell.OuterLoopSubdivision.Num();
+		for (int32 Andex = 0; Andex < SubdivisionCount - 1; ++Andex)
 		{
-			const int32 GlobalIndex1 = Grid.GobalIndex(IndexU, PointIndexV);
-			const int32 GlobalIndex2 = GlobalIndex1 + 1;
-			const int32 IsoGlobalIndex = TryToBuildIsoSegment(LoopNode, Point, GlobalIndex1, GlobalIndex2, ESegmentType::InnerToLoopU);
-			if (IsoGlobalIndex == GlobalIndex1)
-			{
-				GlobalIndexToIsoInnerNodes[IsoGlobalIndex]->SetLinkedToIso(EIsoLink::IsoVNext);
-			}
-			else if (IsoGlobalIndex == GlobalIndex2)
+			TArray<FLoopNode*>& SubLoopA = Cell.OuterLoopSubdivision[Andex];
+			for (int32 Bndex = Andex + 1; Bndex < SubdivisionCount; ++Bndex)
 			{
-				GlobalIndexToIsoInnerNodes[IsoGlobalIndex]->SetLinkedToIso(EIsoLink::IsoVPrevious);
+				TArray<FLoopNode*>& SubLoopB = Cell.OuterLoopSubdivision[Bndex];
+				TryToConnectTwoSubLoopsWithTheMostIsoSegment(Cell, SubLoopA, SubLoopB);
 			}
 		}
 	}
 
-	InnerToLoopSegmentsIntersectionTool.Sort();
-#ifdef CADKERNEL_DEV
-	Chronos.FindIsoSegmentToLinkInnerToLoopDuration += FChrono::Elapse(StartTime);
+	{
+#ifdef DEBUG_FIND_ISO_SEGMENT_TO_LINK_OUTER_LOOP_NODES
+		F3DDebugSession _(bDisplay, ("FindIsoSegmentToLinkOuterLoopNodes 2"));
 #endif
+		for (TArray<FLoopNode*>& SubLoop : Cell.OuterLoopSubdivision)
+		{
+			TryToConnectVertexSubLoopWithTheMostIsoSegment(Cell, SubLoop);
+		}
+	}
+	
+	Cell.SelectSegmentInCandidateSegments(IsoSegmentFactory);
 }
 
-#endif
-
 // =============================================================================================================
 // 	   For each cell
 // 	      - Connect loops together and to cell vertices
@@ -1205,33 +1177,38 @@ void FIsoTriangulator::ConnectCellLoops()
 		ConnectCellSubLoopsByNeighborhood(Cell);
 
 #ifdef DEBUG_CONNECT_CELL_LOOPS
-		Display(EGridSpace::UniformScaled, TEXT("ConnectCellLoops::Step 0"), Cell.FinalSegments, false, false, EVisuProperty::YellowCurve);
+		Grid.DisplayIsoSegments(TEXT("ConnectCellLoops::Step 0"), EGridSpace::UniformScaled, Cell.FinalSegments, false, false, EVisuProperty::YellowCurve);
 #endif
 
 		if (Cell.bHasOuterLoop)
 		{
 			FindIsoSegmentToLinkOuterLoopNodes(Cell);
 #ifdef DEBUG_CONNECT_CELL_LOOPS
-			Display(EGridSpace::UniformScaled, TEXT("ConnectCellLoops::Step 1"), Cell.FinalSegments, false, false, EVisuProperty::YellowCurve);
+			Grid.DisplayIsoSegments(TEXT("ConnectCellLoops::Step 1"), EGridSpace::UniformScaled, Cell.FinalSegments, false, false, EVisuProperty::YellowCurve);
 #endif
 
 			if (Cell.CandidateSegments.Num() == 0)
 			{
 				FindSegmentToLinkOuterLoopNodes(Cell);
 #ifdef DEBUG_CONNECT_CELL_LOOPS
-				Display(EGridSpace::UniformScaled, TEXT("ConnectCellLoops::Step 2"), Cell.FinalSegments, false, false, EVisuProperty::YellowCurve);
+				Grid.DisplayIsoSegments(TEXT("ConnectCellLoops::Step 2"), EGridSpace::UniformScaled, Cell.FinalSegments, false, false, EVisuProperty::YellowCurve);
+#endif
+
+				FindSegmentToLinkOuterToInnerLoopNodes(Cell);
+#ifdef DEBUG_CONNECT_CELL_LOOPS
+				Grid.DisplayIsoSegments(TEXT("ConnectCellLoops::Step 3"), EGridSpace::UniformScaled, Cell.FinalSegments, false, false, EVisuProperty::YellowCurve);
 #endif
 			}
 		}
 		ConnectCellCornerToInnerLoop(Cell);
 #ifdef DEBUG_CONNECT_CELL_LOOPS
-		Display(EGridSpace::UniformScaled, TEXT("ConnectCellLoops::Step 3"), Cell.FinalSegments, false, false, EVisuProperty::YellowCurve);
+		Grid.DisplayIsoSegments(TEXT("ConnectCellLoops::Step 4"), EGridSpace::UniformScaled, Cell.FinalSegments, false, false, EVisuProperty::YellowCurve);
 #endif
 
 		FinalToLoops.Append(Cell.FinalSegments);
 
 #ifdef DEBUG_CONNECT_CELL_LOOPS
-		Display(EGridSpace::UniformScaled, TEXT("ConnectCellLoops::Final Iso ToLink Inner To Loops"), FinalToLoops, false, false, EVisuProperty::YellowCurve);
+		Grid.DisplayIsoSegments(TEXT("ConnectCellLoops::Final Iso ToLink Inner To Loops"), EGridSpace::UniformScaled, FinalToLoops, false, false, EVisuProperty::YellowCurve);
 #endif
 	}
 }
@@ -1535,419 +1512,154 @@ void FIsoTriangulator::FindCandidateSegmentsToLinkInnerAndLoop()
 #endif
 }
 
-void FIsoTriangulator::FindCandidateSegmentsToLinkInnerToLoop()
+void FIsoTriangulator::ConnectUnconnectedInnerSegments()
 {
-#ifdef CADKERNEL_DEV
-	FTimePoint StartTime = FChrono::Now();
-#endif
-
-	TArray<FIsoInnerNode*> SortedInnerNode;
-	SortedInnerNode.Reserve(InnerNodes.Num());
-	for (FIsoInnerNode& IsoNode : InnerNodes)
+	TArray<FIsoNode*> UnconnectedNode;
+	UnconnectedNode.Reserve(FinalInnerSegments.Num());
+	for (FIsoSegment* Segment : FinalInnerSegments)
 	{
-		if (IsoNode.IsComplete())
+		if (Segment->GetFirstNode().GetConnectedSegments().Num() == 1)
 		{
-			continue;
+			UnconnectedNode.Add(&Segment->GetFirstNode());
 		}
-		SortedInnerNode.Add(&IsoNode);
 	}
 
-	Algo::Sort(SortedInnerNode, [this](const FIsoInnerNode* IsoNode1, const FIsoInnerNode* IsoNode2)
-		{
-			const FPoint2D& Node1Coordinates = IsoNode1->Get2DPoint(EGridSpace::UniformScaled, Grid);
-			const FPoint2D& Node2Coordinates = IsoNode2->Get2DPoint(EGridSpace::UniformScaled, Grid);
-			return (Node1Coordinates.U + Node1Coordinates.V) < (Node2Coordinates.U + Node2Coordinates.V);
-		});
-
-#ifdef CADKERNEL_DEV
-	Chronos.FindSegmentToLinkInnerToLoopDuration = FChrono::Elapse(StartTime);
-#endif
-
-	TFunction<void(FIsoInnerNode*, FLoopNode*)> CreateInnerToContoursCandidateSegment = [&](FIsoInnerNode* InnerNode, FLoopNode* LoopNode)
+	for (FIsoNode* Node : UnconnectedNode)
 	{
-		if (FIsoSegment::IsItAlreadyDefined(InnerNode, LoopNode))
+		double MinDistance = HUGE_VALUE;
+		FLoopNode* Candidate = nullptr;
+		for (FLoopNode& LoopNode : LoopNodes)
 		{
-			return;
+			if (LoopNode.IsDelete())
+			{
+				continue;
+			}
+
+			double Distance = LoopNode.Get2DPoint(EGridSpace::Scaled, Grid).SquareDistance(Node->Get2DPoint(EGridSpace::Scaled, Grid));
+			if (Distance < MinDistance)
+			{
+				if (!InnerToLoopSegmentsIntersectionTool.DoesIntersect(*Node, LoopNode))
+				{
+					MinDistance = Distance;
+					Candidate = &LoopNode;
+				}
+			}
 		}
 
-		// check intersection with
-		if (LoopSegmentsIntersectionTool.DoesIntersect(*LoopNode, InnerNode->Get2DPoint(EGridSpace::UniformScaled, Grid)))
+		if (Candidate)
 		{
-			return;
-		}
-
-		if (InnerSegmentsIntersectionTool.DoesIntersect(*InnerNode, LoopNode->Get2DPoint(EGridSpace::UniformScaled, Grid)))
-		{
-			return;
-		}
-
-		if (InnerToLoopSegmentsIntersectionTool.DoesIntersect(*InnerNode, *LoopNode))
-		{
-			return;
-		}
-
-		FIsoSegment& SegCandidate = IsoSegmentFactory.New();
-		SegCandidate.Init(*InnerNode, *LoopNode, ESegmentType::InnerToLoop);
-		CandidateInnerToLoopSegments.Add(&SegCandidate);
-	};
-
-	TFunction<double(const FPoint&, const FPoint&)> ComputeDiagonal = [](const FPoint& PointA, const FPoint& PointB)
-	{
-		FPoint2D AB = PointB - PointA;
-		return AB.U + AB.V;
-	};
-
-	double DiagonalMax = sqrt(FMath::Square(Grid.GetMaxDeltaU(EIso::IsoU)) + FMath::Square(Grid.GetMaxDeltaU(EIso::IsoV)));
-	DiagonalMax *= 1.5;
-
-	TFunction<double(const TArray<double>&, int32&, const double&)> GetDeltaCoordinate = [](const TArray<double>& Iso, int32& Index, const double& PointCoord)
-	{
-		while (PointCoord + SMALL_NUMBER < Iso[Index])
-		{
-			if (Index == 0)
-			{
-				break;
-			}
-			Index--;
-		}
-		for (; Index < Iso.Num(); ++Index)
-		{
-			if (Index == Iso.Num() - 2)
-			{
-				break;
-			}
-			if (PointCoord < Iso[Index + 1u])
-			{
-				break;
-			}
-		}
-
-		if (Index == 1 && Iso.Num() > 3)
-		{
-			return Iso[1] - Iso[0];
-		}
-		else if (Index == 1 && Iso.Num() == 3)
-		{
-			return FMath::Max(Iso[1] - Iso[0], Iso[2] - Iso[1]);
-		}
-		return Iso[Index + 1u] - Iso[Index];
-	};
-
-	if (SortedInnerNode.Num())
-	{
-		CandidateInnerToLoopSegments.Reserve(SortedInnerNode.Num());
-		int32 IndexU = 0;
-		int32 IndexV = 0;
-
-		int32 StartInnerNodeIndex = 0;
-		for (FLoopNode* LoopNode : SortedLoopNodes)
-		{
-			const FPoint2D& PointLoop = LoopNode->Get2DPoint(EGridSpace::UniformScaled, Grid);
-			{
-				const FPoint2D& PointInner = Grid.GetInner2DPoint(EGridSpace::UniformScaled, SortedInnerNode[StartInnerNodeIndex]->GetIndex());
-				double Diagonal = ComputeDiagonal(PointLoop, PointInner);
-				if (Diagonal > DiagonalMax)
-				{
-					continue;
-				}
-			}
-
-			for (; StartInnerNodeIndex < SortedInnerNode.Num(); ++StartInnerNodeIndex)
-			{
-				const FPoint2D& PointInner = Grid.GetInner2DPoint(EGridSpace::UniformScaled, SortedInnerNode[StartInnerNodeIndex]->GetIndex());
-
-				double Diagonal = ComputeDiagonal(PointLoop, PointInner);
-				if (Diagonal >= -DiagonalMax)
-				{
-					break;
-				}
-			}
-			if (StartInnerNodeIndex == SortedInnerNode.Num())
-			{
-				break;
-			}
-
-			for (int32 InnerNodeIndex = StartInnerNodeIndex; InnerNodeIndex < SortedInnerNode.Num(); ++InnerNodeIndex)
-			{
-				const FPoint2D& PointInner = Grid.GetInner2DPoint(EGridSpace::UniformScaled, SortedInnerNode[InnerNodeIndex]->GetIndex());
-
-				FPoint2D DiagonalV = PointLoop - PointInner;
-				double Diagonal = FMath::Abs(DiagonalV.U + DiagonalV.V);
-
-				if (Diagonal > DiagonalMax)
-				{
-					break;
-				}
-
-				double DeltaUMax = GetDeltaCoordinate(Grid.GetUniformCuttingCoordinatesAlongIso(EIso::IsoU), IndexU, PointInner.U);
-				DeltaUMax *= 1.5;
-				double DeltaVMax = GetDeltaCoordinate(Grid.GetUniformCuttingCoordinatesAlongIso(EIso::IsoV), IndexV, PointInner.V);
-				DeltaVMax *= 1.5;
-
-				if (FMath::Abs(DiagonalV.U) > DeltaUMax || FMath::Abs(DiagonalV.V) > DeltaVMax)
-				{
-					continue;
-				}
-
-				CreateInnerToContoursCandidateSegment(SortedInnerNode[InnerNodeIndex], LoopNode);
-			}
+			FIsoSegment& Segment = IsoSegmentFactory.New();
+			Segment.Init(*Node, *Candidate, ESegmentType::InnerToLoop);
+			Segment.ConnectToNode();
+			FinalToLoops.Add(&Segment);
+			InnerToLoopSegmentsIntersectionTool.AddSegment(Segment);
+			InnerToLoopSegmentsIntersectionTool.Sort();
 		}
 	}
-
-#ifdef CADKERNEL_DEV
-	Chronos.FindSegmentToLinkInnerToLoopDuration = FChrono::Elapse(StartTime);
-#endif
 }
 
-void FIsoTriangulator::SelectSegmentInCandidateSegments()
+bool FIsoTriangulator::CanCycleBeMeshed(const TArray<FIsoSegment*>& Cycle, FIntersectionSegmentTool& CycleIntersectionTool)
 {
-	constexpr double FlatAngle = 0.25; // PI/12
-	constexpr double ColinearityCriteriaAngle = 0.25; // PI/12
-
-	FTimePoint StartTime = FChrono::Now();
+	bool bHasSelfIntersection = true;
 
-	// A candidate segment must not be too collinear with existing segment. So if this is the case, it's refused. 
-	// final step of the mesh will manage these cases
-	TFunction<bool(const FIsoNode&, const FIsoSegment*, const double)> CheckColinerarityWithFinalSegments = [&](const FIsoNode& Node, const FIsoSegment* CandidateSegment, const double CandidateSegmentSlop)
+	for (const FIsoSegment* Segment : Cycle)
 	{
-		bool bIsTooColineareWithFinalSegment = false;
-		for (const FIsoSegment* NodeSegment : Node.GetConnectedSegments())
+		if (CycleIntersectionTool.DoesIntersect(*Segment))
 		{
-			if (NodeSegment == CandidateSegment)
-			{
-				continue;
-			}
-
-			double AngleBetweenCandidateAndFinalSegment = &NodeSegment->GetFirstNode() == &Node ?
-				ComputeOrientedSlope(NodeSegment->GetFirstNode().Get2DPoint(EGridSpace::Default2D, Grid), NodeSegment->GetSecondNode().Get2DPoint(EGridSpace::Default2D, Grid), CandidateSegmentSlop) :
-				ComputeOrientedSlope(NodeSegment->GetSecondNode().Get2DPoint(EGridSpace::Default2D, Grid), NodeSegment->GetFirstNode().Get2DPoint(EGridSpace::Default2D, Grid), CandidateSegmentSlop);
-
-			if (FMath::Abs(AngleBetweenCandidateAndFinalSegment) < ColinearityCriteriaAngle)
-			{
-				bIsTooColineareWithFinalSegment = true;
-				return true;
-			}
-		}
-		return false;
-	};
+			FMessage::Printf(Log, TEXT("A cycle of the surface %d is in self intersecting. The mesh of this sector is canceled.\n"), Grid.GetFace()->GetId());
 
-	// The main idea is to select only one segment by Quarter ([0 Pi/2], [Pi/2, Pi], [Pi 3Pi/2], [3Pi/2 2Pi])
-	// in additional of Iso Segment (EIso::IsoU or EIso::IsoV). If the vertex has not an iso segment, closed iso segment are accepted 
-	TFunction<bool(const FIsoNode&, const FIsoNode&, const double, int32)> CheckQuarterCompletude = [&](const FIsoNode& StartNode, const FIsoNode& EndNode, const double CandidateSegmentSlop, int32 Quarter)
-	{
-		if (CandidateSegmentSlop < FlatAngle)
-		{
-			if (StartNode.IsLinkedToLoopInNearlyIso(Quarter))
-			{
-				return false;
-			}
-			if (EndNode.IsLinkedToLoopInNearlyIso(Quarter + 2))
-			{
-				return false;
-			}
-		}
-		else if (CandidateSegmentSlop < 2 - FlatAngle)
-		{
-			if (StartNode.IsLinkedToLoopInQuarter(Quarter))
-			{
-				return false;
-			}
-			if (EndNode.IsLinkedToLoopInQuarter(Quarter > 2 ? Quarter - 2 : Quarter + 2))
-			{
-				return false;
-			}
-		}
-		else
-		{
-			if (StartNode.IsLinkedToLoopInNearlyIso(Quarter + 1))
-			{
-				return false;
-			}
-			if (EndNode.IsLinkedToLoopInNearlyIso(Quarter + 3))
+#ifdef DEBUG_CAN_CYCLE_BE_FIXED_AND_MESHED
+			if (bDisplay)
 			{
-				return false;
+				F3DDebugSession _(TEXT("New Segment"));
+				F3DDebugSession A(TEXT("Segments"));
+				Display(EGridSpace::UniformScaled, *Segment, 0, EVisuProperty::RedPoint);
+				Display(EGridSpace::UniformScaled, *IntersectingSegment, 0, EVisuProperty::BluePoint);
+				Wait();
 			}
+#endif
+			return false;
 		}
-		return true;
-	};
-
-	TFunction<bool(FIsoNode&, FIsoNode&, const double, int32)>  SetNodeFlag = [&](FIsoNode& StartNode, FIsoNode& EndNode, const double CandidateSegmentSlop, int32 Quarter)
-	{
-		if (CandidateSegmentSlop < FlatAngle)
-		{
-			StartNode.SetLinkedToIso(Quarter);
-			EndNode.SetLinkedToIso(Quarter + 2);
-		}
-		else if (CandidateSegmentSlop < 2 - FlatAngle)
-		{
-			StartNode.SetLinkedToLoopInQuarter(Quarter);
-			EndNode.SetLinkedToLoopInQuarter(Quarter > 2 ? Quarter - 2 : Quarter + 2);
-		}
-		else
-		{
-			StartNode.SetLinkedToIso(Quarter + 1);
-			EndNode.SetLinkedToIso(Quarter + 3);
-		}
-		return true;
-	};
-
-	CandidateSegments.Reserve(CandidateLoopToLoopSegments.Num() + CandidateInnerToLoopSegments.Num());
-	CandidateSegments.Insert(CandidateLoopToLoopSegments, CandidateSegments.Num());
-	CandidateSegments.Insert(CandidateInnerToLoopSegments, CandidateSegments.Num());
-
-	for (FIsoSegment* Segment : CandidateSegments)
-	{
-		Segment->SetCandidate();
 	}
 
-	Algo::Sort(CandidateSegments, [this](const FIsoSegment* Segment1, const FIsoSegment* Segment2)
-		{
-			return (Segment1->Get3DLengthSquare(Grid)) < (Segment2->Get3DLengthSquare(Grid));
-		});
-
-	InnerSegmentsIntersectionTool.Reserve(InnerSegmentsIntersectionTool.Count() + (int32)CandidateSegments.Num());
-
-	FinalToLoops.Reserve(CandidateSegments.Num());
-
-	FIntersectionSegmentTool LocalSegmentsIntersectionTool(Grid);
-	for (FIsoSegment* Segment : CandidateSegments)
-	{
-		FIsoNode& StartNode = Segment->GetFirstNode();
-		FIsoNode& EndNode = Segment->GetSecondNode();
-
-		double CandidateSegmentSlop = ComputePositiveSlope(StartNode.Get2DPoint(EGridSpace::Default2D, Grid), EndNode.Get2DPoint(EGridSpace::Default2D, Grid), 0);
-
-		// Define the quarter ([0 Pi/2], [Pi/2, Pi], [Pi 3Pi/2], [3Pi/2 2Pi])
-		int32 Quarter = 3;
-		if (CandidateSegmentSlop < 2)
-		{
-			Quarter = 0;
-		}
-		else if (CandidateSegmentSlop < 4)
-		{
-			Quarter = 1;
-		}
-		else if (CandidateSegmentSlop < 6)
-		{
-			Quarter = 2;
-		}
-
-		if (!CheckQuarterCompletude(StartNode, EndNode, CandidateSegmentSlop - Quarter * 2, Quarter))
-		{
-			continue;
-		}
-
-		if (CheckColinerarityWithFinalSegments(StartNode, Segment, CandidateSegmentSlop))
-		{
-			continue;
-		}
+	return true;
+}
 
-		if (CheckColinerarityWithFinalSegments(EndNode, Segment, CandidateSegmentSlop > 4 ? CandidateSegmentSlop - 4 : CandidateSegmentSlop + 4))
-		{
-			continue;
-		}
+#ifdef WIP_ADD_STEP_TO_TO_FAVOR_ISO_SEGMENTS
+//#define DEBUG_FIND_ISO_SEGMENT_CYCLE
+void FIsoTriangulator::FindIsoCandidateSegmentInCycle(TArray<FIsoNode*> CycleNodes)
+{
+#ifdef DEBUG_FIND_ISO_SEGMENT_CYCLE
+	F3DDebugSession _(bDisplay, TEXT("FindIsoCandidateSegmentInCycle"));
+#endif
 
-		if (LocalSegmentsIntersectionTool.DoesIntersect(*Segment))
-		{
-			continue;
-		}
+	const double FlatSlope = 0.10; // ~5 deg: The segment must make an angle less than 10 deg with the Iso
+	double MinSlope = FlatSlope + SMALL_NUMBER; //.25; // ~10 deg: The segment must make an angle less than 10 deg with the Iso
 
-		Segment->SetSelected();
-		Segment->ConnectToNode();
-		SetNodeFlag(StartNode, EndNode, CandidateSegmentSlop - Quarter * 2, Quarter);
+	FIsoNode* CandidateA = nullptr;
+	FIsoNode* CandidateB = nullptr;
 
-		FinalToLoops.Add(Segment);
+	int32 CycleCount = CycleNodes.Num();
 
-		LocalSegmentsIntersectionTool.AddSegment(*Segment);
-		LocalSegmentsIntersectionTool.Sort();
-	}
+	TArray<FIsoSegment*> CandiateSegment;
+	CandiateSegment.Reserve(CycleCount);
 
-	for (FIsoSegment* Segment : CandidateSegments)
+	for (int32 IndexA = 0; IndexA < CycleCount - 2; ++IndexA)
 	{
-		if (Segment->IsACandidate())
-		{
-			IsoSegmentFactory.DeleteEntity(Segment);
-		}
-	}
-
-	InnerSegmentsIntersectionTool.Sort();
-#ifdef CADKERNEL_DEV
-	Chronos.SelectSegmentToLinkInnerToLoopsDuration = FChrono::Elapse(StartTime);
-#endif
-}
+		FIsoNode* NodeA = CycleNodes[IndexA];
+		const FPoint2D& ACoordinates = NodeA->Get2DPoint(EGridSpace::UniformScaled, Grid);
+		Grid.DisplayIsoNode(UniformScaled, *NodeA, 0, EVisuProperty::GreenPoint);
 
-void FIsoTriangulator::ConnectUnconnectedInnerSegments()
-{
-	TArray<FIsoNode*> UnconnectedNode;
-	UnconnectedNode.Reserve(FinalInnerSegments.Num());
-	for (FIsoSegment* Segment : FinalInnerSegments)
-	{
-		if (Segment->GetFirstNode().GetConnectedSegments().Num() == 1)
+		for (int32 IndexB = IndexA + 2; IndexB < CycleCount; ++IndexB)
 		{
-			UnconnectedNode.Add(&Segment->GetFirstNode());
-		}
-	}
+			FIsoNode* NodeB = CycleNodes[IndexB];
 
-	for (FIsoNode* Node : UnconnectedNode)
-	{
-		double MinDistance = HUGE_VALUE;
-		FLoopNode* Candidate = nullptr;
-		for (FLoopNode& LoopNode : LoopNodes)
-		{
-			if (LoopNode.IsDelete())
+			if(NodeA->GetSegmentConnectedTo(NodeB))
 			{
 				continue;
 			}
 
-			double Distance = LoopNode.Get2DPoint(EGridSpace::Scaled, Grid).SquareDistance(Node->Get2DPoint(EGridSpace::Scaled, Grid));
-			if (Distance < MinDistance)
+			const FPoint2D& BCoordinates = NodeB->Get2DPoint(EGridSpace::UniformScaled, Grid);
+
+			double Slope = ComputeSlopeRelativeToNearestAxis(ACoordinates, BCoordinates);
+			if (Slope < MinSlope)
 			{
-				if (!InnerToLoopSegmentsIntersectionTool.DoesIntersect(*Node, LoopNode))
-				{
-					MinDistance = Distance;
-					Candidate = &LoopNode;
-				}
+				MinSlope = Slope;
+				CandidateA = NodeA;
+				CandidateB = NodeB;
 			}
 		}
 
-		if (Candidate)
+		if (MinSlope < FlatSlope)
 		{
 			FIsoSegment& Segment = IsoSegmentFactory.New();
-			Segment.Init(*Node, *Candidate, ESegmentType::InnerToLoop);
-			Segment.ConnectToNode();
-			FinalToLoops.Add(&Segment);
-			InnerToLoopSegmentsIntersectionTool.AddSegment(Segment);
-			InnerToLoopSegmentsIntersectionTool.Sort();
+			Segment.Init(*CandidateA, *CandidateB, ESegmentType::LoopToLoop);
+			Segment.SetCandidate();
+
+			CandiateSegment.Add(&Segment);
+
+#ifdef DEBUG_FIND_ISO_SEGMENT_CYCLE
+			Grid.DisplayIsoSegment(UniformScaled, Segment, 0, EVisuProperty::YellowCurve);
+			Wait(bDisplay);
+#endif
+			MinSlope = FlatSlope + SMALL_NUMBER;
 		}
 	}
-}
-
-bool FIsoTriangulator::CanCycleBeMeshed(const TArray<FIsoSegment*>& Cycle, FIntersectionSegmentTool& CycleIntersectionTool)
-{
-	bool bHasSelfIntersection = true;
 
-	for (const FIsoSegment* Segment : Cycle)
+	if(bDisplay)
 	{
-		if (CycleIntersectionTool.DoesIntersect(*Segment))
-		{
-			FMessage::Printf(Log, TEXT("A cycle of the surface %d is in self intersecting. The mesh of this sector is canceled.\n"), Grid.GetFace()->GetId());
+		Wait();
+	}
 
-#ifdef DEBUG_CAN_CYCLE_BE_FIXED_AND_MESHED
-			if (bDisplay)
-			{
-				F3DDebugSession _(TEXT("New Segment"));
-				F3DDebugSession A(TEXT("Segments"));
-				Display(EGridSpace::UniformScaled, *Segment, 0, EVisuProperty::RedPoint);
-				Display(EGridSpace::UniformScaled, *IntersectingSegment, 0, EVisuProperty::BluePoint);
-				Wait();
-			}
-#endif
-			return false;
-		}
+	// find the best iso segments i.e. ones that generate smooth surface
+	// check that the selected segments won't block the meshing i.e. a cycle segment cannot be meshed
+	for(FIsoSegment* Segment : CandiateSegment)
+	{
 	}
 
-	return true;
 }
+#endif
+
 
 void FIsoTriangulator::MeshCycle(const EGridSpace Space, const TArray<FIsoSegment*>& Cycle, const TArray<bool>& CycleOrientation)
 {
@@ -2060,13 +1772,16 @@ void FIsoTriangulator::MeshCycle(const EGridSpace Space, const TArray<FIsoSegmen
 	CycleIntersectionTool.Sort();
 
 #ifdef ADD_TRIANGLE_2D
-	F3DDebugSession _(bDisplay, TEXT("Mesh cycle"));
+	static int32 CycleId = 0;
+	F3DDebugSession _(bDisplay, FString::Printf(TEXT("Mesh cycle %d"), CycleId++));
 	Grid.DisplayIsoSegments(TEXT("Cycle"), EGridSpace::UniformScaled, Cycle);
+	Wait(bDisplay);
 #endif
 
 	// Check if the cycle is in self intersecting and fix it. 
 	if (!CanCycleBeMeshed(Cycle, CycleIntersectionTool))
 	{
+		MesherReport.Logs.AddCycleMeshingFailure();
 		return;
 	}
 
@@ -2167,6 +1882,8 @@ void FIsoTriangulator::MeshCycle(const EGridSpace Space, const TArray<FIsoSegmen
 #ifdef ADD_TRIANGLE_2D
 				if (ABSegment->HasTriangleOnLeft())
 				{
+					F3DDebugSession _(FString::Printf(TEXT("Segment")));
+					Grid.DisplayIsoSegment(EGridSpace::UniformScaled, *NodeA, *NodeB, 0, EVisuProperty::YellowCurve);
 					Wait();
 				}
 #endif
@@ -2178,6 +1895,8 @@ void FIsoTriangulator::MeshCycle(const EGridSpace Space, const TArray<FIsoSegmen
 #ifdef ADD_TRIANGLE_2D
 				if (ABSegment->HasTriangleOnRight())
 				{
+					F3DDebugSession _(FString::Printf(TEXT("Segment")));
+					Grid.DisplayIsoSegment(EGridSpace::UniformScaled, *NodeA, *NodeB, 0, EVisuProperty::YellowCurve);
 					Wait();
 				}
 #endif
@@ -2196,12 +1915,11 @@ void FIsoTriangulator::MeshCycle(const EGridSpace Space, const TArray<FIsoSegmen
 		}
 	};
 
+	const SlopMethod GetSlopAtStartNode = ClockwiseSlop;
+	const SlopMethod GetSlopAtEndNode = CounterClockwiseSlop;
+
 	TFunction<void(FIsoSegment*, bool)> FindBestTriangle = [&](FIsoSegment* Segment, bool bOrientation)
 	{
-		using namespace IsoTriangulatorImpl;
-		SlopMethod GetSlopAtStartNode = ClockwiseSlop;
-		SlopMethod GetSlopAtEndNode = CounterClockwiseSlop;
-
 		// StartNode = A
 		FIsoNode& StartNode = bOrientation ? Segment->GetFirstNode() : Segment->GetSecondNode();
 		// EndNode = B
@@ -2247,24 +1965,24 @@ void FIsoTriangulator::MeshCycle(const EGridSpace Space, const TArray<FIsoSegmen
 		FIsoNode& NextNode = NextSegment->GetOtherNode(&EndNode);
 
 #ifdef DEBUG_FIND_BEST_TRIANGLE
-		int32 TriangleIndex = 0;
+		static int32 TriangleIndex = 0;
 		if (bDisplay)
 		{
 			{
 				TriangleIndex++;
 				F3DDebugSession _(FString::Printf(TEXT("Start Segment %d %d"), TriangleIndex, bOrientation));
-				Display(EGridSpace::UniformScaled, *Segment);
-				Display(EGridSpace::UniformScaled, StartNode, 0, EVisuProperty::RedPoint);
-				Display(EGridSpace::UniformScaled, EndNode);
+				Grid.DisplayIsoSegment(EGridSpace::UniformScaled, *Segment);
+				Grid.DisplayIsoNode(EGridSpace::UniformScaled, StartNode, 0, EVisuProperty::RedPoint);
+				Grid.DisplayIsoNode(EGridSpace::UniformScaled, EndNode);
 			}
 
 			{
 				F3DDebugSession _(TEXT("Next Segments"));
-				Display(EGridSpace::UniformScaled, *PreviousSegment);
-				Display(EGridSpace::UniformScaled, *NextSegment, 0, EVisuProperty::BlueCurve);
-				Display(EGridSpace::UniformScaled, EndNode, 0, EVisuProperty::RedPoint);
-				Display(EGridSpace::UniformScaled, NextNode);
-				Wait(Grid.bDisplay);
+				Grid.DisplayIsoSegment(EGridSpace::UniformScaled, *PreviousSegment);
+				Grid.DisplayIsoSegment(EGridSpace::UniformScaled, *NextSegment, 0, EVisuProperty::BlueCurve);
+				Grid.DisplayIsoNode(EGridSpace::UniformScaled, EndNode, 0, EVisuProperty::RedPoint);
+				Grid.DisplayIsoNode(EGridSpace::UniformScaled, NextNode);
+				Wait(TriangleIndex == 1000);
 			}
 		}
 #endif
@@ -2286,7 +2004,7 @@ void FIsoTriangulator::MeshCycle(const EGridSpace Space, const TArray<FIsoSegmen
 			double StartMaxSlope = GetSlopAtStartNode(StartPoint2D, PreviousNode.Get2DPoint(EGridSpace::UniformScaled, Grid), StartReferenceSlope);
 			if (&EndNode != &PreviousNode)
 			{
-				// Case of probable auto-intersection cycle at PreviousNode, cancel is prefered 
+				// Case of probable auto-intersection cycle at PreviousNode, cancel is preferred 
 				if (FMath::IsNearlyEqual(StartMaxSlope, 8., (double)KINDA_SMALL_NUMBER))
 				{
 					return;
@@ -2298,7 +2016,7 @@ void FIsoTriangulator::MeshCycle(const EGridSpace Space, const TArray<FIsoSegmen
 			double EndMaxSlope = GetSlopAtEndNode(EndPoint2D, NextNode.Get2DPoint(EGridSpace::UniformScaled, Grid), EndReferenceSlope);
 			if (&StartNode != &NextNode)
 			{
-				// Case of probable auto-intersection cycle at PreviousNode, cancel is prefered 
+				// Case of probable auto-intersection cycle at PreviousNode, cancel is preferred 
 				if (FMath::IsNearlyEqual(EndMaxSlope, 8., (double)KINDA_SMALL_NUMBER))
 				{
 					return;
@@ -2307,7 +2025,8 @@ void FIsoTriangulator::MeshCycle(const EGridSpace Space, const TArray<FIsoSegmen
 			EndMaxSlope = FMath::Min(EndMaxSlope, 4.);
 
 			double MinCriteria = HUGE_VALUE;
-			const double MinSlopToNotBeAligned = 0.0001;
+			const double MinSlopToNotBeAligned = 0.01;
+			const double EpsilonSlop = 0.0001;
 			double CandidateSlopeAtStartNode = 8.;
 			double CandidateSlopeAtEndNode = 8.;
 
@@ -2329,7 +2048,7 @@ void FIsoTriangulator::MeshCycle(const EGridSpace Space, const TArray<FIsoSegmen
 
 				if (Node != &PreviousNode)
 				{
-					if (SlopeAtStartNode <= MinSlopToNotBeAligned || SlopeAtStartNode >= StartMaxSlope - MinSlopToNotBeAligned)
+					if (SlopeAtStartNode <= EpsilonSlop || SlopeAtStartNode >= StartMaxSlope - EpsilonSlop)
 					{
 						continue;
 					}
@@ -2337,18 +2056,28 @@ void FIsoTriangulator::MeshCycle(const EGridSpace Space, const TArray<FIsoSegmen
 
 				if (Node != &NextNode)
 				{
-					if (SlopeAtEndNode <= MinSlopToNotBeAligned || SlopeAtEndNode >= EndMaxSlope - MinSlopToNotBeAligned)
+					if (SlopeAtEndNode <= EpsilonSlop || SlopeAtEndNode >= EndMaxSlope - EpsilonSlop)
 					{
 						continue;
 					}
 				}
 
-				if (FMath::IsNearlyEqual(SlopeAtStartNode, CandidateSlopeAtStartNode, MinSlopToNotBeAligned) && SlopeAtEndNode > CandidateSlopeAtEndNode)
+#ifdef DEBUG_FIND_BEST_TRIANGLE
+				if (TriangleIndex == 14)
+				{
+					static int32 NodeIndex = 1;
+					F3DDebugSession _(FString::Printf(TEXT("Next Node %d"), NodeIndex);
+					Grid.DisplayIsoNode(EGridSpace::UniformScaled, *Node, ++NodeIndex, EVisuProperty::YellowPoint);
+					Wait();
+				}
+#endif
+
+				if (FMath::IsNearlyEqual(SlopeAtStartNode, CandidateSlopeAtStartNode, MinSlopToNotBeAligned) && SlopeAtEndNode > CandidateSlopeAtEndNode)
 				{
 					continue;
 				}
 
-				if (FMath::IsNearlyEqual(SlopeAtEndNode, CandidateSlopeAtEndNode, MinSlopToNotBeAligned) && SlopeAtStartNode > CandidateSlopeAtStartNode)
+				if (FMath::IsNearlyEqual(SlopeAtEndNode, CandidateSlopeAtEndNode, EpsilonSlop) && SlopeAtStartNode > CandidateSlopeAtStartNode)
 				{
 					continue;
 				}
@@ -2436,6 +2165,12 @@ void FIsoTriangulator::MeshCycle(const EGridSpace Space, const TArray<FIsoSegmen
 		}
 	};
 
+#ifdef WIP_ADD_STEP_TO_TO_FAVOR_ISO_SEGMENTS
+	FindIsoCandidateSegmentInCycle(CycleNodes);
+#endif
+
+
+
 	for (int32 Index = 0; Index < SegmentStack.Num(); ++Index)
 	{
 		FIsoSegment* Segment = SegmentStack[Index];
@@ -2454,6 +2189,26 @@ void FIsoTriangulator::MeshCycle(const EGridSpace Space, const TArray<FIsoSegmen
 		}
 	}
 
+	// in case of incomplete cycle meshing, remove degenerated flags and rerun the process
+	for (FIsoSegment* Segment : SegmentStack)
+	{
+		Segment->ResetDegenerated();
+	}
+
+	for (int32 Index = 0; Index < SegmentStack.Num(); ++Index)
+	{
+		FIsoSegment* Segment = SegmentStack[Index];
+
+		if (!Segment->HasTriangleOnLeft())
+		{
+			FindBestTriangle(Segment, false);
+		}
+		if (!Segment->HasTriangleOnRight())
+		{
+			FindBestTriangle(Segment, true);
+		}
+	}
+
 	// Reset the flags "has triangle" of cycle's segments to avoid to block the meshing of next cycles
 	for (FIsoSegment* Segment : Cycle)
 	{
@@ -2756,7 +2511,7 @@ void FIsoTriangulator::FindCycle(FIsoSegment* StartSegment, bool LeftSide, TArra
 
 		if (&Segment->GetFirstNode() == Node)
 		{
-#ifdef CADKERNEL_DEV
+#ifdef DEBUG_HAS_CYCLE
 			if (Segment->HasCycleOnLeft())
 			{
 				F3DDebugSession _(TEXT("Segment HasCycleOnLeft"));
@@ -2780,7 +2535,7 @@ void FIsoTriangulator::FindCycle(FIsoSegment* StartSegment, bool LeftSide, TArra
 		}
 		else
 		{
-#ifdef CADKERNEL_DEV
+#ifdef DEBUG_HAS_CYCLE
 			if (Segment->HasCycleOnRight())
 			{
 				F3DDebugSession _(TEXT("Segment HasCycleOnRight"));
@@ -2925,8 +2680,8 @@ void FIsoTriangulator::TriangulateInnerNodes()
 			}
 
 #ifdef ADD_TRIANGLE_2D
-			Grid.DisplayTriangle(EGridSpace::Default2D, *GlobalIndexToIsoInnerNodes[Index], *GlobalIndexToIsoInnerNodes[Index + 1], *GlobalIndexToIsoInnerNodes[OppositIndex]);
-			Grid.DisplayTriangle(EGridSpace::Default2D, *GlobalIndexToIsoInnerNodes[OppositIndex], *GlobalIndexToIsoInnerNodes[OppositIndex - 1], *GlobalIndexToIsoInnerNodes[Index]);
+			Grid.DisplayTriangle(EGridSpace::UniformScaled, *GlobalIndexToIsoInnerNodes[Index], *GlobalIndexToIsoInnerNodes[Index + 1], *GlobalIndexToIsoInnerNodes[OppositIndex]);
+			Grid.DisplayTriangle(EGridSpace::UniformScaled, *GlobalIndexToIsoInnerNodes[OppositIndex], *GlobalIndexToIsoInnerNodes[OppositIndex - 1], *GlobalIndexToIsoInnerNodes[Index]);
 #endif 
 
 			Mesh->AddTriangle(GlobalIndexToIsoInnerNodes[Index]->GetFaceIndex(), GlobalIndexToIsoInnerNodes[Index + 1]->GetFaceIndex(), GlobalIndexToIsoInnerNodes[OppositIndex]->GetFaceIndex());
@@ -3295,137 +3050,125 @@ void FIsoTriangulator::TryToConnectTwoLoopsWithIsocelesTriangle(FCell& Cell, con
 
 };
 
-
+//#define DEBUG_FIND_ISO_SEGMENT_TO_LINK_OUTER_LOOP_NODES
 void FIsoTriangulator::TryToConnectVertexSubLoopWithTheMostIsoSegment(FCell& Cell, const TArray<FLoopNode*>& Loop)
 {
-	double MinSlope = 0.25; // ~10 deg: The segment must make an angle less than 10 deg with the Iso
-	int32 MinIndexA = -1;
-	int32 MinIndexB = -1;
+#ifdef DEBUG_FIND_ISO_SEGMENT_TO_LINK_OUTER_LOOP_NODES
+	F3DDebugSession _(bDisplay, TEXT("TryToConnectVertexSubLoopWithTheMostIsoSegment"));
+	if(bDisplay)
+	{
+		Wait();
+	}
+#endif
+
+	const double FlatSlope = 0.10; // ~5 deg: The segment must make an angle less than 10 deg with the Iso
+	double MinSlope = FlatSlope + SMALL_NUMBER; //.25; // ~10 deg: The segment must make an angle less than 10 deg with the Iso
+
 
 	if (Loop.Num() <= 2)
 	{
 		return;
 	}
 
-	// Perhaps only the both extremities could be test 
-	for (int32 IndexA = 0; IndexA < Loop.Num() - 2; ++IndexA)
+	int32 LoopCount = Loop.Num();
+	for (int32 IndexA = 0; IndexA < LoopCount - 2; ++IndexA)
 	{
-		const FLoopNode* NodeA = Loop[IndexA];
-		const FPoint2D& ACoordinates = NodeA->Get2DPoint(EGridSpace::UniformScaled, Grid);
-		const FLoopNode* NextA = Loop[IndexA + 1];
+		FLoopNode* CandidateB = nullptr;
+
+		FLoopNode* CandidateA = Loop[IndexA];
+		const FPoint2D& ACoordinates = CandidateA->Get2DPoint(EGridSpace::UniformScaled, Grid);
+
+		FLoopNode* NextA = Loop[IndexA + 1];
 		const FPoint2D& NextACoordinates = NextA->Get2DPoint(EGridSpace::UniformScaled, Grid);
-		// slop SegmentA (NodeA, NextA)
-		double Slope = ComputeUnorientedSlope(ACoordinates, NextACoordinates, 0);
-		if (Slope > 2)
-		{
-			Slope = 4 - Slope;
-		}
+
 		double ReferenceSlope = 0;
-		if (Slope < 0.5)
-		{
-			// SegmentA close to IsoV => looking for a segment close to IsoU
-			ReferenceSlope = 2;
-		}
-		else if (Slope < 1.5)
 		{
-			// SegmentA is neither close to IsoV nor IsoU
-			continue;
+			// slope SegmentA (NodeA, NextA)
+			double Slope = ComputeUnorientedSlope(ACoordinates, NextACoordinates, 0);
+			if (Slope > 1.5 && Slope < 2.5)
+			{
+				ReferenceSlope = 0;
+			}
+			else if ((Slope < 0.5) || (Slope > 3.5))
+			{
+				ReferenceSlope = 2;
+			}
+			else
+			{
+				// SegmentA is neither close to IsoV nor IsoU
+				continue;
+			}
 		}
 
-		for (int32 IndexB = IndexA + 2; IndexB < Loop.Num(); ++IndexB)
+		for (int32 IndexB = IndexA + 2; IndexB < LoopCount; ++IndexB)
 		{
-			const FLoopNode* NodeB = Loop[IndexB];
+			FLoopNode* NodeB = Loop[IndexB];
 			const FPoint2D& BCoordinates = NodeB->Get2DPoint(EGridSpace::UniformScaled, Grid);
 
-			Slope = ComputeUnorientedSlope(ACoordinates, BCoordinates, 0);
-			// Want a slop between 0 and 2
-			if (Slope > 2)
-			{
-				Slope = 4 - Slope;
-			}
-
-			Slope = FMath::Abs(Slope - ReferenceSlope);
-
+			double Slope = ComputeSlopeRelativeToReferenceAxis(ACoordinates, BCoordinates, ReferenceSlope);
 			if (Slope < MinSlope)
 			{
 				MinSlope = Slope;
-				MinIndexA = IndexA;
-				MinIndexB = IndexB;
+				CandidateB = NodeB;
 			}
 		}
-	}
 
-	if (MinIndexA >= 0 && MinIndexB >= 0)
-	{
-		FLoopNode* NodeA = Loop[MinIndexA];
-		const FPoint2D& ACoordinates = NodeA->Get2DPoint(EGridSpace::UniformScaled, Grid);
-		FLoopNode* NodeB = Loop[MinIndexB];
-		const FPoint2D& BCoordinates = NodeB->Get2DPoint(EGridSpace::UniformScaled, Grid);
+		if (MinSlope < FlatSlope)
+		{
+			const FPoint2D& BCoordinates = CandidateB->Get2DPoint(EGridSpace::UniformScaled, Grid);
 
 #ifdef DEBUG_FIND_ISO_SEGMENT_TO_LINK_OUTER_LOOP_NODES
-		DisplaySegment(NodeA->Get2DPoint(EGridSpace::UniformScaled, Grid), NodeB->Get2DPoint(EGridSpace::UniformScaled, Grid), 0, EVisuProperty::RedCurve);
-		Wait();
+			Grid.DisplayIsoSegment(UniformScaled, *CandidateA, *CandidateB, 0, EVisuProperty::RedCurve);
 #endif
-		TryToCreateSegment(Cell, NodeA, ACoordinates, NodeB, BCoordinates, 0.1);
+			TryToCreateSegment(Cell, CandidateA, ACoordinates, CandidateB, BCoordinates, 0.1);
+			MinSlope = FlatSlope + SMALL_NUMBER;
+		}
 	}
 
 };
 
 void FIsoTriangulator::TryToConnectTwoSubLoopsWithTheMostIsoSegment(FCell& Cell, const TArray<FLoopNode*>& LoopA, const TArray<FLoopNode*>& LoopB)
 {
-	double MinSlope = 0.25; // ~10 deg: The segment must make an angle less than 10 deg with the Iso
-	int32 MinIndexA = -1;
-	int32 MinIndexB = -1;
+#ifdef DEBUG_FIND_ISO_SEGMENT_TO_LINK_OUTER_LOOP_NODES
+	F3DDebugSession _(bDisplay, TEXT("TryToConnectVertexSubLoopWithTheMostIsoSegment"));
+	if (bDisplay)
+	{
+		Wait();
+	}
+#endif
 
-	for (int32 IndexA = 0; IndexA < LoopA.Num(); ++IndexA)
+	const double FlatSlope = 0.10; // ~5 deg: The segment must make an angle less than 10 deg with the Iso
+	double MinSlope = FlatSlope + SMALL_NUMBER;// 0.25; // ~15 deg: The segment must make an angle less than 10 deg with the Iso
+
+	for (FLoopNode* CandidateA : LoopA)
 	{
-		const FLoopNode* NodeA = LoopA[IndexA];
-		const FPoint2D& ACoordinates = NodeA->Get2DPoint(EGridSpace::UniformScaled, Grid);
+		FLoopNode* CandidateB = nullptr;
+		const FPoint2D& ACoordinates = CandidateA->Get2DPoint(EGridSpace::UniformScaled, Grid);
 
-		for (int32 IndexB = 0; IndexB < LoopB.Num(); ++IndexB)
+		for (FLoopNode* NodeB : LoopB)
 		{
-			const FLoopNode* NodeB = LoopB[IndexB];
 			const FPoint2D& BCoordinates = NodeB->Get2DPoint(EGridSpace::UniformScaled, Grid);
 
-			double Slope = ComputeUnorientedSlope(ACoordinates, BCoordinates, 0);
-			// Want a slop between 0 and 2
-			if (Slope > 2)
-			{
-				Slope = 4 - Slope;
-			}
-
-			// if slop close to 2 means segment close to IsoU, otherwize segment close to IsoV
-			// Wants a slop between 0 and 1 to manage either IsoU and IsoV
-			// Close to 0 means close to IsoU or IsoV
-			if (Slope > 1)
-			{
-				Slope = 2 - Slope;
-			}
-
+			double Slope = ComputeSlopeRelativeToNearestAxis(ACoordinates, BCoordinates);
 			if (Slope < MinSlope)
 			{
 				MinSlope = Slope;
-				MinIndexA = IndexA;
-				MinIndexB = IndexB;
+				CandidateB = NodeB;
 			}
 		}
-	}
 
-	if (MinIndexA >= 0 && MinIndexB >= 0)
-	{
-		FLoopNode* NodeA = LoopA[MinIndexA];
-		const FPoint2D& ACoordinates = NodeA->Get2DPoint(EGridSpace::UniformScaled, Grid);
-		FLoopNode* NodeB = LoopB[MinIndexB];
-		const FPoint2D& BCoordinates = NodeB->Get2DPoint(EGridSpace::UniformScaled, Grid);
+		if (MinSlope < FlatSlope)
+		{
+			const FPoint2D& BCoordinates = CandidateB->Get2DPoint(EGridSpace::UniformScaled, Grid);
 
 #ifdef DEBUG_FIND_ISO_SEGMENT_TO_LINK_OUTER_LOOP_NODES
-		DisplaySegment(NodeA->Get2DPoint(EGridSpace::UniformScaled, Grid), NodeB->Get2DPoint(EGridSpace::UniformScaled, Grid), 0, EVisuProperty::RedCurve);
-		Wait();
-#endif
-		TryToCreateSegment(Cell, NodeA, ACoordinates, NodeB, BCoordinates, 0.1);
+			Grid.DisplayIsoSegment(EGridSpace::UniformScaled, *CandidateA, *CandidateB, 0, EVisuProperty::BlueCurve);
+#endif			
+			TryToCreateSegment(Cell, CandidateA, ACoordinates, CandidateB, BCoordinates, 0.1);
+			MinSlope = FlatSlope + SMALL_NUMBER;
+		}
 	}
-
-};
+}
 
 //#define DEBUG_TRY_TO_CREATE_SEGMENT
 bool FIsoTriangulator::TryToCreateSegment(FCell& Cell, FLoopNode* NodeA, const FPoint2D& ACoordinates, FIsoNode* NodeB, const FPoint2D& BCoordinates, const double FlatAngle)
@@ -3549,7 +3292,6 @@ void FIsoTriangulator::ConnectCellCornerToInnerLoop(FCell& Cell)
 
 			TryToCreateSegment(Cell, NodeA, ACoordinates, InnerNode, InnerCoordinates, 0.1);
 		}
-
 	};
 
 	int32 IntersectionToolCount = Cell.IntersectionTool.Count();
@@ -3577,638 +3319,4 @@ void FIsoTriangulator::ConnectCellCornerToInnerLoop(FCell& Cell)
 	Cell.SelectSegmentInCandidateSegments(IsoSegmentFactory);
 }
 
-
-
-// =========================================================================================================================================================================================================
-// =========================================================================================================================================================================================================
-// =========================================================================================================================================================================================================
-//
-//
-//                                                                            NOT YET REVIEWED
-//
-//
-// =========================================================================================================================================================================================================
-// =========================================================================================================================================================================================================
-// =========================================================================================================================================================================================================
-
-
-#ifdef NEED_TO_CHECK_USEFULNESS
-void FIsoTriangulator::CompleteIsoSegmentLoopToLoop()
-{
-	FIntersectionSegmentTool NewSegmentsIntersectionTool(Grid);
-	// similare to FindSegmentToLinkloopToloop::CreateSegment but different 
-	TFunction<void(FLoopNode&, FIsoNode&)> CreateComplementarySegment = [&](FLoopNode& Node1, FIsoNode& Node2)
-	{
-		if (Node1.GetSegmentConnectedTo(&Node2))
-		{
-			return;
-		}
-
-		if (NewSegmentsIntersectionTool.DoesIntersect(Node1, Node2))
-		{
-			return;
-		}
-
-		if (LoopSegmentsIntersectionTool.DoesIntersect(Node1, Node2))
-		{
-			return;
-		}
-
-		FIsoSegment& Segment = IsoSegmentFactory.New();
-		Segment.Init(Node1, Node2, ESegmentType::Unknown);
-		Segment.ConnectToNode();
-		FinalToLoops.Add(&Segment);
-		InnerToLoopSegmentsIntersectionTool.AddSegment(Segment);
-		NewSegmentsIntersectionTool.AddSegment(Segment);
-		NewSegmentsIntersectionTool.Sort();
-	};
-
-	for (FLoopNode& LoopNode : LoopNodes)
-	{
-		if (LoopNode.GetConnectedSegments().Num() == 2)
-		{
-			if (LoopNode.GetPreviousNode().GetConnectedSegments().Num() == 3)
-			{
-				if ((LoopNode.GetPreviousNode().GetConnectedSegments().Last()->GetType() == ESegmentType::IsoLoopToLoop) && (LoopNode.GetNextNode().GetConnectedSegments().Last()->GetType() == ESegmentType::IsoLoopToLoop))
-				{
-					FIsoNode& Node1 = LoopNode.GetPreviousNode().GetConnectedSegments().Last()->GetOtherNode(&LoopNode.GetPreviousNode());
-					DisplaySegment(Node1.Get2DPoint(EGridSpace::Default2D, Grid), LoopNode.Get2DPoint(EGridSpace::Default2D, Grid));
-					CreateComplementarySegment(LoopNode, Node1);
-					FIsoNode& Node2 = LoopNode.GetNextNode().GetConnectedSegments().Last()->GetOtherNode(&LoopNode.GetNextNode());
-					CreateComplementarySegment(LoopNode, Node2);
-					DisplaySegment(Node2.Get2DPoint(EGridSpace::Default2D, Grid), LoopNode.Get2DPoint(EGridSpace::Default2D, Grid));
-				}
-				else if (LoopNode.GetPreviousNode().GetConnectedSegments().Last()->GetType() == ESegmentType::IsoLoopToLoop)
-				{
-					FIsoNode& Node = LoopNode.GetPreviousNode().GetConnectedSegments().Last()->GetOtherNode(&LoopNode.GetPreviousNode());
-					CreateComplementarySegment(LoopNode, Node);
-					DisplaySegment(Node.Get2DPoint(EGridSpace::Default2D, Grid), LoopNode.Get2DPoint(EGridSpace::Default2D, Grid));
-				}
-				else if (LoopNode.GetNextNode().GetConnectedSegments().Last()->GetType() == ESegmentType::IsoLoopToLoop)
-				{
-					FIsoNode& Node = LoopNode.GetNextNode().GetConnectedSegments().Last()->GetOtherNode(&LoopNode.GetNextNode());
-					CreateComplementarySegment(LoopNode, Node);
-					DisplaySegment(Node.Get2DPoint(EGridSpace::Default2D, Grid), LoopNode.Get2DPoint(EGridSpace::Default2D, Grid));
-				}
-			}
-		}
-	}
-}
-#endif
-
-#ifdef UNUSED
-void FIsoTriangulator::FindSegmentToLinkLoopToLoop()
-{
-	FTimePoint StartTime = FChrono::Now();
-
-	double DiagonalMax = Grid.GetMaxDeltaU(EIso::IsoU) + Grid.GetMaxDeltaU(EIso::IsoV);
-	DiagonalMax *= 2;
-
-	TFunction<double(const TArray<double>&, int32&, const double&)> GetDeltaCoordinate = [](const TArray<double>& Iso, int32& Index, const double& PointCoord)
-	{
-		while (PointCoord + SMALL_NUMBER < Iso[Index])
-		{
-			if (Index == 0)
-			{
-				break;
-			}
-			Index--;
-		}
-		for (; Index < Iso.Num(); ++Index)
-		{
-			if (Index == Iso.Num() - 2)
-			{
-				break;
-			}
-			if (PointCoord < Iso[Index + 1])
-			{
-				break;
-			}
-		}
-
-		// If PointCoord is closer Iso[Index + 1]
-		if ((Iso[Index + 1] - PointCoord) * 2 < (Iso[Index + 1] - Iso[Index]))
-		{
-			if (Index + 2 < Iso.Num())
-			{
-				return FMath::Max(Iso[Index + 1] - Iso[Index], Iso[Index + 2] - Iso[Index + 1]);
-			}
-			else
-			{
-				return Iso[Index + 1] - Iso[Index];
-			}
-		}
-		else
-		{
-			if (Index > 1)
-			{
-				return FMath::Max(Iso[Index] - Iso[Index - 1], Iso[Index + 1] - Iso[Index]);
-			}
-			else
-			{
-				return Iso[Index + 1] - Iso[Index];
-			}
-		}
-		return Iso[Index + 1] - Iso[Index];
-	};
-
-	TFunction<void(FLoopNode*, const FPoint2D&, FLoopNode*, const FPoint2D&)> CreateSegment = [this](FLoopNode* NodeA, const FPoint2D& ACoordinates, FLoopNode* NodeB, const FPoint2D& BCoordinates)
-	{
-		if (NodeA->GetSegmentConnectedTo(NodeB))
-		{
-			return;
-		}
-
-		// IsExiste
-		if (NodeA->GetLoopIndex() == NodeB->GetLoopIndex())
-		{
-			if ((NodeA->GetIndex() == NodeB->GetIndex() + 1) || (NodeB->GetIndex() == NodeA->GetIndex() + 1))
-			{
-				return;
-			}
-		}
-
-		if (InnerSegmentsIntersectionTool.DoesIntersect(ACoordinates, BCoordinates))
-		{
-			return;
-		}
-
-		if (InnerToLoopSegmentsIntersectionTool.DoesIntersect(*NodeA, *NodeB))
-		{
-			return;
-		}
-
-		if (LoopSegmentsIntersectionTool.DoesIntersect(*NodeA, *NodeB))
-		{
-			return;
-		}
-
-		// Is Outside and not to flat at NodeA
-		const double FlatAngle = 1;
-		if (!IsPointPInsideSectorABC(NodeA->GetPreviousNode().Get2DPoint(EGridSpace::Default2D, Grid), ACoordinates, NodeA->GetNextNode().Get2DPoint(EGridSpace::Default2D, Grid), BCoordinates, FlatAngle))
-		{
-			return;
-		}
-
-		//double SlopWithNextloop = GetZeroTwoPiSlope(ACoordinates, NodeA->GetNextNode().Get2DPoint(EGridSpace::Default2D, Grid), 0);
-		//double loopDeltaSlop = GetZeroTwoPiSlope(ACoordinates, NodeA->GetPreviousNode().Get2DPoint(EGridSpace::Default2D, Grid), SlopWithNextloop);
-		//double SegmentSlop = GetZeroTwoPiSlope(ACoordinates, BCoordinates, SlopWithNextloop);
-		//if (SegmentSlop < FlatAngle || SegmentSlop + FlatAngle > loopDeltaSlop)
-		//{
-		//	return;
-		//}
-
-		// Is Outside and not to flat at NodeB
-		if (!IsPointPInsideSectorABC(NodeB->GetPreviousNode().Get2DPoint(EGridSpace::Default2D, Grid), BCoordinates, NodeB->GetNextNode().Get2DPoint(EGridSpace::Default2D, Grid), ACoordinates, FlatAngle))
-		{
-			return;
-		}
-
-		//SlopWithNextloop = GetZeroTwoPiSlope(BCoordinates, NodeB->GetNextNode().Get2DPoint(EGridSpace::Default2D, Grid), 0);
-		//loopDeltaSlop = GetZeroTwoPiSlope(BCoordinates, NodeB->GetPreviousNode().Get2DPoint(EGridSpace::Default2D, Grid), SlopWithNextloop);
-		//SegmentSlop = GetZeroTwoPiSlope(BCoordinates, ACoordinates, SlopWithNextloop);
-
-		//if (SegmentSlop < FlatAngle || SegmentSlop + FlatAngle > loopDeltaSlop)
-		//{
-		//	return;
-		//}
-
-		FIsoSegment& Segment = IsoSegmentFactory.New();
-		Segment.Init(*NodeA, *NodeB, ESegmentType::LoopToLoop);
-		CandidateLoopToLoopSegments.Add(&Segment);
-	};
-
-	//Open3DDebugSession(TEXT("FIsoTrianguler::FindSegmentToLinkloopToloop"));
-
-	// find closed points
-	int32 IndexU = 0;
-	int32 IndexV = 0;
-
-	for (int32 IndexA = 0; IndexA < SortedLoopNodes.Num() - 1; ++IndexA)
-	{
-		FLoopNode* NodeA = SortedLoopNodes[IndexA];
-		if (NodeA->GetConnectedSegments().Num() > 2)
-		{
-			// Enough connected
-			continue;
-		}
-
-		// check intersection with
-		const FPoint2D& ACoordinates = NodeA->Get2DPoint(EGridSpace::Default2D, Grid);
-
-		double DeltaUMax = GetDeltaCoordinate(Grid.GetCuttingCoordinatesAlongIso(EIso::IsoU), IndexU, ACoordinates.U);
-		//double DeltaUMin = DeltaUMax * 0.1;
-		DeltaUMax *= 1.75;
-		double DeltaVMax = GetDeltaCoordinate(Grid.GetCuttingCoordinatesAlongIso(EIso::IsoV), IndexV, ACoordinates.V);
-		//double DeltaVMin = DeltaVMax * 0.1;
-		DeltaVMax *= 1.75;
-
-		for (int32 IndexB = IndexA + 1; IndexB < SortedLoopNodes.Num(); ++IndexB)
-		{
-			FLoopNode* NodeB = SortedLoopNodes[IndexB];
-			const FPoint2D& BCoordinates = NodeB->Get2DPoint(EGridSpace::Default2D, Grid);
-			FPoint2D AB = BCoordinates - ACoordinates;
-			double Diagonal = FMath::Abs(AB.U + AB.V);
-			if (Diagonal > DiagonalMax)
-			{
-				break;
-			}
-
-			if (FMath::Abs(AB.U) < DeltaUMax && FMath::Abs(AB.V) < DeltaVMax)
-			{
-				CreateSegment(NodeA, ACoordinates, NodeB, BCoordinates);
-			}
-		}
-	}
-	//Close3DDebugSession();
-	Chronos.FindSegmentToLinkLoopToLoopDuration += FChrono::Elapse(StartTime);
-}
-
-void FIsoTriangulator::ConnectCellSubLoopsByNeighborhood()
-{
-	FTimePoint StartTime = FChrono::Now();
-	//F3DDebugSession DelaunayDebugSession(TEXT("Delaunay Algo"));
-
-	int32 LoopCount = Grid.GetLoops3D().Num();
-
-	TArray<FPoint2D> LoopBarycenters;
-	LoopBarycenters.Reserve(LoopCount + 4);
-	LoopBarycenters.Emplace(FPoint2D::ZeroPoint); // outer loop is not defined
-
-	TArray<DelaunayTriangle> TriangleSet; // Index of the Vertex
-	TriangleSet.Reserve(3 * LoopCount);
-
-	TArray<int32> TriangleIndices;
-	TriangleIndices.Reserve(LoopCount);
-
-	TArray<int32> EdgeVertexIndices;
-	EdgeVertexIndices.Reserve(4 * LoopCount);
-	TArray<int32> OuterEdgeIndices;
-	OuterEdgeIndices.Reserve(2 * LoopCount);
-
-	TArray<int32> OuterLoopIndices;
-
-#ifdef DEBUG_DELAUNAY
-	TFunction<void()> DisplayEdges = [&]()
-	{
-		F3DDebugSession _(TEXT("Edges"));
-		for (int32 Index = 0; Index < EdgeVertexIndices.Num(); Index += 2)
-		{
-			if (OuterEdgeIndices[Index / 2])
-			{
-				DisplaySegment(LoopBarycenters[EdgeVertexIndices[Index]], LoopBarycenters[EdgeVertexIndices[Index + 1]], 0, EVisuProperty::YellowCurve);
-			}
-			else
-			{
-				DisplaySegment(LoopBarycenters[EdgeVertexIndices[Index]], LoopBarycenters[EdgeVertexIndices[Index + 1]]);
-			}
-		}
-
-		for (int32 Index = 0; Index < OuterLoopIndices.Num(); Index += 2)
-		{
-			DisplayPoint(LoopBarycenters[OuterLoopIndices[Index]], EVisuProperty::GreenPoint);
-		}
-	};
-
-	TFunction<void()> DisplayTriangles = [&]()
-	{
-		F3DDebugSession _(TEXT("Triangles"));
-		for (int32 Index = 0; Index < TriangleSet.Num(); Index++)
-		{
-			DisplaySegment(LoopBarycenters[TriangleSet[Index].VertexIndex[0]], LoopBarycenters[TriangleSet[Index].VertexIndex[1]]);
-			DisplaySegment(LoopBarycenters[TriangleSet[Index].VertexIndex[1]], LoopBarycenters[TriangleSet[Index].VertexIndex[2]]);
-			DisplaySegment(LoopBarycenters[TriangleSet[Index].VertexIndex[2]], LoopBarycenters[TriangleSet[Index].VertexIndex[0]]);
-		}
-		//Wait();
-	};
-#endif
-
-	FAABB2D LoobsBBox;
-	const TArray<TArray<FPoint2D>>& Loops = Grid.GetLoops2D(EGridSpace::UniformScaled);
-	for (int32 LoopIndex = 1; LoopIndex < LoopCount; ++LoopIndex)
-	{
-		const TArray<FPoint2D>& LoopPoints = Loops[LoopIndex];
-		FPoint2D& BaryCenter = LoopBarycenters.Emplace_GetRef(FPoint2D::ZeroPoint);
-		for (const FPoint2D& Point : LoopPoints)
-		{
-			BaryCenter += Point;
-		}
-		BaryCenter /= (double)LoopPoints.Num();
-		LoobsBBox += BaryCenter;
-	}
-
-	LoobsBBox.Offset(LoobsBBox.DiagonalLength());
-	LoopBarycenters.Emplace(LoobsBBox.GetCorner(3));
-	LoopBarycenters.Emplace(LoobsBBox.GetCorner(2));
-	LoopBarycenters.Emplace(LoobsBBox.GetCorner(0));
-	LoopBarycenters.Emplace(LoobsBBox.GetCorner(1));
-
-
-#ifdef DEBUG_DELAUNAY
-	{
-		F3DDebugSession _(TEXT("LoopBarycenters"));
-		for (int32 LoopIndex = 1; LoopIndex < LoopCount; ++LoopIndex)
-		{
-			DisplayPoint(LoopBarycenters[LoopIndex], EVisuProperty::YellowPoint);
-		}
-		DisplayAABB2D(LoobsBBox);
-
-		for (int32 LoopIndex = LoopCount; LoopIndex < LoopCount + 4; ++LoopIndex)
-		{
-			DisplayPoint(LoopBarycenters[LoopIndex], EVisuProperty::GreenPoint);
-		}
-	}
-#endif
-
-	TriangleSet.Emplace(LoopCount, LoopCount + 1, LoopCount + 2, LoopBarycenters);
-	TriangleSet.Emplace(LoopCount + 2, LoopCount + 3, LoopCount, LoopBarycenters);
-
-#ifdef DEBUG_DELAUNAY
-	DisplayTriangles();
-#endif
-
-	for (int32 CenterIndex = 1; CenterIndex < LoopCount; ++CenterIndex)
-	{
-		TriangleIndices.Empty(LoopCount);
-		const FPoint2D& NewVertex = LoopBarycenters[CenterIndex];
-		for (int32 Tndex = 0; Tndex < TriangleSet.Num(); Tndex++)
-		{
-			const FPoint2D Center = TriangleSet[Tndex].Center;
-			double SquareRadius = TriangleSet[Tndex].SquareRadius;
-			double SquareDistanceToCenter = Center.SquareDistance(NewVertex);
-			if (SquareDistanceToCenter < SquareRadius)
-			{
-				TriangleIndices.Add(Tndex);
-			}
-		}
-
-		EdgeVertexIndices.Empty(LoopCount);
-		for (int32 TriangleIndex : TriangleIndices)
-		{
-			int32 EndVertex = TriangleSet[TriangleIndex].VertexIndex[2];
-			for (int32 Index = 0; Index < 3; Index++)
-			{
-				int32 StartVertex = TriangleSet[TriangleIndex].VertexIndex[Index];
-				int32 Endex = 0;
-				for (; Endex < EdgeVertexIndices.Num(); Endex += 2)
-				{
-					if (EdgeVertexIndices[Endex] == EndVertex && EdgeVertexIndices[Endex + 1] == StartVertex)
-					{
-						EdgeVertexIndices[Endex] = -1;
-						EdgeVertexIndices[Endex + 1] = -1;
-						break;
-					}
-				}
-				if (Endex == EdgeVertexIndices.Num())
-				{
-					EdgeVertexIndices.Add(StartVertex);
-					EdgeVertexIndices.Add(EndVertex);
-				}
-				EndVertex = StartVertex;
-			}
-		}
-
-		int32 Endex = 0;
-		for (int32 Tndex : TriangleIndices)
-		{
-			while (EdgeVertexIndices[Endex] < 0)
-			{
-				Endex += 2;
-			}
-			TriangleSet[Tndex].Set(EdgeVertexIndices[Endex + 1], EdgeVertexIndices[Endex], CenterIndex, LoopBarycenters);
-			Endex += 2;
-		}
-		for (; Endex < EdgeVertexIndices.Num(); Endex += 2)
-		{
-			if (EdgeVertexIndices[Endex] < 0)
-			{
-				continue;
-			}
-			TriangleSet.Emplace(EdgeVertexIndices[Endex + 1], EdgeVertexIndices[Endex], CenterIndex, LoopBarycenters);
-		}
-#ifdef DEBUG_DELAUNAY
-		//DisplayTriangles();
-#endif
-	}
-
-#ifdef DEBUG_DELAUNAY
-	Wait();
-
-	DisplayTriangles();
-	Wait();
-#endif
-
-	EdgeVertexIndices.Reserve(LoopCount);
-
-	EdgeVertexIndices.Empty(LoopCount);
-	for (int32 TriangleIndex = 0; TriangleIndex < TriangleSet.Num(); TriangleIndex++)
-	{
-		int32 EndVertex = TriangleSet[TriangleIndex].VertexIndex[2];
-		int32 Index = 0;
-		for (; Index < 3; Index++)
-		{
-			if (TriangleSet[TriangleIndex].VertexIndex[Index] >= LoopCount)
-			{
-				break;
-			}
-		}
-
-		// the triangle connect one of the first 4 nodes
-		if (Index < 3)
-		{
-
-			continue;
-		}
-
-		for (Index = 0; Index < 3; Index++)
-		{
-			int32 StartVertex = TriangleSet[TriangleIndex].VertexIndex[Index];
-			int32 Endex = 0;
-			for (; Endex < EdgeVertexIndices.Num(); Endex += 2)
-			{
-				if (EdgeVertexIndices[Endex] == EndVertex && EdgeVertexIndices[Endex + 1] == StartVertex)
-				{
-					OuterEdgeIndices[Endex / 2] = 0;
-					break;
-				}
-			}
-			if (Endex == EdgeVertexIndices.Num())
-			{
-				EdgeVertexIndices.Add(StartVertex);
-				EdgeVertexIndices.Add(EndVertex);
-				OuterEdgeIndices.Add(1);
-			}
-			EndVertex = StartVertex;
-		}
-	}
-
-	OuterLoopIndices.Reserve(LoopCount * 2);
-	for (int32 Index = 0; Index < EdgeVertexIndices.Num(); Index += 2)
-	{
-		if (OuterEdgeIndices[Index / 2])
-		{
-			OuterLoopIndices.Add(EdgeVertexIndices[Index]);
-			OuterLoopIndices.Add(EdgeVertexIndices[Index + 1]);
-		}
-	}
-	OuterLoopIndices.Sort();
-
-#ifdef DEBUG_DELAUNAY
-	DisplayEdges();
-	Wait();
-#endif
-
-	TFunction<void(FLoopNode&, const FPoint2D&, FLoopNode&, const FPoint2D&)> TryCreateSegment = [&](FLoopNode& NodeA, const FPoint2D& ACoordinates, FLoopNode& NodeB, const FPoint2D& BCoordinates)
-	{
-		if (NodeA.GetSegmentConnectedTo(&NodeB))
-		{
-			DisplaySegment(NodeA.Get2DPoint(EGridSpace::UniformScaled, Grid), NodeB.Get2DPoint(EGridSpace::UniformScaled, Grid));
-			return;
-		}
-
-		if (InnerSegmentsIntersectionTool.DoesIntersect(ACoordinates, BCoordinates))
-		{
-			DisplaySegment(NodeA.Get2DPoint(EGridSpace::UniformScaled, Grid), NodeB.Get2DPoint(EGridSpace::UniformScaled, Grid));
-			return;
-		}
-
-		if (InnerToLoopSegmentsIntersectionTool.DoesIntersect(NodeA, NodeB))
-		{
-			DisplaySegment(NodeA.Get2DPoint(EGridSpace::UniformScaled, Grid), NodeB.Get2DPoint(EGridSpace::UniformScaled, Grid), 0, EVisuProperty::YellowCurve);
-			return;
-		}
-
-		if (LoopSegmentsIntersectionTool.DoesIntersect(NodeA, NodeB))
-		{
-			DisplaySegment(NodeA.Get2DPoint(EGridSpace::UniformScaled, Grid), NodeB.Get2DPoint(EGridSpace::UniformScaled, Grid), 0, EVisuProperty::YellowCurve);
-			return;
-		}
-
-		// Is Outside and not to flat at NodeA
-		const double FlatAngle = 1;
-		//if (!IsPointPBeInsideSectorABC(NodeA.GetPreviousNode().Get2DPoint(EGridSpace::Default2D, Grid), ACoordinates, NodeA.GetNextNode().Get2DPoint(EGridSpace::Default2D, Grid), BCoordinates, FlatAngle))
-
-		if (NodeA.IsSegmentBeInsideFace(BCoordinates, Grid, FlatAngle))
-		{
-			DisplaySegment(NodeA.Get2DPoint(EGridSpace::UniformScaled, Grid), NodeB.Get2DPoint(EGridSpace::UniformScaled, Grid), 0, EVisuProperty::YellowCurve);
-			return;
-		}
-
-		// Is Outside and not to flat at NodeB
-		//if (!IsPointPBeInsideSectorABC(NodeB.GetPreviousNode().Get2DPoint(EGridSpace::Default2D, Grid), BCoordinates, NodeB.GetNextNode().Get2DPoint(EGridSpace::Default2D, Grid), ACoordinates, FlatAngle))
-		if (NodeB.IsSegmentBeInsideFace(ACoordinates, Grid, FlatAngle))
-		{
-			DisplaySegment(NodeA.Get2DPoint(EGridSpace::UniformScaled, Grid), NodeB.Get2DPoint(EGridSpace::UniformScaled, Grid), 0, EVisuProperty::YellowCurve);
-			return;
-		}
-
-		FIsoSegment& Segment = IsoSegmentFactory.New();
-		Segment.Init(NodeA, NodeB, ESegmentType::LoopToLoop);
-		CandidateLoopToLoopSegments.Add(&Segment);
-#ifdef DEBUG_DELAUNAY
-		DisplaySegment(NodeA.Get2DPoint(EGridSpace::UniformScaled, Grid), NodeB.Get2DPoint(EGridSpace::UniformScaled, Grid));
-#endif
-	};
-
-	TFunction<void(int32, int32)> FindAndTryCreateCandidateSegmentToLinkLoops = [&](int32 IndexLoopA, int32 IndexLoopB)
-	{
-		int32 StartIndexLoopA = LoopStartIndex[IndexLoopA];
-		int32 StartIndexLoopB = LoopStartIndex[IndexLoopB];
-
-		double MinDistanceSquare = HUGE_VALUE_SQUARE;
-		int32 MinIndexA = -1;
-		int32 MinIndexB = -1;
-		for (int32 IndexA = StartIndexLoopA; IndexA < LoopNodeCount - 1; ++IndexA)
-		{
-			FLoopNode& NodeA = LoopNodes[IndexA];
-			if (NodeA.GetLoopIndex() != IndexLoopA)
-			{
-				break;
-			}
-
-			const FPoint2D& ACoordinates = NodeA.Get2DPoint(EGridSpace::UniformScaled, Grid);
-
-			for (int32 IndexB = StartIndexLoopB; IndexB < LoopNodeCount; ++IndexB)
-			{
-				FLoopNode& NodeB = LoopNodes[IndexB];
-				if (NodeB.GetLoopIndex() != IndexLoopB)
-				{
-					break;
-				}
-				const FPoint2D& BCoordinates = NodeB.Get2DPoint(EGridSpace::UniformScaled, Grid);
-
-				double SquareDistance = ACoordinates.SquareDistance(BCoordinates);
-				if (SquareDistance < MinDistanceSquare)
-				{
-					MinDistanceSquare = SquareDistance;
-					MinIndexA = IndexA;
-					MinIndexB = IndexB;
-				}
-			}
-		}
-
-		if (MinIndexA >= 0 && MinIndexB >= 0)
-		{
-			FLoopNode& NodeA = LoopNodes[MinIndexA];
-			const FPoint2D& ACoordinates = NodeA.Get2DPoint(EGridSpace::UniformScaled, Grid);
-			FLoopNode& NodeB = LoopNodes[MinIndexB];
-			const FPoint2D& BCoordinates = NodeB.Get2DPoint(EGridSpace::UniformScaled, Grid);
-
-			//DisplaySegment(NodeA.Get2DPoint(EGridSpace::UniformScaled, Grid), NodeB.Get2DPoint(EGridSpace::UniformScaled, Grid));
-
-			TryCreateSegment(NodeA, ACoordinates, NodeB, BCoordinates);
-		}
-
-	};
-
-	// Between inner loops
-	{
-#ifdef DEBUG_DELAUNAY
-		Open3DDebugSession(TEXT("Build Segments"));
-#endif
-		for (int32 Index = 0; Index < EdgeVertexIndices.Num(); Index += 2)
-		{
-			//Open3DDebugSession(TEXT("Segment"));
-			int32 IndexLoopA = EdgeVertexIndices[Index];
-			int32 IndexLoopB = EdgeVertexIndices[Index + 1];
-
-			//DisplaySegment(LoopBarycenters[IndexLoopA], LoopBarycenters[IndexLoopB]);
-
-			FindAndTryCreateCandidateSegmentToLinkLoops(IndexLoopA, IndexLoopB);
-			//Close3DDebugSession();
-			//Wait();
-		}
-#ifdef DEBUG_DELAUNAY
-		Close3DDebugSession();
-		Wait();
-#endif
-	}
-
-	// With Outer loop
-	{
-#ifdef DEBUG_DELAUNAY
-		Open3DDebugSession(TEXT("Build Segments"));
-#endif
-		for (int32 Index = 0; Index < OuterLoopIndices.Num(); Index += 2)
-		{
-			int32 IndexLoopA = OuterLoopIndices[Index];
-			int32 IndexLoopB = 0;
-
-			FindAndTryCreateCandidateSegmentToLinkLoops(IndexLoopA, IndexLoopB);
-		}
-#ifdef DEBUG_DELAUNAY
-		Close3DDebugSession();
-#endif
-	}
-
-#ifdef CADKERNEL_DEV
-	Chronos.FindSegmentToLinkLoopToLoopByDelaunayDuration += FChrono::Elapse(StartTime);
-#endif
-}
-
-
-#endif
-
-
-
 } //namespace CADKernel
\ No newline at end of file
diff --git a/Engine/Source/Runtime/Datasmith/CADKernel/Private/CADKernel/Mesh/Meshers/IsoTriangulator.h b/Engine/Source/Runtime/Datasmith/CADKernel/Private/CADKernel/Mesh/Meshers/IsoTriangulator.h
--- a/Engine/Source/Runtime/Datasmith/CADKernel/Private/CADKernel/Mesh/Meshers/IsoTriangulator.h
+++ b/Engine/Source/Runtime/Datasmith/CADKernel/Private/CADKernel/Mesh/Meshers/IsoTriangulator.h
@@ -41,7 +41,6 @@ struct CADKERNEL_API FIsoTriangulatorChronos
 	FDuration FindSegmentToLinkLoopToLoopDuration = FChrono::Init();
 	FDuration FindSegmentToLinkLoopToLoopByDelaunayDuration = FChrono::Init();
 	FDuration FindSegmentToLinkInnerToLoopDuration = FChrono::Init();
-	FDuration SelectSegmentToLinkInnerToLoopsDuration = FChrono::Init();
 	FDuration TriangulateOverCycleDuration = FChrono::Init();
 	FDuration TriangulateInnerNodesDuration = FChrono::Init();
 
@@ -59,7 +58,6 @@ struct CADKERNEL_API FIsoTriangulatorChronos
 		IsoTriangulerDuration += FindSegmentToLinkLoopToLoopDuration;
 		IsoTriangulerDuration += FindSegmentToLinkLoopToLoopByDelaunayDuration;
 		IsoTriangulerDuration += FindSegmentToLinkInnerToLoopDuration;
-		IsoTriangulerDuration += SelectSegmentToLinkInnerToLoopsDuration;
 		IsoTriangulerDuration += TriangulateOverCycleDuration;
 		IsoTriangulerDuration += TriangulateInnerNodesDuration;
 
@@ -74,7 +72,6 @@ struct CADKERNEL_API FIsoTriangulatorChronos
 		FChrono::PrintClockElapse(Log, TEXT("    "), TEXT("Find Segment ToLink LoopToLoop by Delaunay"), FindSegmentToLinkLoopToLoopByDelaunayDuration);
 		FChrono::PrintClockElapse(Log, TEXT("    "), TEXT("Find Segment ToLink LoopToLoop"), FindSegmentToLinkLoopToLoopDuration);
 		FChrono::PrintClockElapse(Log, TEXT("    "), TEXT("Find Segment ToLink InnerToLoop"), FindSegmentToLinkInnerToLoopDuration);
-		FChrono::PrintClockElapse(Log, TEXT("    "), TEXT("Select Segment ToLink InnerToLoop"), SelectSegmentToLinkInnerToLoopsDuration);
 		FChrono::PrintClockElapse(Log, TEXT("    "), TEXT("Mesh Over Cycle"), TriangulateOverCycleDuration);
 		FChrono::PrintClockElapse(Log, TEXT("    "), TEXT("Mesh Inner Nodes"), TriangulateInnerNodesDuration);
 		FChrono::PrintClockElapse(Log, TEXT("  "), TEXT("Triangulate1"), TriangulateDuration1);
@@ -150,16 +147,6 @@ protected:
 	 */
 	TArray<FIsoSegment*> CandidateSegments;
 
-	/**
-	 * Segments to link boundary to boundary and to complete the mesh
-	 */
-	TArray<FIsoSegment*> CandidateLoopToLoopSegments;
-
-	/**
-	 * Segments to link inner to boundary and to complete the mesh
-	 */
-	TArray<FIsoSegment*> CandidateInnerToLoopSegments;
-
 	TArray<FIsoSegment*> NewTestSegments;
 
 	bool bDisplay = false;
@@ -244,31 +231,19 @@ public:
 
 	/**
 	 * The closest loops are connected together
-	 * To do it, a Delaunay triangulation of the loop barycenters is realized.
+	 * To do it, a Delaunay triangulation of the loop barycenter is realized.
 	 * Each edge of this mesh defined a near loops pair
 	 * The shortest segment is then build between this pair of loops
 	 */
 	void ConnectCellSubLoopsByNeighborhood(FCell& cell);
 
-	/**
-	 * 2nd step
-	 */
-	void ConnectCellCornerToInnerLoop(FCell& Cell);
-
-	/**
-	 * 3rd step : IsoSegment linking
-	 */
-	void FindIsoSegmentToLinkOuterLoopNodes(FCell& Cell);
+	void FindIsoSegmentToLinkOuterLoopNodes(FCell& Cell);	
 
-	/**
-	 * 4th step : If their is no segment candidate
-	 */
 	void FindSegmentToLinkOuterLoopNodes(FCell& Cell);
 
-	/**
-	 * 5th step : try to connect outer loop Extremities
-	 */
-	void AddSementToLinkOuterLoopExtremities(FCell& Cell);
+	void FindSegmentToLinkOuterToInnerLoopNodes(FCell& Cell);
+
+	void ConnectCellCornerToInnerLoop(FCell& Cell);
 
 	/**
 	 * The goal of this algorithm is to connect iso U (or V) aligned loop nodes as soon as they are nearly in the same iso V (or U) strip.
@@ -277,18 +252,6 @@ public:
 	 * - In the same strip: each node of the segment has the same index "i" that verify: isoV[i] - TolV < Node.V < isoV[i+1] + TolV
 	 */
 	void FindIsoSegmentToLinkLoopToLoop();
-#ifdef UNUSED_TO_DELETE_
-	void LastChanceToCreateSegmentInCell(FCell& Cell);
-	/**
-	 * The goal of this algorithm is to connect inner node (node of the grid UV) to iso aligned loop node when they are in the same iso V (or U) strip.
-	 * I.e.:
-	 * - Iso U aligned: NodeA.U < NodeB.U
-	 * - In the same strip: each node of the segment has the same index "i" that verify: isoV[i-1] - TolV < BoundaryNode.V or BoundaryNode.V < isoV[i+1] + TolV
-	 * https://docs.google.com/presentation/d/1qUVOH-2kU_QXBVKyRUcdDy1Y6WGkcaJCiaS8wGjSZ6M/edit?usp=sharing
-	 * Slide "Find Iso Segment To Link Inner To Loop"
-	 */
-	void FindIsoSegmentToLinkInnerToLoop();
-#endif
 
 	/**
 	 * The purpose of the method is select a minimal set of segments connecting loops together
@@ -296,17 +259,6 @@ public:
 	 */
 	void ConnectCellSubLoopsByNeighborhood();  // to rename and clean
 
-#ifdef TODELETE
-	void FindSegmentToLinkLoopToLoop();
-#endif
-
-	/**
-	 * The purpose of the method is select a minimal set of segments
-	 * The final segments will be selected with SelectSegmentInCandidateSegments
-	 */
-	void FindCandidateSegmentsToLinkInnerToLoop();
-
-
 	void FindCandidateSegmentsToLinkInnerAndLoop();
 
 	/**
@@ -337,6 +289,10 @@ public:
 	 */
 	void MeshCycle(const EGridSpace Space, const TArray<FIsoSegment*>& cycle, const TArray<bool>& cycleOrientation);
 
+#ifdef WIP_ADD_STEP_TO_TO_FAVOR_ISO_SEGMENTS
+	void FindIsoCandidateSegmentInCycle(TArray<FIsoNode*> CycleNodes);
+#endif
+
 	bool CanCycleBeMeshed(const TArray<FIsoSegment*>& Cycle, FIntersectionSegmentTool& CycleIntersectionTool);
 
 	/**
@@ -367,11 +323,6 @@ private:
 
 	FIsoSegment* FindNextSegment(EGridSpace Space, const FIsoSegment* StartSegment, const FIsoNode* StartNode, SlopMethod GetSlop) const;
 
-
-#ifdef UNUSED_TO_DELETE_
-	void CreateLoopToLoopSegment(FLoopNode* NodeA, const FPoint2D& ACoordinates, FLoopNode* NodeB, const FPoint2D& BCoordinates, TArray<FIsoSegment*>& NewSegments);
-#endif
-
 	// ==========================================================================================
 	// 	   Create segments
 	// ==========================================================================================
diff --git a/Engine/Source/Runtime/Datasmith/CADKernel/Private/CADKernel/Mesh/Meshers/IsoTriangulator/DefineForDebug.h b/Engine/Source/Runtime/Datasmith/CADKernel/Private/CADKernel/Mesh/Meshers/IsoTriangulator/DefineForDebug.h
--- a/Engine/Source/Runtime/Datasmith/CADKernel/Private/CADKernel/Mesh/Meshers/IsoTriangulator/DefineForDebug.h
+++ b/Engine/Source/Runtime/Datasmith/CADKernel/Private/CADKernel/Mesh/Meshers/IsoTriangulator/DefineForDebug.h
@@ -2,26 +2,31 @@
 // Refactor in progress
 #pragma once
 
+#ifdef CADKERNEL_DEV
+
 //#define DEBUG_CADKERNEL
 
 #ifdef DEBUG_CADKERNEL
-static int32 FaceToDebug = 70;
+static int32 FaceToDebug = 202910;
 
 #define ADD_TRIANGLE_2D
 #define DEBUG_ISOTRIANGULATOR
 
 // Loop cleaner
-#define DEBUG_CLEAN_LOOPS
-
-#define DEBUG_LOOP_INTERSECTION_AND_FIX_IT		
-#define DEBUG_REMOVE_LOOP_INTERSECTIONS
-#define DEBUG_REMOVE_UNIQUE_INTERSECTION
-#define DEBUG_SWAP_SEGMENTS_OR_REMOVE
-#define DEBUG_FIND_LOOP_INTERSECTIONS	
-#define DEBUG_REMOVE_INTERSECTIONS
+//#define DEBUG_CLEAN_LOOPS
+//#define DEBUG_LOOP_INTERSECTION_AND_FIX_IT		
+//#define DEBUG_REMOVE_LOOP_INTERSECTIONS
+//#define DEBUG_REMOVE_UNIQUE_INTERSECTION
+//#define DEBUG_SWAP_SEGMENTS_OR_REMOVE
+//#define DEBUG_FIND_LOOP_INTERSECTIONS	
+//#define DEBUG_REMOVE_INTERSECTIONS
 
 
 #define DEBUG_ONLY_SURFACE_TO_DEBUG
 static bool bDisplayDebugMeshStep=false;
 
+#define DEBUG_HAS_CYCLE
+
+#endif
+
 #endif
diff --git a/Engine/Source/Runtime/Datasmith/CADKernel/Private/CADKernel/Mesh/Meshers/IsoTriangulator/IntersectionSegmentTool.cpp b/Engine/Source/Runtime/Datasmith/CADKernel/Private/CADKernel/Mesh/Meshers/IsoTriangulator/IntersectionSegmentTool.cpp
--- a/Engine/Source/Runtime/Datasmith/CADKernel/Private/CADKernel/Mesh/Meshers/IsoTriangulator/IntersectionSegmentTool.cpp
+++ b/Engine/Source/Runtime/Datasmith/CADKernel/Private/CADKernel/Mesh/Meshers/IsoTriangulator/IntersectionSegmentTool.cpp
@@ -65,12 +65,7 @@ const FIsoSegment* FIntersectionSegmentTool::DoesIntersect(const FIsoNode& Start
 
 	for (const FSegment4IntersectionTools& Segment : Segments)
 	{
-		if (!Segment.IsoSegment)
-		{
-			continue;
-		}
-
-		if (Segment.IsoSegment != nullptr && Segment.IsoSegment->IsDelete())
+		if (!Segment.IsoSegment || Segment.IsoSegment->IsDelete())
 		{
 			continue;
 		}
diff --git a/Engine/Source/Runtime/Datasmith/CADKernel/Private/CADKernel/Mesh/Meshers/IsoTriangulator/IsoCell.h b/Engine/Source/Runtime/Datasmith/CADKernel/Private/CADKernel/Mesh/Meshers/IsoTriangulator/IsoCell.h
--- a/Engine/Source/Runtime/Datasmith/CADKernel/Private/CADKernel/Mesh/Meshers/IsoTriangulator/IsoCell.h
+++ b/Engine/Source/Runtime/Datasmith/CADKernel/Private/CADKernel/Mesh/Meshers/IsoTriangulator/IsoCell.h
@@ -109,7 +109,7 @@ namespace CADKernel
 		void SelectSegmentInCandidateSegments(TFactory<FIsoSegment>& SegmentFactory)
 		{
 #ifdef DEBUG_SELECT_SEGMENT
-			F3DDebugSession _(TEXT("SelectSegmentInCandidateSegments "));
+			F3DDebugSession _(Grid.bDisplay, TEXT("SelectSegmentInCandidateSegments "));
 			IntersectionTool.Display(TEXT("Cell.IntersectionTool at SelectSegmentInCandidateSegments start"));
 			//Wait();
 #endif
@@ -119,19 +119,17 @@ namespace CADKernel
 					return Segment1->Get2DLengthSquare(EGridSpace::UniformScaled, Grid) < Segment2->Get2DLengthSquare(EGridSpace::UniformScaled, Grid);
 				});
 
-			//IntersectionTool.Empty(IntersectionTool.Count());
-
 			// Validate all candidate segments
 			for (FIsoSegment* Segment : CandidateSegments)
 			{
 #ifdef DEBUG_SELECT_SEGMENT
-				F3DDebugSession _(TEXT("Segment"));
+				F3DDebugSession _(Grid.bDisplay, TEXT("Segment"));
 #endif
 				if (IntersectionTool.DoesIntersect(*Segment))
 				{
 #ifdef DEBUG_SELECT_SEGMENT
-					F3DDebugSession _(TEXT("SelectSegmentInCandidateSegments "));
-					DisplaySegment(Segment->GetFirstNode().GetPoint(EGridSpace::UniformScaled, Grid), Segment->GetSecondNode().GetPoint(EGridSpace::UniformScaled, Grid), 0, YellowCurve);
+					F3DDebugSession _(Grid.bDisplay, TEXT("SelectSegmentInCandidateSegments "));
+					Grid.DisplayIsoSegment(EGridSpace::UniformScaled, *Segment, YellowCurve);
 #endif
 					SegmentFactory.DeleteEntity(Segment);
 					continue;
@@ -140,14 +138,14 @@ namespace CADKernel
 				if (FIsoSegment::IsItAlreadyDefined(&Segment->GetFirstNode(), &Segment->GetSecondNode()))
 				{
 #ifdef DEBUG_SELECT_SEGMENT
-					DisplaySegment(Segment->GetFirstNode().GetPoint(EGridSpace::UniformScaled, Grid), Segment->GetSecondNode().GetPoint(EGridSpace::UniformScaled, Grid), 0, GreenCurve);
+					Grid.DisplayIsoSegment(EGridSpace::UniformScaled, *Segment, GreenCurve);
 #endif
 					SegmentFactory.DeleteEntity(Segment);
 					continue;
 				}
 
 #ifdef DEBUG_SELECT_SEGMENT
-				DisplaySegment(Segment->GetFirstNode().GetPoint(EGridSpace::UniformScaled, Grid), Segment->GetSecondNode().GetPoint(EGridSpace::UniformScaled, Grid), 0, BlueCurve);
+				Grid.DisplayIsoSegment(EGridSpace::UniformScaled, *Segment, BlueCurve);
 #endif
 
 				FinalSegments.Add(Segment);
@@ -165,7 +163,6 @@ namespace CADKernel
 			return SubLoops[Index].Find(NodeToFind) != INDEX_NONE;
 		}
 
-
 	};
 
 } // namespace CADKernel
diff --git a/Engine/Source/Runtime/Datasmith/CADKernel/Private/CADKernel/Mesh/Meshers/IsoTriangulator/IsoNode.h b/Engine/Source/Runtime/Datasmith/CADKernel/Private/CADKernel/Mesh/Meshers/IsoTriangulator/IsoNode.h
--- a/Engine/Source/Runtime/Datasmith/CADKernel/Private/CADKernel/Mesh/Meshers/IsoTriangulator/IsoNode.h
+++ b/Engine/Source/Runtime/Datasmith/CADKernel/Private/CADKernel/Mesh/Meshers/IsoTriangulator/IsoNode.h
@@ -145,21 +145,6 @@ public:
 		}
 	}
 
-	bool IsLinkedToLoopInQuarter(int32 Quarter) const // [0 Pi/2]
-	{
-		switch (Quarter)
-		{
-		case 0:
-			return (States & EIsoNodeStates::FirstQuarter) == EIsoNodeStates::FirstQuarter;
-		case 1:
-			return (States & EIsoNodeStates::SecondQuarter) == EIsoNodeStates::SecondQuarter;
-		case 2:
-			return (States & EIsoNodeStates::ThirdQuarter) == EIsoNodeStates::ThirdQuarter;
-		default:
-			return (States & EIsoNodeStates::FourthQuarter) == EIsoNodeStates::FourthQuarter;
-		}
-	}
-
 	bool IsLinkedToLoopInNearlyIso(int32 Iso) const
 	{
 		if (Iso > 3)
@@ -385,10 +370,9 @@ public:
 	 *    |                                |        |                              |
 	 *    X---------------X----------------X        X------------------X-----------X
 	 */
-	bool IsSegmentBeInsideFace(const FPoint2D& EndSegmentCoordinate, const FGrid& Grid, const double FlatAngle)
+	bool IsSegmentBeInsideFace(const FPoint2D& EndSegmentCoordinate, const FGrid& Grid, const double FlatAngle) const
 	{
-		//bool bIsOuterLoop = (GetLoopIndex() == 0);
-		return (/*bIsOuterLoop  == */!IsPointPInsideSectorABC(GetPreviousNode().Get2DPoint(EGridSpace::UniformScaled, Grid), Get2DPoint(EGridSpace::UniformScaled, Grid), GetNextNode().Get2DPoint(EGridSpace::UniformScaled, Grid), EndSegmentCoordinate, FlatAngle));
+		return (!IsPointPInsideSectorABC(GetPreviousNode().Get2DPoint(EGridSpace::UniformScaled, Grid), Get2DPoint(EGridSpace::UniformScaled, Grid), GetNextNode().Get2DPoint(EGridSpace::UniformScaled, Grid), EndSegmentCoordinate, FlatAngle));
 	}
 
 	void SetAsIsoU()
diff --git a/Engine/Source/Runtime/Datasmith/CADKernel/Private/CADKernel/Mesh/Meshers/IsoTriangulator/IsoSegment.h b/Engine/Source/Runtime/Datasmith/CADKernel/Private/CADKernel/Mesh/Meshers/IsoTriangulator/IsoSegment.h
--- a/Engine/Source/Runtime/Datasmith/CADKernel/Private/CADKernel/Mesh/Meshers/IsoTriangulator/IsoSegment.h
+++ b/Engine/Source/Runtime/Datasmith/CADKernel/Private/CADKernel/Mesh/Meshers/IsoTriangulator/IsoSegment.h
@@ -228,6 +228,11 @@ public:
 		States |= EIsoSegmentStates::Degenerate;
 	}
 
+	void ResetDegenerated()
+	{
+		States &= ~EIsoSegmentStates::Degenerate;
+	}
+
 	bool IsDegenerated() const
 	{
 		return (States & EIsoSegmentStates::Degenerate) == EIsoSegmentStates::Degenerate;
diff --git a/Engine/Source/Runtime/Datasmith/CADKernel/Private/CADKernel/Mesh/Meshers/IsoTriangulator/IsoTriangulatorDisplay.cpp b/Engine/Source/Runtime/Datasmith/CADKernel/Private/CADKernel/Mesh/Meshers/IsoTriangulator/IsoTriangulatorDisplay.cpp
--- a/Engine/Source/Runtime/Datasmith/CADKernel/Private/CADKernel/Mesh/Meshers/IsoTriangulator/IsoTriangulatorDisplay.cpp
+++ b/Engine/Source/Runtime/Datasmith/CADKernel/Private/CADKernel/Mesh/Meshers/IsoTriangulator/IsoTriangulatorDisplay.cpp
@@ -41,8 +41,8 @@ void FIsoTriangulator::DisplayPixel(const int32 IndexU, const int32 IndexV) cons
 
 void FIsoTriangulator::DisplayPixel(const int32 Index) const
 {
-	FPoint Point = (Grid.GetInner2DPoint(EGridSpace::Default2D, Index) + Grid.GetInner2DPoint(EGridSpace::Default2D, Index + Grid.GetCuttingCount(EIso::IsoU) + 1)) * 0.5;
-	DisplayPoint(Point, EVisuProperty::GreenPoint);
+	FPoint Point = (Grid.GetInner2DPoint(EGridSpace::UniformScaled, Index) + Grid.GetInner2DPoint(EGridSpace::UniformScaled, Index + Grid.GetCuttingCount(EIso::IsoU) + 1)) * 0.5;
+	DisplayPoint(Point * Grid.DisplayScale, EVisuProperty::GreenPoint);
 };
 
 void FIsoTriangulator::DisplayPixels(TArray<uint8>& Pixel) const
@@ -52,7 +52,7 @@ void FIsoTriangulator::DisplayPixels(TArray<uint8>& Pixel) const
 		return;
 	}
 
-	Open3DDebugSession(TEXT("FIsoTrianguler::Pixel"));
+	F3DDebugSession _(TEXT("FIsoTrianguler::Pixel"));
 	for (int32 Index = 0; Index <Grid.GetTotalCuttingCount(); ++Index)
 	{
 		if (Pixel[Index])
@@ -60,7 +60,6 @@ void FIsoTriangulator::DisplayPixels(TArray<uint8>& Pixel) const
 			DisplayPixel(Index);
 		}
 	}
-	Close3DDebugSession();
 }
 
 void FIsoTriangulator::DisplayLoops(const TCHAR* Message, bool bOneNode, bool bSplitBySegment) const
@@ -149,15 +148,15 @@ void FIsoTriangulator::DrawCellBoundary(int32 Index, EVisuProperty Property) con
 	int32 IndexV;
 	Grid.UVIndexFromGlobalIndex(Index, IndexU, IndexV);
 
-	DisplayPoint(Grid.GetInner2DPoint(EGridSpace::UniformScaled, IndexU, IndexV), Property, Index);
-	DisplayPoint(Grid.GetInner2DPoint(EGridSpace::UniformScaled, IndexU, IndexV + 1), Property, Index);
-	DisplayPoint(Grid.GetInner2DPoint(EGridSpace::UniformScaled, IndexU + 1, IndexV), Property, Index);
-	DisplayPoint(Grid.GetInner2DPoint(EGridSpace::UniformScaled, IndexU + 1, IndexV + 1), Property, Index);
+	DisplayPoint(Grid.GetInner2DPoint(EGridSpace::UniformScaled, IndexU, IndexV) * Grid.DisplayScale, Property, Index);
+	DisplayPoint(Grid.GetInner2DPoint(EGridSpace::UniformScaled, IndexU, IndexV + 1) * Grid.DisplayScale, Property, Index);
+	DisplayPoint(Grid.GetInner2DPoint(EGridSpace::UniformScaled, IndexU + 1, IndexV) * Grid.DisplayScale, Property, Index);
+	DisplayPoint(Grid.GetInner2DPoint(EGridSpace::UniformScaled, IndexU + 1, IndexV + 1) * Grid.DisplayScale, Property, Index);
 
-	DisplaySegment(Grid.GetInner2DPoint(EGridSpace::UniformScaled, IndexU, IndexV), Grid.GetInner2DPoint(EGridSpace::UniformScaled, IndexU, IndexV + 1), Index, EVisuProperty(Property + 1));
-	DisplaySegment(Grid.GetInner2DPoint(EGridSpace::UniformScaled, IndexU + 1, IndexV), Grid.GetInner2DPoint(EGridSpace::UniformScaled, IndexU + 1, IndexV + 1), Index, EVisuProperty(Property + 1));
-	DisplaySegment(Grid.GetInner2DPoint(EGridSpace::UniformScaled, IndexU, IndexV), Grid.GetInner2DPoint(EGridSpace::UniformScaled, IndexU + 1, IndexV), Index, EVisuProperty(Property + 1));
-	DisplaySegment(Grid.GetInner2DPoint(EGridSpace::UniformScaled, IndexU, IndexV + 1), Grid.GetInner2DPoint(EGridSpace::UniformScaled, IndexU + 1, IndexV + 1), Index, EVisuProperty(Property + 1));
+	DisplaySegment(Grid.GetInner2DPoint(EGridSpace::UniformScaled, IndexU, IndexV) * Grid.DisplayScale, Grid.GetInner2DPoint(EGridSpace::UniformScaled, IndexU, IndexV + 1) * Grid.DisplayScale, Index, EVisuProperty(Property + 1));
+	DisplaySegment(Grid.GetInner2DPoint(EGridSpace::UniformScaled, IndexU + 1, IndexV) * Grid.DisplayScale, Grid.GetInner2DPoint(EGridSpace::UniformScaled, IndexU + 1, IndexV + 1) * Grid.DisplayScale, Index, EVisuProperty(Property + 1));
+	DisplaySegment(Grid.GetInner2DPoint(EGridSpace::UniformScaled, IndexU, IndexV) * Grid.DisplayScale, Grid.GetInner2DPoint(EGridSpace::UniformScaled, IndexU + 1, IndexV) * Grid.DisplayScale, Index, EVisuProperty(Property + 1));
+	DisplaySegment(Grid.GetInner2DPoint(EGridSpace::UniformScaled, IndexU, IndexV + 1) * Grid.DisplayScale, Grid.GetInner2DPoint(EGridSpace::UniformScaled, IndexU + 1, IndexV + 1) * Grid.DisplayScale, Index, EVisuProperty(Property + 1));
 }
 
 
@@ -168,7 +167,6 @@ void FIsoTriangulator::DisplayCell(const FCell& Cell) const
 		return;
 	}
 
-
 	F3DDebugSession _(FString::Printf(TEXT("Cell %d"), Cell.Id));
 
 	DrawCellBoundary(Cell.Id, EVisuProperty::YellowPoint);
@@ -176,7 +174,7 @@ void FIsoTriangulator::DisplayCell(const FCell& Cell) const
 	{
 		for (const FLoopNode* Node : Nodes)
 		{
-			DisplayPoint(Node->GetPoint(EGridSpace::UniformScaled, Grid), EVisuProperty::BluePoint, Cell.Id);
+			Grid.DisplayIsoNode(EGridSpace::UniformScaled, *Node, Cell.Id, EVisuProperty::BluePoint);
 		}
 	}
 }
@@ -193,7 +191,7 @@ void FIntersectionSegmentTool::Display(const TCHAR* Message, EVisuProperty Prope
 	Open3DDebugSession(Message);
 	for (const FSegment4IntersectionTools& Segment : Segments)
 	{
-		DisplaySegment(Segment.Segment2D[0], Segment.Segment2D[1], Index++, Property);
+		DisplaySegment(Segment.Segment2D[0] * Grid.DisplayScale, Segment.Segment2D[1] * Grid.DisplayScale, Index++, Property);
 	}
 	Close3DDebugSession();
 }
diff --git a/Engine/Source/Runtime/Datasmith/CADKernel/Private/CADKernel/Mesh/Meshers/MesherTools.h b/Engine/Source/Runtime/Datasmith/CADKernel/Private/CADKernel/Mesh/Meshers/MesherTools.h
--- a/Engine/Source/Runtime/Datasmith/CADKernel/Private/CADKernel/Mesh/Meshers/MesherTools.h
+++ b/Engine/Source/Runtime/Datasmith/CADKernel/Private/CADKernel/Mesh/Meshers/MesherTools.h
@@ -2,7 +2,7 @@
 #pragma once
 
 #include "CADKernel/Core/Types.h"
-
+#include "CADKernel/Math/Boundary.h"
 
 namespace CADKernel
 {
@@ -47,6 +47,21 @@ namespace CADKernel
 		 */
 		static void ComputeFinalCuttingPointsWithImposedCuttingPoints(const TArray<double>& CrossingUs, const TArray<double>& DeltaUs, const TArray<FCuttingPoint>& ImposedCuttingPoints, TArray<double>& OutCuttingPoints);
 
+		/**
+		 * Fill an array of CuttingPointCoordinates with "PointCount" coordinates evenly divided between the extremities of the boundary 
+		 */
+		static void FillCuttingPointCoordinates(const FLinearBoundary& Boundary, int32 PointCount, TArray<double>& OutCuttingPointCoordinates)
+		{
+			PointCount--;
+			ensureCADKernel(PointCount > 0);
+			double Delta = Boundary.Length() / PointCount;
+			OutCuttingPointCoordinates.Empty(PointCount);
+			OutCuttingPointCoordinates.Add(Boundary.GetMin());
+			for (int32 Index = 0; Index < PointCount; ++Index)
+			{
+				OutCuttingPointCoordinates.Add(OutCuttingPointCoordinates.Last() + Delta);
+			}
+		}
 	};
 
 } // namespace CADKernel
diff --git a/Engine/Source/Runtime/Datasmith/CADKernel/Private/CADKernel/Mesh/Meshers/ParametricMesher.cpp b/Engine/Source/Runtime/Datasmith/CADKernel/Private/CADKernel/Mesh/Meshers/ParametricMesher.cpp
--- a/Engine/Source/Runtime/Datasmith/CADKernel/Private/CADKernel/Mesh/Meshers/ParametricMesher.cpp
+++ b/Engine/Source/Runtime/Datasmith/CADKernel/Private/CADKernel/Mesh/Meshers/ParametricMesher.cpp
@@ -100,6 +100,8 @@ void FParametricMesher::MeshEntities(TArray<TSharedPtr<FEntity>>& InEntities)
 	}
 
 	MeshEntities();
+
+	MesherReport.Print();
 }
 
 void FParametricMesher::MeshEntities()
@@ -228,14 +230,16 @@ void FParametricMesher::Mesh(TSharedRef<FTopologicalFace> Face)
 	FGrid Grid(Face, MeshModel);
 	GenerateCloud(Grid);
 
+	FDuration GenerateCloudDuration = FChrono::Elapse(GenerateCloudStartTime);
+
 	if (Grid.IsDegenerated())
 	{
+		MesherReport.Logs.AddDegeneratedGrid();
 		FMessage::Printf(EVerboseLevel::Log, TEXT("The meshing of the surface %d failed due to a degenerated grid\n"), Face->GetId());
 		Face->SetMeshed();
 		return;
 	}
 
-	FDuration GenerateCloudDuration = FChrono::Elapse(GenerateCloudStartTime);
 
 	FTimePoint IsoTriangulerStartTime = FChrono::Now();
 
@@ -412,7 +416,7 @@ void FParametricMesher::Mesh(FTopologicalEdge& InEdge, FTopologicalFace& Face)
 			// In the particular case where the both case are opposed, we can have the 2d line sampled with 4 points, and the 2d curve sampled with 2 points (because in 3d, the 2d curve is a 3d line)
 			// In this case, the loop is flat i.e. in 2d the meshes of the 2d line and 2d curve are coincident
 			// So the grid is degenerated and the surface is not meshed
-			// to avoid this case, the Edge is virtually meshed i.e. the nodes inside the edge has the is of the mesh of the vertices.
+			// to avoid this case, the Edge is virtually meshed i.e. the nodes inside the edge have the id of the mesh of the vertices.
 			InEdge.SetAsVirtuallyMeshed();
 		}
 
@@ -1608,6 +1612,11 @@ void FParametricMesher::InitParameters(const FString& ParametersString)
 	Parameters->SetFromString(ParametersString);
 }
 
+void FParametricMesher::PrintReport()
+{
+	MesherReport.Print();
+}
+
 #ifdef DEBUG_INTERSECTEDGEISOS
 TMap<int32, int32> SurfaceDrawed;
 bool bDisplayIsoCurve = true;
diff --git a/Engine/Source/Runtime/Datasmith/CADKernel/Private/CADKernel/Mesh/Structure/Grid.cpp b/Engine/Source/Runtime/Datasmith/CADKernel/Private/CADKernel/Mesh/Structure/Grid.cpp
--- a/Engine/Source/Runtime/Datasmith/CADKernel/Private/CADKernel/Mesh/Structure/Grid.cpp
+++ b/Engine/Source/Runtime/Datasmith/CADKernel/Private/CADKernel/Mesh/Structure/Grid.cpp
@@ -46,28 +46,21 @@ void FGrid::ProcessPointCloud()
 	}
 
 	ScaleLoops();
-	if (!RemoveCoincidentNodes())
-	{
-		// the external loop is degenerated, the face is degenerate...
-		// the mesh is canceled
-		SetAsDegenerated();
-		return;
-	}
 
 #ifdef DEBUG_GRID
-	DisplayLoop(TEXT("FGrid::Loop 2D"), GetLoops2D(EGridSpace::Default2D), true, false);
+	DisplayGridLoop(TEXT("FGrid::Loop 2D"), GetLoops2D(EGridSpace::Default2D), true, false);
 	DisplayInnerDomainPoints(TEXT("FGrid::PointCloud 2D"), GetInner2DPoints(EGridSpace::Default2D));
 
-	DisplayLoop(TEXT("FGrid::Loop 2D Scaled"), GetLoops2D(EGridSpace::Scaled), true, false);
+	DisplayGridLoop(TEXT("FGrid::Loop 2D Scaled"), GetLoops2D(EGridSpace::Scaled), true, false);
 	DisplayInnerDomainPoints(TEXT("FGrid::PointCloud 2D Scaled"), GetInner2DPoints(EGridSpace::Scaled));
 
-	DisplayLoop(TEXT("FGrid::Loop 2D UniformScaled"), GetLoops2D(EGridSpace::UniformScaled), true, false);
+	DisplayGridLoop(TEXT("FGrid::Loop 2D UniformScaled"), GetLoops2D(EGridSpace::UniformScaled), true, false);
 	DisplayInnerDomainPoints(TEXT("FGrid::PointCloud 2D UniformScaled"), GetInner2DPoints(EGridSpace::UniformScaled));
 
-	DisplayLoop(TEXT("FGrid::Loop 3D"), GetLoops3D(), true, false);
+	DisplayGridLoop(TEXT("FGrid::Loop 3D"), GetLoops3D(), true, false);
 	DisplayInnerDomainPoints(TEXT("FGrid::PointCloud 3D"), GetInner3DPoints());
 
-	DisplayLoop(TEXT("FGrid::Loop 3D"), GetLoops3D(), true, false);
+	DisplayGridLoop(TEXT("FGrid::Loop 3D"), GetLoops3D(), true, false);
 	DisplayInnerDomainPoints(TEXT("FGrid::PointCloud 3D"), GetInner3DPoints());
 
 	DisplayGridNormal();
@@ -82,22 +75,14 @@ void FGrid::ProcessPointCloud()
 	Wait(bDisplay);
 #endif
 	FindPointsCloseToLoop();
-#ifdef DEBUG_GRID
-	DisplayGridInnerPoints(DisplaySpace, TEXT("FGrid::FindPointsCloseAndInsideToLoop result"));
-#endif
 	RemovePointsClosedToLoop();
 
-#ifdef DEBUG_GRID
-	DisplayGridInnerPoints(DisplaySpace, TEXT("FGrid::RemovePointsClosedToLoop (Final PointCloud)"));
-	//Wait(bDisplay);
-#endif
-
-		// Removed of Thin zone boundary (the last boundaries). In case of thin zone, the number of 2d boundary will be biggest than 3d boundary one.
-		// Only EGridSpace::Default2D is needed.
+	// Removed of Thin zone boundary (the last boundaries). In case of thin zone, the number of 2d boundary will be biggest than 3d boundary one.
+	// Only EGridSpace::Default2D is needed.
 	FaceLoops2D[EGridSpace::Default2D].SetNum(FaceLoops3D.Num());
 
 #ifdef DEBUG_GRID
-	DisplayLoop(TEXT("FGrid::Final Loop 2D"), GetLoops2D(DisplaySpace), true, false);
+	DisplayGridLoop(TEXT("FGrid::Final Loop 2D"), GetLoops2D(DisplaySpace), true, false);
 	//DisplayLoop(TEXT("FGrid::Final Loop 3D"), GetLoops3D(), true, false);
 #endif
 
@@ -961,7 +946,7 @@ bool FGrid::GetMeshOfLoops()
 	NodeIdsOfFaceLoops.Reserve(LoopCount);
 
 #ifdef DEBUG_GET_BOUNDARY_MESH
-	F3DDebugSession _(TEXT("GetLoopMesh"));
+	F3DDebugSession _(bDisplay, ("GetLoopMesh"));
 #endif
 
 	for (const TSharedPtr<FTopologicalLoop>& Loop : Face->GetLoops())
@@ -1005,8 +990,19 @@ bool FGrid::GetMeshOfLoops()
 
 			if (Edge->IsDegenerated())
 			{
-				ensureCADKernel(!Edge->GetCuttingPoints().IsEmpty());
-				// else converts CuttingPoint2D form ActiveEdge to Edge (with Linear conversion)
+				if (EdgeCuttingPointCoordinates.IsEmpty())
+				{
+					int32 CuttingPointCount = 2;
+					for (const FTopologicalEdge* TwinEdge : Edge->GetTwinsEntities())
+					{
+						int32 TwinCuttingPointCount = TwinEdge->GetCuttingPoints().Num();
+						if (TwinCuttingPointCount > CuttingPointCount)
+						{
+							CuttingPointCount = TwinCuttingPointCount;
+						}
+					}
+					FMesherTools::FillCuttingPointCoordinates(Edge->GetBoundary(), CuttingPointCount, EdgeCuttingPointCoordinates);
+				}
 
 				Swap(CuttingPolyline.Coordinates, EdgeCuttingPointCoordinates);
 				Edge->Approximate2DPoints(CuttingPolyline.Coordinates, CuttingPolyline.Points2D);
@@ -1030,6 +1026,23 @@ bool FGrid::GetMeshOfLoops()
 					Edge->ProjectTwinEdgePoints(MeshVertex3D, bSameDirection, CuttingPolyline.Coordinates);
 					CuttingPolyline.Coordinates.Insert(Edge->GetBoundary().GetMin(), 0);
 					CuttingPolyline.Coordinates.Add(Edge->GetBoundary().GetMax());
+
+					// check if there are coincident coordinates
+					bool bProjectionFailed = false;
+					for (int32 Index = 0; Index < CuttingPolyline.Coordinates.Num() - 1;)
+					{
+						double Diff = CuttingPolyline.Coordinates[Index++];
+						Diff -= CuttingPolyline.Coordinates[Index];
+						if (Diff > -SMALL_NUMBER)
+						{
+							bProjectionFailed = true;
+							break;
+						}
+					}
+					if(bProjectionFailed)
+					{
+						FMesherTools::FillCuttingPointCoordinates(Edge->GetBoundary(), MeshVertex3D.Num() + 2, CuttingPolyline.Coordinates);
+					}
 				}
 				else
 				{
@@ -1072,8 +1085,9 @@ bool FGrid::GetMeshOfLoops()
 			}
 
 #ifdef DEBUG_GET_BOUNDARY_MESH
+			if (bDisplay)
 			{
-				F3DDebugSession _(FString::Printf(TEXT("Edge %d cutting points on surface"), ActiveEdge->GetId()));
+				F3DDebugSession _(bDisplay, FString::Printf(TEXT("Edge %d cutting points on surface"), ActiveEdge->GetId()));
 				for (const FPoint2D& Point2D : CuttingPolyline.Points2D)
 				{
 					DisplayPoint(Point2D);
@@ -1403,88 +1417,6 @@ void FGrid::ScaleLoops()
 	}
 }
 
-bool FGrid::RemoveCoincidentNodes()
-{
-	const double SquareMinimumSize = FMath::Square(MinimumElementSize);
-
-	for (int32 IndexBoudnary = FaceLoops2D[EGridSpace::Default2D].Num() - 1; IndexBoudnary >= 0; --IndexBoudnary)
-	{
-		TArray<FPoint2D>& Loop2D = FaceLoops2D[EGridSpace::Default2D][IndexBoudnary];
-		TArray<FPoint2D>& ScaledLoop = FaceLoops2D[EGridSpace::Scaled][IndexBoudnary];
-		TArray<FPoint2D>& UniformScaledLoop = FaceLoops2D[EGridSpace::UniformScaled][IndexBoudnary];
-		TArray<int32>& NodeIdsOfLoop = NodeIdsOfFaceLoops[IndexBoudnary];
-		TArray<FPoint>& Loop3D = FaceLoops3D[IndexBoudnary];
-
-		double SquareDistance2D;
-		double SquareDistance;
-		int32 NewIndex = 1;
-		int32 NextIndex = 1;
-
-		TFunction<void(int32)> CheckAndUpdate = [&](int32 IndexFor3D)
-		{
-			if (SquareDistance2D > SquareMinimumSize)
-			{
-				if (SquareDistance > SquareMinimumSize)
-				{
-					if (NextIndex != NewIndex)
-					{
-						UniformScaledLoop[NewIndex] = UniformScaledLoop[NextIndex];
-						Loop3D[NewIndex] = Loop3D[NextIndex];
-						Loop2D[NewIndex] = Loop2D[NextIndex];
-						ScaledLoop[NewIndex] = ScaledLoop[NextIndex];
-						NodeIdsOfLoop[NewIndex] = NodeIdsOfLoop[NextIndex];
-					}
-				}
-				else
-				{
-					if (NextIndex != NewIndex)
-					{
-						UniformScaledLoop[NewIndex] = UniformScaledLoop[NextIndex];
-					}
-					Loop3D[NewIndex] = Loop3D[IndexFor3D];
-					NodeIdsOfLoop[NewIndex] = NodeIdsOfLoop[IndexFor3D];
-				}
-				NewIndex++;
-			}
-		};
-
-		for (; NextIndex < Loop2D.Num(); ++NextIndex)
-		{
-			SquareDistance2D = UniformScaledLoop[NewIndex - 1].SquareDistance(UniformScaledLoop[NextIndex]);
-			SquareDistance = Loop3D[NewIndex - 1].SquareDistance(Loop3D[NextIndex]);
-			CheckAndUpdate(NewIndex - 1);
-		}
-
-		NewIndex--;
-		NextIndex--;
-		SquareDistance2D = UniformScaledLoop[NewIndex].SquareDistance(UniformScaledLoop[0]);
-		SquareDistance = Loop3D[NewIndex].SquareDistance(Loop3D[0]);
-		CheckAndUpdate(0);
-
-		Loop3D.SetNum(NewIndex);
-		Loop2D.SetNum(NewIndex);
-		ScaledLoop.SetNum(NewIndex);
-		NodeIdsOfLoop.SetNum(NewIndex);
-		UniformScaledLoop.SetNum(NewIndex);
-
-		if(NewIndex < 3)
-		{
-			FaceLoops2D[EGridSpace::Default2D].RemoveAt(IndexBoudnary);
-			FaceLoops2D[EGridSpace::Scaled].RemoveAt(IndexBoudnary);
-			FaceLoops2D[EGridSpace::UniformScaled].RemoveAt(IndexBoudnary);
-			NodeIdsOfFaceLoops.RemoveAt(IndexBoudnary);
-			FaceLoops3D.RemoveAt(IndexBoudnary);
-			if (IndexBoudnary == 0)
-			{
-				// the external loop is degenerated, the face is degenerate...
-				return false;
-			}
-		}
-	}
-	return true;
-}
-
-
 void FGrid::ComputeMaxDeltaUV()
 {
 	MaxDeltaUV[EIso::IsoU] = 0;
diff --git a/Engine/Source/Runtime/Datasmith/CADKernel/Private/CADKernel/Mesh/Structure/Grid.h b/Engine/Source/Runtime/Datasmith/CADKernel/Private/CADKernel/Mesh/Structure/Grid.h
--- a/Engine/Source/Runtime/Datasmith/CADKernel/Private/CADKernel/Mesh/Structure/Grid.h
+++ b/Engine/Source/Runtime/Datasmith/CADKernel/Private/CADKernel/Mesh/Structure/Grid.h
@@ -220,12 +220,6 @@ protected:
 	 */
 	void ScaleLoops();
 
-	/**
-	 * After the loops have been scaled, the distance between two consecutive vertices is verified to avoid degenerated mesh elements
-	 * @ return false if the external loop is degenerated after the process
-	 */
-	bool RemoveCoincidentNodes();
-
 	/**
 	 * Convert Coordinate of "DefaultParametric" space into a scaled parametric space
 	 * @see ScaleLoops
@@ -564,11 +558,12 @@ public:
 	// ======================================================================================================================================================================================================================
 #ifdef CADKERNEL_DEV
 	bool bDisplay = false;
+	const double DisplayScale = 100.;
 
 	void DisplayIsoNode(EGridSpace Space, const int32 PointIndex, FIdent Ident = 0, EVisuProperty Property = EVisuProperty::BluePoint) const;
 	void DisplayIsoNode(EGridSpace Space, const FIsoNode& Node, FIdent Ident = 0, EVisuProperty Property = EVisuProperty::BluePoint) const;
-	void DisplayIsoNodes(EGridSpace Space, const TArray<FIsoNode*>& Nodes, EVisuProperty Property = EVisuProperty::BluePoint) const;
-	void DisplayIsoNodes(const FString& Message, EGridSpace Space, const TArray<FIsoNode*>& Nodes, EVisuProperty Property = EVisuProperty::BluePoint) const
+	void DisplayIsoNodes(EGridSpace Space, const TArray<const FIsoNode*>& Nodes, EVisuProperty Property = EVisuProperty::BluePoint) const;
+	void DisplayIsoNodes(const FString& Message, EGridSpace Space, const TArray<const FIsoNode*>& Nodes, EVisuProperty Property = EVisuProperty::BluePoint) const
 	{
 		if (!bDisplay)
 		{
@@ -638,11 +633,11 @@ public:
 		{
 			if (IsInsideFace[Index])
 			{
-				DisplayPoint(Points[Index], IsCloseToLoop[Index] ? EVisuProperty::BluePoint : EVisuProperty::GreenPoint, Index);
+				DisplayPoint(Points[Index] * DisplayScale, IsCloseToLoop[Index] ? EVisuProperty::BluePoint : EVisuProperty::GreenPoint, Index);
 			}
 			else
 			{
-				DisplayPoint(Points[Index], EVisuProperty::OrangePoint, Index);
+				DisplayPoint(Points[Index] * DisplayScale, EVisuProperty::OrangePoint, Index);
 			}
 		}
 	}
@@ -667,20 +662,20 @@ public:
 			const TPoint* FirstSegmentPoint = &Loop[0];
 			if (bDisplayNodes)
 			{
-				DisplayPoint(*FirstSegmentPoint, EVisuProperty::BluePoint, 0);
+				DisplayPoint(*FirstSegmentPoint * DisplayScale, EVisuProperty::BluePoint, 0);
 			}
 
 			for (int32 Index = 1; Index < Loop.Num(); ++Index)
 			{
 				const TPoint* SecondSegmentPoint = &Loop[Index];
-				DisplaySegment(*FirstSegmentPoint, *SecondSegmentPoint);
+				DisplaySegment(*FirstSegmentPoint * DisplayScale, *SecondSegmentPoint * DisplayScale);
 				if (bDisplayNodes)
 				{
-					DisplayPoint(*SecondSegmentPoint, EVisuProperty::BluePoint, Index);
+					DisplayPoint(*SecondSegmentPoint * DisplayScale, EVisuProperty::BluePoint, Index);
 				}
 				FirstSegmentPoint = SecondSegmentPoint;
 			}
-			DisplaySegment(*FirstSegmentPoint, *Loop.begin());
+			DisplaySegment(*FirstSegmentPoint * DisplayScale, *Loop.begin() * DisplayScale);
 			if (bMakeGroup)
 			{
 				Close3DDebugSession();
diff --git a/Engine/Source/Runtime/Datasmith/CADKernel/Private/CADKernel/Mesh/Structure/GridDisplay.cpp b/Engine/Source/Runtime/Datasmith/CADKernel/Private/CADKernel/Mesh/Structure/GridDisplay.cpp
--- a/Engine/Source/Runtime/Datasmith/CADKernel/Private/CADKernel/Mesh/Structure/GridDisplay.cpp
+++ b/Engine/Source/Runtime/Datasmith/CADKernel/Private/CADKernel/Mesh/Structure/GridDisplay.cpp
@@ -16,7 +16,7 @@ void FGrid::DisplayIsoNode(EGridSpace DisplaySpace, const int32 PointIndex, FIde
 	{
 		return;
 	}
-	DisplayPoint(Points2D[DisplaySpace][PointIndex], Ident);
+	DisplayPoint(Points2D[DisplaySpace][PointIndex] * DisplayScale, Ident);
 }
 
 void FGrid::DisplayIsoNode(EGridSpace Space, const FIsoNode& Node, FIdent Ident, EVisuProperty Property) const
@@ -25,17 +25,17 @@ void FGrid::DisplayIsoNode(EGridSpace Space, const FIsoNode& Node, FIdent Ident,
 	{
 		return;
 	}
-	DisplayPoint(Node.GetPoint(Space, *this), Property, Ident);
+	DisplayPoint(Node.GetPoint(Space, *this) * DisplayScale, Property, Ident);
 }
 
-void FGrid::DisplayIsoNodes(EGridSpace Space, const TArray<FIsoNode*>& Nodes, EVisuProperty Property) const
+void FGrid::DisplayIsoNodes(EGridSpace Space, const TArray<const FIsoNode*>& Nodes, EVisuProperty Property) const
 {
 	if (!bDisplay)
 	{
 		return;
 	}
 
-	for (FIsoNode* Node : Nodes)
+	for (const FIsoNode* Node : Nodes)
 	{
 		DisplayIsoNode(Space, *Node, Node->GetId(), Property);
 	}
@@ -47,7 +47,7 @@ void FGrid::DisplayIsoSegment(EGridSpace Space, const FIsoSegment& Segment, FIde
 	{
 		return;
 	}
-	DisplaySegment(Segment.GetFirstNode().GetPoint(Space, *this), Segment.GetSecondNode().GetPoint(Space, *this), Ident, Property, bDisplayOrientation);
+	DisplaySegment(Segment.GetFirstNode().GetPoint(Space, *this) * DisplayScale, Segment.GetSecondNode().GetPoint(Space, *this) * DisplayScale, Ident, Property, bDisplayOrientation);
 }
 
 void FGrid::DisplayIsoSegment(EGridSpace Space, const FIsoNode& NodeA, const FIsoNode& NodeB, FIdent Ident, EVisuProperty Property) const
@@ -56,16 +56,16 @@ void FGrid::DisplayIsoSegment(EGridSpace Space, const FIsoNode& NodeA, const FIs
 	{
 		return;
 	}
-	DisplaySegment(NodeA.GetPoint(Space, *this), NodeB.GetPoint(Space, *this), Ident, Property);
+	DisplaySegment(NodeA.GetPoint(Space, *this) * DisplayScale, NodeB.GetPoint(Space, *this) * DisplayScale, Ident, Property);
 }
 
 void FGrid::DisplayTriangle(EGridSpace Space, const FIsoNode& NodeA, const FIsoNode& NodeB, const FIsoNode& NodeC) const
 {
 	TArray<FPoint> Points;
 	Points.SetNum(3);
-	Points[0] = NodeA.GetPoint(Space, *this);
-	Points[1] = NodeB.GetPoint(Space, *this);
-	Points[2] = NodeC.GetPoint(Space, *this);
+	Points[0] = NodeA.GetPoint(Space, *this) * DisplayScale;
+	Points[1] = NodeB.GetPoint(Space, *this) * DisplayScale;
+	Points[2] = NodeC.GetPoint(Space, *this) * DisplayScale;
 	DrawElement(2, Points, EVisuProperty::Element);
 
 	DisplaySegment(Points[0], Points[1], 0, EVisuProperty::EdgeMesh);
@@ -85,11 +85,11 @@ void FGrid::DisplayIsoSegments(EGridSpace Space, const TArray<FIsoSegment*>& InS
 	{
 		if (Segment)
 		{
-			DisplaySegment(Segment->GetFirstNode().GetPoint(Space, *this), Segment->GetSecondNode().GetPoint(Space, *this), Index++, Property, bDisplayOrientation);
+			DisplaySegment(Segment->GetFirstNode().GetPoint(Space, *this) * DisplayScale, Segment->GetSecondNode().GetPoint(Space, *this) * DisplayScale, Index++, Property, bDisplayOrientation);
 			if(bDisplayNode)
 			{
-				DisplayPoint(Segment->GetFirstNode().GetPoint(Space, *this), (EVisuProperty)(Property - 1), Segment->GetFirstNode().GetIndex());
-				DisplayPoint(Segment->GetSecondNode().GetPoint(Space, *this), (EVisuProperty)(Property - 1), Segment->GetSecondNode().GetIndex());
+				DisplayPoint(Segment->GetFirstNode().GetPoint(Space, *this) * DisplayScale, (EVisuProperty)(Property - 1), Segment->GetFirstNode().GetIndex());
+				DisplayPoint(Segment->GetSecondNode().GetPoint(Space, *this) * DisplayScale, (EVisuProperty)(Property - 1), Segment->GetSecondNode().GetIndex());
 			}
 		}
 	}
@@ -311,7 +311,7 @@ void FGrid::DisplayGridPoints(EGridSpace DisplaySpace) const
 		{
 			if (IsInsideFace[Index])
 			{
-				DisplayPoint(Points2D[DisplaySpace][Index], EVisuProperty::BluePoint, Index);
+				DisplayPoint(Points2D[DisplaySpace][Index] * DisplayScale, EVisuProperty::BluePoint, Index);
 				NbNum++;
 			}
 		}
@@ -324,7 +324,7 @@ void FGrid::DisplayGridPoints(EGridSpace DisplaySpace) const
 		{
 			if (!IsInsideFace[Index])
 			{
-				DisplayPoint(Points2D[DisplaySpace][Index], EVisuProperty::OrangePoint, Index);
+				DisplayPoint(Points2D[DisplaySpace][Index] * DisplayScale, EVisuProperty::OrangePoint, Index);
 			}
 		}
 	}
diff --git a/Engine/Source/Runtime/Datasmith/CADKernel/Private/CADKernel/Mesh/Structure/LoopCleaner.cpp b/Engine/Source/Runtime/Datasmith/CADKernel/Private/CADKernel/Mesh/Structure/LoopCleaner.cpp
--- a/Engine/Source/Runtime/Datasmith/CADKernel/Private/CADKernel/Mesh/Structure/LoopCleaner.cpp
+++ b/Engine/Source/Runtime/Datasmith/CADKernel/Private/CADKernel/Mesh/Structure/LoopCleaner.cpp
@@ -45,7 +45,7 @@ bool FLoopCleaner::CleanLoops()
 	{
 		Grid.DisplayIsoSegments(TEXT("Loops Orientation"), EGridSpace::UniformScaled, LoopSegments, false, true, EVisuProperty::BlueCurve);
 		Grid.DisplayIsoSegments(TEXT("Loops Init"), EGridSpace::UniformScaled, LoopSegments, true, false, EVisuProperty::BlueCurve);
-		Grid.DisplayIsoNodes(TEXT("BestStartNodeOfLoop"), EGridSpace::UniformScaled, (const TArray<FIsoNode*>&) BestStartNodeOfLoops, EVisuProperty::BluePoint);
+		Grid.DisplayIsoNodes(TEXT("BestStartNodeOfLoop"), EGridSpace::UniformScaled, (const TArray<const FIsoNode*>) BestStartNodeOfLoops, EVisuProperty::BluePoint);
 		Wait(bDisplay);
 	}
 #endif
@@ -98,7 +98,8 @@ bool FLoopCleaner::CleanLoops()
 
 		if (!RemoveLoopPicks())
 		{
-			continue;
+			FMessage::Printf(Log, TEXT("Loop intersections of the surface %d cannot be fixed. The mesh of this surface is canceled.\n"), Grid.GetFace()->GetId());
+			return false;
 		}
 
 		if (NodesOfLoopCount == 0)
@@ -817,7 +818,7 @@ bool FLoopCleaner::RemoveSelfIntersectionsOfLoop()
 		FLoopNode* Segment0End = GetNodeAt(NextIndex((int32)Intersection.Key));
 		FLoopNode* Segment1Start = GetNodeAt((int32)Intersection.Value);
 
-		if (!Segment0End || !Segment1Start)
+		if (Segment0End == nullptr || Segment1Start == nullptr)
 		{
 			++IntersectionIndex;
 			continue;
@@ -929,7 +930,7 @@ bool FLoopCleaner::RemoveOutgoingLoop(const TPair<double, double>& Intersection,
 
 	FLoopNode* TmpNode = GetNodeAt(NextIndex((int32)Intersection.Value));
 	FLoopNode* EndNode = GetNodeAt((int32)NextIntersection.Value);
-	if (TmpNode == nullptr || TmpNode->IsDelete() || EndNode == nullptr || EndNode->IsDelete())
+	if (TmpNode == nullptr || EndNode == nullptr)
 	{
 		return false;
 	}
@@ -1007,7 +1008,7 @@ bool FLoopCleaner::RemoveIntersectionsOfSubLoop(int32 IntersectionIndex, int32 I
 		FLoopNode* Node = GetNodeAt(StartIndex);
 		FLoopNode* StopNode = GetNodeAt(EndIndex);
 		FLoopNode* EndSegment = GetNodeAt(SegmnentEndIndex);
-		if (Node == nullptr || Node->IsDelete() || StopNode == nullptr || StopNode->IsDelete() || EndSegment == nullptr)
+		if (Node == nullptr || StopNode == nullptr || EndSegment == nullptr)
 		{
 			return false;
 		}
@@ -1025,7 +1026,7 @@ bool FLoopCleaner::RemoveIntersectionsOfSubLoop(int32 IntersectionIndex, int32 I
 
 		for (Node = GetNodeAt(StartIndex); Node != StopNode; )
 		{
-			if (Node->IsDelete())
+			if (Node == nullptr)
 			{
 				return false;
 			}
@@ -1102,7 +1103,7 @@ bool FLoopCleaner::RemoveIntersectionsOfSubLoop(int32 IntersectionIndex, int32 I
 			{
 				FLoopNode* NodeSide0 = GetNodeAt(IndexSide0);
 				FLoopNode* NodeSide1 = GetNodeAt(IndexSide1);
-				if (NodeSide0->IsDelete() || NodeSide1->IsDelete())
+				if (NodeSide0 == nullptr || NodeSide1 == nullptr)
 				{
 					return false;
 				}
@@ -1125,14 +1126,14 @@ bool FLoopCleaner::RemoveIntersectionsOfSubLoop(int32 IntersectionIndex, int32 I
 				FLoopSection IntersectingSection;
 				IntersectingSection.Key = GetNodeAt(IndexSide0);
 				IntersectingSection.Value = GetNodeAt((int32)SecondIntersection.Key);
-				if(!IntersectingSection.Key || !IntersectingSection.Value)
+				if(IntersectingSection.Key == nullptr || IntersectingSection.Value == nullptr)
 				{
 					return false;
 				}
 
 				IntersectingSection.Key = GetPrevious(IntersectingSection.Key);
 				IntersectingSection.Value = GetNext(IntersectingSection.Value);
-				if (!IntersectingSection.Key || IntersectingSection.Key->IsDelete() || !IntersectingSection.Value || IntersectingSection.Value->IsDelete())
+				if (IntersectingSection.Key == nullptr || IntersectingSection.Key->IsDelete() || IntersectingSection.Value == nullptr || IntersectingSection.Value->IsDelete())
 				{
 					return false;
 				}
@@ -1140,7 +1141,7 @@ bool FLoopCleaner::RemoveIntersectionsOfSubLoop(int32 IntersectionIndex, int32 I
 				FLoopSection OppositeSection;
 				OppositeSection.Key = GetNodeAt(IndexSide1);
 				OppositeSection.Value = GetNodeAt((int32)FirstIntersection.Value);
-				if (!OppositeSection.Key || !OppositeSection.Value)
+				if (OppositeSection.Key == nullptr || OppositeSection.Value == nullptr)
 				{
 					return false;
 				}
@@ -1148,7 +1149,7 @@ bool FLoopCleaner::RemoveIntersectionsOfSubLoop(int32 IntersectionIndex, int32 I
 				OppositeSection.Key = GetPrevious(OppositeSection.Key);
 				OppositeSection.Value = GetNext(OppositeSection.Value);
 
-				if (!OppositeSection.Key || OppositeSection.Key->IsDelete() || !OppositeSection.Value || OppositeSection.Value->IsDelete())
+				if (OppositeSection.Key == nullptr || OppositeSection.Key->IsDelete() || OppositeSection.Value == nullptr || OppositeSection.Value->IsDelete())
 				{
 					return false;
 				}
@@ -1167,17 +1168,36 @@ bool FLoopCleaner::RemoveIntersectionsOfSubLoop(int32 IntersectionIndex, int32 I
 #endif
 
 				// Check if there is no more intersection, otherwise fix last intersections by moving node of OppositeSection
-				// As we don't know the existence of intersection, the process is chesk and fix
+				// As we don't know the existence of intersection, the process checks and fixes
 				{
 					LoopSegmentsIntersectionTool.Empty(NodesOfLoopCount);
 					for (FLoopNode* Node = IntersectingSection.Key; Node != IntersectingSection.Value; Node = GetNext(Node))
 					{
-						LoopSegmentsIntersectionTool.AddSegment(*Node->GetSegmentConnectedTo(GetNext(Node)));
+						if (Node->IsDelete())
+						{
+							// should not happen => so the process is canceled
+							return false;
+						}
+						
+						FIsoSegment* Segment = Node->GetSegmentConnectedTo(GetNext(Node));
+						if(!Segment)
+						{
+							// should not happen => so the process is canceled
+							return false;
+						}
+
+						LoopSegmentsIntersectionTool.AddSegment(*Segment);
 					}
 				}
 
 				for (FLoopNode* Node = OppositeSection.Key; Node != OppositeSection.Value; Node = GetNext(Node))
 				{
+					if (Node->IsDelete() || OppositeSection.Value->IsDelete())
+					{
+						// should not happen => so the process is canceled
+						return false;
+					}
+
 					if (const FIsoSegment* Intersection = LoopSegmentsIntersectionTool.DoesIntersect(*Node, *GetNext(Node)))
 					{
 						MoveNodeBehindSegment(*Intersection, *GetNext(Node));
@@ -1412,7 +1432,7 @@ bool FLoopCleaner::Fill(LoopCleanerImpl::FPinchIntersectionContext& Context)
 	{
 		Context.Points[Side].SetNum(3);
 		Context.Nodes[Side][1] = GetNodeAt(Index);
-		if (Context.Nodes[Side][1] == nullptr || Context.Nodes[Side][1]->IsDelete())
+		if (Context.Nodes[Side][1] == nullptr)
 		{
 			return false;
 		}
@@ -1568,7 +1588,7 @@ bool FLoopCleaner::DisconnectCrossingSegments(LoopCleanerImpl::FPinchIntersectio
 bool FLoopCleaner::SwapNodes(const TPair<double, double>& Intersection)
 {
 	FLoopNode* Node0 = GetNodeAt(NextIndex((int32)Intersection.Key));
-	if (Node0 == nullptr || Node0->IsDelete())
+	if (Node0 == nullptr)
 	{
 		return false;
 	}
@@ -1593,7 +1613,7 @@ bool FLoopCleaner::TryToSwapSegmentsOrRemoveLoop(const TPair<double, double>& In
 
 	FLoopNode* Segment0_Node1 = GetNodeAt(Segment0StartIndex);
 	FLoopNode* Segment1_Node0 = GetNodeAt(Segment1EndIndex);
-	if (Segment0_Node1 == nullptr || Segment0_Node1->IsDelete() || Segment1_Node0 == nullptr || Segment1_Node0->IsDelete())
+	if (Segment0_Node1 == nullptr || Segment1_Node0 == nullptr)
 	{
 		return false;
 	}
@@ -1927,7 +1947,7 @@ void FLoopCleaner::FindLoopIntersections()
 bool FLoopCleaner::RemoveOuterNode(const TPair<double, double>& Intersection)
 {
 	FLoopNode* Node0 = GetNodeAt(NextIndex((int32)Intersection.Key));
-	if (Node0->IsDelete())
+	if (Node0 == nullptr)
 	{
 		return false;
 	}
@@ -1973,9 +1993,19 @@ bool FLoopCleaner::RemoveLoopPicks()
 
 bool FLoopCleaner::RemovePickRecursively(FLoopNode* Node0, FLoopNode* Node1)
 {
+	if (Node0->IsDelete() || Node1->IsDelete())
+	{
+		return false;
+	}
+
 	FLoopNode* PreviousNode = &Node0->GetPreviousNode();
 	FLoopNode* NextNode = &Node1->GetNextNode();
 
+	if (PreviousNode->IsDelete() || NextNode->IsDelete())
+	{
+		return false;
+	}
+
 	const FPoint2D* PreviousPoint = &PreviousNode->Get2DPoint(EGridSpace::UniformScaled, Grid);
 	const FPoint2D* Point0 = &Node0->Get2DPoint(EGridSpace::UniformScaled, Grid);
 	const FPoint2D* Point1 = &Node1->Get2DPoint(EGridSpace::UniformScaled, Grid);
@@ -1983,16 +2013,7 @@ bool FLoopCleaner::RemovePickRecursively(FLoopNode* Node0, FLoopNode* Node1)
 
 	while (true)
 	{
-		bool bNodeIsToDelete = (Point0->SquareDistance(*Point1) < SquareGeometricTolerance2);
-		if (bNodeIsToDelete)
-		{
-			if (!RemoveNodeOfLoop(*Node0))
-			{
-				return false;
-			}
-		}
-
-		if (bNodeIsToDelete || CheckAndRemovePick(*PreviousPoint, *Point0, *Point1, *Node0))
+		if (CheckAndRemovePick(*PreviousPoint, *Point0, *Point1, *Node0))
 		{
 			if (PreviousNode->IsDelete())
 			{
@@ -2232,16 +2253,32 @@ EOrientation FLoopCleaner::GetLoopOrientation(const FLoopNode* StartNode)
 	LoopIndex = StartNode->GetLoopIndex();
 	double OptimalSlop = (LoopIndex == 0) ? -1 : 9.;
 
+	double MaxFrontSlope = 4;
+	double MinBackSlope = 4;
+
+#ifdef DEBUG_LOOP_ORIENTATION
+	F3DDebugSession _(Grid.bDisplay, TEXT("GetLoopOrientation"));
+#endif
 	TFunction<void(const FLoopNode*)> CompareWithSlopAt = [&](const FLoopNode* Node)
 	{
 		FLoopNode& PreviousNode = Node->GetPreviousNode();
 		FLoopNode& NextNode = Node->GetNextNode();
-		double Slop = ComputePositiveSlope(Node->Get2DPoint(EGridSpace::UniformScaled, Grid), PreviousNode.Get2DPoint(EGridSpace::UniformScaled, Grid), NextNode.Get2DPoint(EGridSpace::UniformScaled, Grid));
+		double Slope = ComputePositiveSlope(Node->Get2DPoint(EGridSpace::UniformScaled, Grid), PreviousNode.Get2DPoint(EGridSpace::UniformScaled, Grid), NextNode.Get2DPoint(EGridSpace::UniformScaled, Grid));
 
-		if ((Slop > OptimalSlop) == (LoopIndex == 0))
+#ifdef DEBUG_LOOP_ORIENTATION
+		F3DDebugSession A(Grid.bDisplay, FString::Printf(TEXT("Slop %f"), Slope));
+		Grid.DisplayIsoNode(EGridSpace::UniformScaled, *Node);
+		Grid.DisplayIsoSegment(EGridSpace::UniformScaled, *Node, PreviousNode);
+		Grid.DisplayIsoSegment(EGridSpace::UniformScaled, *Node, NextNode);
+#endif
+
+		if (Slope > 4)
 		{
-			OptimalSlop = Slop;
-			BestNode = Node;
+			if (MaxFrontSlope < Slope) MaxFrontSlope = Slope;
+		}
+		else
+		{
+			if (MinBackSlope > Slope) MinBackSlope = Slope;
 		}
 	};
 
@@ -2266,8 +2303,8 @@ EOrientation FLoopCleaner::GetLoopOrientation(const FLoopNode* StartNode)
 		}
 	};
 
-	const FLoopNode* Node = LoopCleanerImpl::GetNextConstNodeImpl(StartNode);
-	for (; Node != StartNode; Node = LoopCleanerImpl::GetNextConstNodeImpl(Node))
+	const FLoopNode* Node = StartNode;
+	do
 	{
 		const FPoint2D& Point = Node->Get2DPoint(EGridSpace::UniformScaled, Grid);
 
@@ -2292,16 +2329,20 @@ EOrientation FLoopCleaner::GetLoopOrientation(const FLoopNode* StartNode)
 			VMin = Point.V;
 			ExtremityNodes[3] = Node;
 		}
-	}
+		Node = LoopCleanerImpl::GetNextConstNodeImpl(Node);
+	} while (Node != StartNode);
+
 	FindLoopExtremity();
 
+	MaxFrontSlope = 8. - MaxFrontSlope;
+
 	if (LoopIndex == 0)
 	{
-		return OptimalSlop > 4 ? EOrientation::Front : EOrientation::Back;
+		return MaxFrontSlope < MinBackSlope ? EOrientation::Front : EOrientation::Back;
 	}
 	else
 	{
-		return OptimalSlop < 4 ? EOrientation::Front : EOrientation::Back;
+		return MaxFrontSlope > MinBackSlope ? EOrientation::Front : EOrientation::Back;
 	}
 }
 
diff --git a/Engine/Source/Runtime/Datasmith/CADKernel/Private/CADKernel/Mesh/Structure/LoopCleaner.h b/Engine/Source/Runtime/Datasmith/CADKernel/Private/CADKernel/Mesh/Structure/LoopCleaner.h
--- a/Engine/Source/Runtime/Datasmith/CADKernel/Private/CADKernel/Mesh/Structure/LoopCleaner.h
+++ b/Engine/Source/Runtime/Datasmith/CADKernel/Private/CADKernel/Mesh/Structure/LoopCleaner.h
@@ -7,6 +7,7 @@
 #include "CADKernel/Mesh/Meshers/IsoTriangulator/IntersectionSegmentTool.h"
 #include "CADKernel/Mesh/Meshers/IsoTriangulator/IsoNode.h"
 #include "CADKernel/Mesh/Meshers/IsoTriangulator/IsoSegment.h"
+#include "CADKernel/Mesh/Meshers/MesherReport.h"
 
 namespace CADKernel
 {
@@ -153,11 +154,13 @@ public:
 	{
 		if (!CleanLoops())
 		{
+			MesherReport.Logs.AddRemoveSelfIntersectionFailure();
 			return false;
 		}
 		
 		if (!UncrossLoops())
 		{
+			MesherReport.Logs.AddRemoveCrossingLoopsFailure();
 			return false;
 		}
 
@@ -189,11 +192,7 @@ private:
 	{
 		double Slop = ComputeUnorientedSlope(NodeToRemovePoint, PreviousPoint, NextPoint);
 		bool bRemoveNode = false;
-		if (Slop < 0.03)
-		{
-			bRemoveNode = true;
-		}
-		else if (Slop < 0.1)
+		if (Slop < 0.1)
 		{
 			double SquareDistance1 = SquareDistanceOfPointToSegment(PreviousPoint, NodeToRemovePoint, NextPoint);
 			double SquareDistance2 = SquareDistanceOfPointToSegment(NextPoint, NodeToRemovePoint, PreviousPoint);
@@ -322,7 +321,7 @@ private:
 		FLoopNode* Node = NodesOfLoop[Index];
 		if (Node->IsDelete())
 		{
-#ifdef DEBUG_LOOP_INTERSECTION_AND_FIX_IT	
+#ifdef GET_NODE_AT	
 			Wait();
 #endif
 			return nullptr;
@@ -417,7 +416,7 @@ private:
 			FLoopNode* Segment0End = GetNodeAt(NextIndex((int32)Intersection.Key));
 			FLoopNode* Segment1Start = GetNodeAt((int32)Intersection.Value);
 
-			if (!Segment0End || !Segment1Start)
+			if (Segment0End == nullptr || Segment1Start == nullptr)
 			{
 				return;
 			}
diff --git a/Engine/Source/Runtime/Datasmith/CADKernel/Private/CADKernel/Mesh/Structure/ModelMeshAnalyzer.cpp b/Engine/Source/Runtime/Datasmith/CADKernel/Private/CADKernel/Mesh/Structure/ModelMeshAnalyzer.cpp
new file mode 100644
--- /dev/null
+++ b/Engine/Source/Runtime/Datasmith/CADKernel/Private/CADKernel/Mesh/Structure/ModelMeshAnalyzer.cpp
@@ -0,0 +1,223 @@
+// Copyright Epic Games, Inc. All Rights Reserved.
+
+#include "CADKernel/Mesh/Structure/ModelMeshAnalyzer.h"
+
+#include "CADKernel/Math/Geometry.h"
+#include "CADKernel/Math/MathConst.h"
+#include "CADKernel/Mesh/Structure/EdgeMesh.h"
+#include "CADKernel/Mesh/Structure/Mesh.h"
+#include "CADKernel/Mesh/Structure/ModelMesh.h"
+#include "CADKernel/Mesh/Structure/FaceMesh.h"
+#include "CADKernel/Mesh/Structure/VertexMesh.h"
+#include "CADKernel/Topo/TopologicalEntity.h"
+
+#include "Templates/UnrealTemplate.h"
+
+namespace CADKernel
+{
+
+void FModelMeshAnalyzer::BuildMesh()
+{
+	ModelMesh.GetNodeCoordinates(NodeCoordinates);
+
+	int32 LastVertexIndex = 0;
+
+	int32 TriangleCount = 0;
+	for (const TSharedPtr<FFaceMesh>& FaceMesh : ModelMesh.GetFaceMeshes())
+	{
+		TriangleCount += FaceMesh->TrianglesVerticesIndex.Num() / 3;
+	}
+
+	Triangles.Reserve(TriangleCount);
+	Edges.Reserve(TriangleCount*3);
+
+	TMap<uint32, Analyzer::FEdge*> EdgeMap;
+	EdgeMap.Reserve(TriangleCount * 3);
+
+	TFunction<Analyzer::FEdge*(uint32, uint32)> FindOrBuildEdge = [&](uint32 VertexIndex0, uint32 VertexIndex1) -> Analyzer::FEdge*
+	{
+		uint32 EdgeVertexIndex0;
+		uint32 EdgeVertexIndex1;
+		GetMinMax(VertexIndex0, VertexIndex1, EdgeVertexIndex0, EdgeVertexIndex1);
+		uint32 EdgeHash = HashCombine(EdgeVertexIndex0, EdgeVertexIndex1);
+		Analyzer::FEdge** EdgePtr = EdgeMap.Find(EdgeHash);
+		if (EdgePtr != nullptr)
+		{
+			return *EdgePtr;
+		}
+		Analyzer::FEdge& Edge = Edges.Emplace_GetRef(EdgeVertexIndex0, EdgeVertexIndex1);
+		EdgeMap.Add(EdgeHash, &Edge);
+
+		return &Edge;
+	};
+
+	for (const TSharedPtr<FFaceMesh>& FaceMesh : ModelMesh.GetFaceMeshes())
+	{
+		TArray<int32>& TrianglesVerticesIndex = FaceMesh->TrianglesVerticesIndex;
+		TArray<int32>& VerticesGlobalIndex = FaceMesh->VerticesGlobalIndex;
+		TArray<FVector>& FaceNormals = FaceMesh->Normals;
+
+		for (int32 Index = 0; Index < FaceMesh->TrianglesVerticesIndex.Num(); Index += 3)
+		{
+			uint32 Vertices[3];
+			Vertices[0] = VerticesGlobalIndex[TrianglesVerticesIndex[Index + 0]] + LastVertexIndex;
+			Vertices[1] = VerticesGlobalIndex[TrianglesVerticesIndex[Index + 1]] + LastVertexIndex;
+			Vertices[2] = VerticesGlobalIndex[TrianglesVerticesIndex[Index + 2]] + LastVertexIndex;
+
+			Analyzer::FEdge* TriangleEdges[3];
+			TriangleEdges[0] = FindOrBuildEdge(Vertices[0], Vertices[1]);
+			TriangleEdges[1] = FindOrBuildEdge(Vertices[1], Vertices[2]);
+			TriangleEdges[2] = FindOrBuildEdge(Vertices[2], Vertices[0]);
+
+			FVector Normals[3];
+			for (uint32 NormalI = 0; NormalI < 3; ++NormalI)
+			{
+				Normals[NormalI] = FaceNormals[TrianglesVerticesIndex[Index + NormalI]];
+			}
+
+			Analyzer::FTriangle& Triangle = Triangles.Emplace_GetRef(Vertices, TriangleEdges, Normals, FaceMesh.Get());
+
+			for (uint32 EdgeI = 0; EdgeI < 3; ++EdgeI)
+			{
+				TriangleEdges[EdgeI]->AddTriangle(Triangle);
+			}
+		}
+	}
+
+	VertexToEdges.SetNum(NodeCoordinates.Num());
+	for (Analyzer::FEdge& Edge : Edges)
+	{
+		if (Edge.TriangleCount != 2)
+		{
+			VertexToEdges[Edge.VertexIndices[0]].Add(&Edge);
+			VertexToEdges[Edge.VertexIndices[1]].Add(&Edge);
+		}
+	}
+
+	LastVertexIndex += NodeCoordinates.Num();
+}
+
+void FModelMeshAnalyzer::ComputeBorderCount(int32& OutBorderCount, int32& OutNonManifoldCount) const
+{
+	OutBorderCount = 0;
+	OutNonManifoldCount = 0;
+	for (const Analyzer::FEdge& Edge : Edges)
+	{
+		switch (Edge.TriangleCount)
+		{
+		case 1:
+			OutBorderCount++;
+			break;
+
+		case 2:
+			break;
+
+		default:
+			OutNonManifoldCount++;
+			break;
+		}
+	}
+}
+
+void FModelMeshAnalyzer::ComputeMeshGapCount(int32& OutCycleCount, int32& OutChainCount) const
+{
+	OutCycleCount = 0;
+	OutChainCount = 0;
+	for (const Analyzer::FEdge& Edge : Edges)
+	{
+		if (Edge.TriangleCount == 2)
+		{
+			continue;
+		}
+
+		if (Edge.HasMarker1())
+		{
+			continue;
+		}
+		Edge.SetMarker1();
+
+		uint32 FirstVertex = Edge.VertexIndices[0];
+		uint32 NextVertex = Edge.VertexIndices[1];
+		
+		const Analyzer::FEdge* NextEdge = &Edge;
+
+		bool bIsCycle = true;
+		while (NextVertex != FirstVertex)
+		{
+			if (VertexToEdges[NextVertex].Num() != 2)
+			{
+				if (!bIsCycle)
+				{
+					break;
+				}
+				else
+				{
+					Swap(NextVertex, FirstVertex);
+					bIsCycle = false;
+				}
+			}
+			else
+			{
+				NextEdge = VertexToEdges[NextVertex][0] == NextEdge ? VertexToEdges[NextVertex][1] : VertexToEdges[NextVertex][0];
+				NextEdge->SetMarker1();
+				NextVertex = NextEdge->OtherVertexIndex(NextVertex);
+			}
+		}
+
+		if (bIsCycle)
+		{
+			OutCycleCount++;
+		}
+		else
+		{
+			OutChainCount++;
+		}
+	}
+
+	for (const Analyzer::FEdge& Edge : Edges)
+	{
+		Edge.ResetMarker1();
+	}
+
+}
+
+double FModelMeshAnalyzer::ComputeMaxAngle() const 
+{
+	double MaxAngle = 0;
+
+	for (const Analyzer::FEdge& Edge : Edges)
+	{
+		if (Edge.TriangleCount != 2)
+		{
+			continue;
+		}
+
+		if (Edge.Triangles[0]->FaceMesh != Edge.Triangles[1]->FaceMesh)
+		{
+			continue;
+		}
+
+		FPoint Normal0 = Edge.Triangles[0]->ComputeNormal(NodeCoordinates);
+		FPoint Normal1 = Edge.Triangles[1]->ComputeNormal(NodeCoordinates);
+
+		double Angle = Normal0.ComputeAngle(Normal1);
+		if (Angle > MaxAngle)
+		{
+			MaxAngle = Angle;
+		}
+	}
+	return MaxAngle;
+}
+
+bool FModelMeshAnalyzer::CheckOrientation() const
+{
+	return true;
+}
+
+FPoint Analyzer::FTriangle::ComputeNormal(const TArray<FPoint>& NodeCoordinates) const
+{
+	return CADKernel::FTriangle(NodeCoordinates[Vertices[0]], NodeCoordinates[Vertices[1]], NodeCoordinates[Vertices[2]]).ComputeNormal();
+}
+
+} // namespace CADKernel
+
diff --git a/Engine/Source/Runtime/Datasmith/CADKernel/Private/CADKernel/Mesh/Structure/ModelMeshAnalyzer.h b/Engine/Source/Runtime/Datasmith/CADKernel/Private/CADKernel/Mesh/Structure/ModelMeshAnalyzer.h
new file mode 100644
--- /dev/null
+++ b/Engine/Source/Runtime/Datasmith/CADKernel/Private/CADKernel/Mesh/Structure/ModelMeshAnalyzer.h
@@ -0,0 +1,103 @@
+// Copyright Epic Games, Inc. All Rights Reserved.
+#pragma once
+
+#include "CADKernel/Core/HaveStates.h"
+#include "CADKernel/Core/Types.h"
+
+#include "CADKernel/Math/Point.h"
+
+namespace CADKernel
+{
+
+class FModelMesh;
+class FFaceMesh;
+
+namespace Analyzer
+{
+
+struct FEdge;
+
+struct FTriangle
+{
+	uint32 Vertices[3];
+	FVector Normals[3];
+	FEdge* Edges[3];
+	const FFaceMesh* FaceMesh;
+
+
+	FTriangle(uint32* InVertices, FEdge** InEdges, FVector* InNormals, const FFaceMesh* InFaceMesh)
+	: FaceMesh(InFaceMesh)
+	{
+		for (int32 Index = 0; Index < 3; ++Index)
+		{
+			Vertices[Index] = InVertices[Index];
+		}
+		for (int32 Index = 0; Index < 3; ++Index)
+		{
+			Edges[Index] = InEdges[Index];
+		}
+		for (int32 Index = 0; Index < 3; ++Index)
+		{
+			Normals[Index] = InNormals[Index];
+		}
+	}
+
+	FPoint ComputeNormal(const TArray<FPoint>& NodeCoordinates) const;
+};
+
+struct FEdge : public FHaveStates
+{
+	static const int32 MaxTraingleStoredCount = 2;
+	uint32 VertexIndices[2];
+	FTriangle* Triangles[MaxTraingleStoredCount]; // n first adjacent triangles
+	uint32 TriangleCount = 0;
+
+	FEdge(uint32 Index0, uint32 Index1)
+	{
+		VertexIndices[0] = Index0;
+		VertexIndices[1] = Index1;
+	}
+
+	void AddTriangle(FTriangle& Triangle)
+	{
+		if (TriangleCount < MaxTraingleStoredCount)
+		{
+			Triangles[TriangleCount] = &Triangle;
+		}
+		TriangleCount++;
+	}
+
+	uint32 OtherVertexIndex(uint32 Index) const
+	{
+		return Index == VertexIndices[0] ? VertexIndices[1] : VertexIndices[0];
+	}
+};
+
+}
+
+class CADKERNEL_API FModelMeshAnalyzer
+{
+	const FModelMesh& ModelMesh;
+
+	TArray<FPoint> NodeCoordinates;
+	TArray<Analyzer::FEdge> Edges;
+	TArray<Analyzer::FTriangle> Triangles;
+	TArray<TArray<Analyzer::FEdge*>> VertexToEdges;  // Vertex to border or non manifold edges 
+
+public:
+	FModelMeshAnalyzer(const FModelMesh& InModelMesh)
+		: ModelMesh(InModelMesh)
+	{
+	}
+
+	void BuildMesh();
+
+	void ComputeBorderCount(int32& OutBorderCount, int32& OutNonManifoldCount) const;
+	void ComputeMeshGapCount(int32& OutCycleCount, int32& OutChainCount) const;
+	double ComputeMaxAngle() const;
+	bool CheckOrientation() const;
+
+};
+
+}
+
diff --git a/Engine/Source/Runtime/Datasmith/CADKernel/Private/CADKernel/Mesh/Structure/ModelMeshReport.cpp b/Engine/Source/Runtime/Datasmith/CADKernel/Private/CADKernel/Mesh/Structure/ModelMeshReport.cpp
new file mode 100644
--- /dev/null
+++ b/Engine/Source/Runtime/Datasmith/CADKernel/Private/CADKernel/Mesh/Structure/ModelMeshReport.cpp
@@ -0,0 +1,56 @@
+// Copyright Epic Games, Inc. All Rights Reserved.
+
+#include "CADKernel/Mesh/Structure/ModelMeshReport.h"
+
+#include "CADKernel/Mesh/Structure/ModelMesh.h"
+#include "CADKernel/Mesh/Structure/ModelMeshAnalyzer.h"
+
+namespace CADKernel
+{
+
+void FModelMeshReport::Run()
+{
+	for (const FModelMesh* ModelMesh : ModelMeshes)
+	{
+		bool MeshGoodOrientation = true;
+		int32 MeshBorderEdgeCount = 0;
+		int32 MeshNonManifoldEdgeCount = 0;
+		int32 MeshCycleCount = 0;
+		int32 MeshChainCount = 0;
+		double MeshMaxAngle = 0;
+
+		FModelMeshAnalyzer ModelMeshAnalyzer(*ModelMesh);
+		ModelMeshAnalyzer.BuildMesh();
+
+		if (!ModelMeshAnalyzer.CheckOrientation())
+		{
+			GoodOrientation = false;
+		}
+
+		ModelMeshAnalyzer.ComputeBorderCount(MeshBorderEdgeCount, MeshNonManifoldEdgeCount);
+		ModelMeshAnalyzer.ComputeMeshGapCount(MeshCycleCount, MeshChainCount);
+		BorderEdgeCount += MeshBorderEdgeCount;
+		NonManifoldEdgeCount += MeshNonManifoldEdgeCount;
+		CycleCount += MeshCycleCount;
+		ChainCount += MeshChainCount;
+
+		MeshMaxAngle = ModelMeshAnalyzer.ComputeMaxAngle();
+		if (MaxAngle < MeshMaxAngle)
+		{
+			MaxAngle = MeshMaxAngle;
+		}
+	}
+}
+
+void FModelMeshReport::Print()
+{
+	FMessage::FillReportFile(TEXT(""), GoodOrientation ? TEXT("True") : TEXT("False"));
+	FMessage::FillReportFile(TEXT(""), BorderEdgeCount);
+	FMessage::FillReportFile(TEXT(""), NonManifoldEdgeCount);
+	FMessage::FillReportFile(TEXT(""), CycleCount);
+	FMessage::FillReportFile(TEXT(""), ChainCount);
+	FMessage::FillReportFile(TEXT(""), TEXT(""));
+}
+
+} // namespace CADKernel
+
diff --git a/Engine/Source/Runtime/Datasmith/CADKernel/Private/CADKernel/Mesh/Structure/ModelMeshReport.h b/Engine/Source/Runtime/Datasmith/CADKernel/Private/CADKernel/Mesh/Structure/ModelMeshReport.h
new file mode 100644
--- /dev/null
+++ b/Engine/Source/Runtime/Datasmith/CADKernel/Private/CADKernel/Mesh/Structure/ModelMeshReport.h
@@ -0,0 +1,34 @@
+// Copyright Epic Games, Inc. All Rights Reserved.
+#pragma once
+
+#include "CADKernel/Core/Types.h"
+
+namespace CADKernel
+{
+class FModelMesh;
+
+class CADKERNEL_API FModelMeshReport
+{
+	const TArray<const FModelMesh*>& ModelMeshes;
+
+	bool GoodOrientation = true;
+	int32 BorderEdgeCount = 0;
+	int32 NonManifoldEdgeCount = 0;
+	int32 CycleCount = 0;
+	int32 ChainCount = 0;
+	double MaxAngle = 0;
+
+public:
+	FModelMeshReport(const TArray<const FModelMesh*>& InModelMeshes)
+		: ModelMeshes(InModelMeshes)
+	{
+	}
+
+	void Run();
+
+	void Print();
+
+};
+
+}
+
diff --git a/Engine/Source/Runtime/Datasmith/CADKernel/Private/CADKernel/Topo/Body.cpp b/Engine/Source/Runtime/Datasmith/CADKernel/Private/CADKernel/Topo/Body.cpp
--- a/Engine/Source/Runtime/Datasmith/CADKernel/Private/CADKernel/Topo/Body.cpp
+++ b/Engine/Source/Runtime/Datasmith/CADKernel/Private/CADKernel/Topo/Body.cpp
@@ -2,6 +2,8 @@
 #include "CADKernel/Topo/Body.h"
 
 #include "CADKernel/Topo/Shell.h"
+#include "CADKernel/Topo/TopologyReport.h"
+
 
 namespace CADKernel
 {
@@ -43,4 +45,13 @@ FInfoEntity& FBody::GetInfo(FInfoEntity& Info) const
 }
 #endif
 
+void FBody::FillTopologyReport(FTopologyReport& Report) const
+{
+	Report.Add(this);
+	for (TSharedPtr<FShell> Shell : Shells)
+	{
+		Shell->FillTopologyReport(Report);
+	}
+}
+
 }
diff --git a/Engine/Source/Runtime/Datasmith/CADKernel/Private/CADKernel/Topo/Model.cpp b/Engine/Source/Runtime/Datasmith/CADKernel/Private/CADKernel/Topo/Model.cpp
--- a/Engine/Source/Runtime/Datasmith/CADKernel/Private/CADKernel/Topo/Model.cpp
+++ b/Engine/Source/Runtime/Datasmith/CADKernel/Private/CADKernel/Topo/Model.cpp
@@ -3,11 +3,11 @@
 
 #include "CADKernel/Core/EntityGeom.h"
 #include "CADKernel/Topo/Body.h"
-#include "CADKernel/Topo/TopologicalFace.h"
+#include "CADKernel/Topo/Joiner.h"
 #include "CADKernel/Topo/TopologicalEdge.h"
-#include "CADKernel/Topo/TopologicalVertex.h"
 #include "CADKernel/Topo/TopologicalFace.h"
-#include "CADKernel/Topo/Joiner.h"
+#include "CADKernel/Topo/TopologicalVertex.h"
+#include "CADKernel/Topo/TopologyReport.h"
 
 namespace CADKernel
 {
@@ -189,7 +189,6 @@ struct FBodyShell
 
 };
 
-//void FModel::CheckTopology()
 void FModel::CheckTopology() 
 {
 	TArray<FBodyShell> IsolatedBodies;
@@ -248,6 +247,18 @@ for (const FFaceSubset& Subset : SubShells)
 	}
 }
 
+void FModel::FillTopologyReport(FTopologyReport& Report) const
+{
+	for (TSharedPtr<FBody> Body : Bodies)
+	{
+		Body->FillTopologyReport(Report);
+	}
+
+	for (TSharedPtr<FTopologicalFace> Face : Faces)
+	{
+		Face->FillTopologyReport(Report);
+	}
+}
 
 void FModel::Split(TSharedPtr<FBody> Body, TArray<TSharedPtr<FBody>>& OutNewBody)
 {
diff --git a/Engine/Source/Runtime/Datasmith/CADKernel/Private/CADKernel/Topo/Shell.cpp b/Engine/Source/Runtime/Datasmith/CADKernel/Private/CADKernel/Topo/Shell.cpp
--- a/Engine/Source/Runtime/Datasmith/CADKernel/Private/CADKernel/Topo/Shell.cpp
+++ b/Engine/Source/Runtime/Datasmith/CADKernel/Private/CADKernel/Topo/Shell.cpp
@@ -4,6 +4,7 @@
 #include "CADKernel/Geo/GeoEnum.h"
 #include "CADKernel/Topo/Body.h"
 #include "CADKernel/Topo/TopologicalFace.h"
+#include "CADKernel/Topo/TopologyReport.h"
 
 namespace CADKernel
 {
@@ -267,4 +268,13 @@ void FShell::CheckTopology(TArray<FFaceSubset>& SubShells)
 	}
 }
 
+void FShell::FillTopologyReport(FTopologyReport& Report) const 
+{
+	Report.Add(this);
+	for (const FOrientedFace& OrientedFace : GetFaces())
+	{
+		OrientedFace.Entity->FillTopologyReport(Report);
+	}
+}
+
 }
\ No newline at end of file
diff --git a/Engine/Source/Runtime/Datasmith/CADKernel/Private/CADKernel/Topo/TopologicalEdge.cpp b/Engine/Source/Runtime/Datasmith/CADKernel/Private/CADKernel/Topo/TopologicalEdge.cpp
--- a/Engine/Source/Runtime/Datasmith/CADKernel/Private/CADKernel/Topo/TopologicalEdge.cpp
+++ b/Engine/Source/Runtime/Datasmith/CADKernel/Private/CADKernel/Topo/TopologicalEdge.cpp
@@ -49,8 +49,8 @@ FTopologicalEdge::FTopologicalEdge(const TSharedRef<FRestrictionCurve>& InCurve,
 	TArray<FCurvePoint> Points;
 	Curve->EvaluatePoints(Coordinates, Points);
 
-	StartVertex = FEntity::MakeShared<FTopologicalVertex>(Points[0].Point);
-	EndVertex = FEntity::MakeShared<FTopologicalVertex>(Points[1].Point);
+	StartVertex = FTopologicalVertex::Make(Points[0].Point);
+	EndVertex = FTopologicalVertex::Make(Points[1].Point);
 }
 
 FTopologicalEdge::FTopologicalEdge(const TSharedRef<FSurface>& InSurface, const FPoint2D& InCoordinateVertex1, const TSharedRef<FTopologicalVertex>& InVertex1, const FPoint2D& InCoordinateVertex2, const TSharedRef<FTopologicalVertex>& InVertex2)
@@ -162,6 +162,7 @@ TSharedPtr<FTopologicalEdge> FTopologicalEdge::ReturnIfValid(TSharedRef<FTopolog
 		return TSharedPtr<FTopologicalEdge>();
 	}
 
+	InEdge->Finalize();
 	InEdge->LinkVertex();
 	return InEdge;
 }
@@ -524,8 +525,7 @@ bool FTopologicalEdge::IsSameDirection(const FTopologicalEdge& Edge) const
 		Edge.SetAsDegenerated();
 		return true;
 
-		// TODO
-		ensureCADKernel(false);
+		// Todo: Two cycles can have connected extremities but not with the same orientation
 	}
 
 	return vertex1Edge == Edge.GetStartVertex()->GetLink();
@@ -946,7 +946,7 @@ TSharedPtr<FTopologicalVertex> FTopologicalEdge::SplitAt(double SplittingCoordin
 		// TODO
 	}
 
-	TSharedRef<FTopologicalVertex> MiddelVertex = FEntity::MakeShared<FTopologicalVertex>(NewVertexCoordinate);
+	TSharedRef<FTopologicalVertex> MiddelVertex = FTopologicalVertex::Make(NewVertexCoordinate);
 
 	if (bKeepStartVertexConnectivity)
 	{
diff --git a/Engine/Source/Runtime/Datasmith/CADKernel/Private/CADKernel/Topo/TopologicalFace.cpp b/Engine/Source/Runtime/Datasmith/CADKernel/Private/CADKernel/Topo/TopologicalFace.cpp
--- a/Engine/Source/Runtime/Datasmith/CADKernel/Private/CADKernel/Topo/TopologicalFace.cpp
+++ b/Engine/Source/Runtime/Datasmith/CADKernel/Private/CADKernel/Topo/TopologicalFace.cpp
@@ -11,6 +11,7 @@
 #include "CADKernel/Mesh/Structure/FaceMesh.h"
 #include "CADKernel/Mesh/Structure/Grid.h"
 #include "CADKernel/Topo/TopologicalEdge.h"
+#include "CADKernel/Topo/TopologyReport.h"
 
 namespace CADKernel
 {
@@ -101,7 +102,7 @@ void FTopologicalFace::ApplyNaturalLoops(const FSurfacicBoundary& Boundaries)
 	AddLoop(Loop);
 }
 
-void FTopologicalFace::AddLoops(const TArray<TSharedPtr<FTopologicalLoop>>& InLoops)
+void FTopologicalFace::AddLoops(const TArray<TSharedPtr<FTopologicalLoop>>& InLoops, int32& DoubtfulLoopOrientationCount)
 {
 	for (TSharedPtr<FTopologicalLoop> Loop : InLoops)
 	{
@@ -110,7 +111,10 @@ void FTopologicalFace::AddLoops(const TArray<TSharedPtr<FTopologicalLoop>>& InLo
 
 	for (TSharedPtr<FTopologicalLoop> Loop : InLoops)
 	{
-		Loop->Orient();
+		if (!Loop->Orient())
+		{
+			DoubtfulLoopOrientationCount++;
+		}
 	}
 }
 
@@ -217,6 +221,19 @@ const FTopologicalEdge* FTopologicalFace::GetLinkedEdge(const FTopologicalEdge&
 	return nullptr;
 }
 
+void FTopologicalFace::FillTopologyReport(FTopologyReport& Report) const
+{
+	Report.Add(this);
+
+	for (const TSharedPtr<FTopologicalLoop>& Loop : GetLoops())
+	{
+		for (const FOrientedEdge& Edge : Loop->GetEdges())
+		{
+			Report.Add(Edge.Entity.Get());
+		}
+	}
+}
+
 void FTopologicalFace::GetEdgeIndex(const FTopologicalEdge& Edge, int32& OutBoundaryIndex, int32& OutEdgeIndex) const
 {
 	OutEdgeIndex = INDEX_NONE;
diff --git a/Engine/Source/Runtime/Datasmith/CADKernel/Private/CADKernel/Topo/TopologicalLoop.cpp b/Engine/Source/Runtime/Datasmith/CADKernel/Private/CADKernel/Topo/TopologicalLoop.cpp
--- a/Engine/Source/Runtime/Datasmith/CADKernel/Private/CADKernel/Topo/TopologicalLoop.cpp
+++ b/Engine/Source/Runtime/Datasmith/CADKernel/Private/CADKernel/Topo/TopologicalLoop.cpp
@@ -206,9 +206,14 @@ bool FTopologicalLoop::Get2DSamplingWithoutDegeneratedEdges(TArray<FPoint2D>& Lo
  *
  * If the slop is still closed to 0 or 8, the point is "UndefinedOrientation"
  * This case was found in the shape of char '1'
+ * 
+ * @return false if the orientation is doubtful
+ * 
  */
-void FTopologicalLoop::Orient()
+bool FTopologicalLoop::Orient()
 {
+	bool bSucceed = true;
+
 	ensureCADKernel(Edges.Num() > 0);
 
 	TArray<FPoint2D> LoopSampling;
@@ -216,7 +221,7 @@ void FTopologicalLoop::Orient()
 	if (!Get2DSamplingWithoutDegeneratedEdges(LoopSampling))
 	{
 		// the loop is degenerated
-		return;
+		return false;
 	}
 
 #ifdef DEBUG_ORIENT
@@ -456,7 +461,8 @@ void FTopologicalLoop::Orient()
 		Wait();
 #endif
 
-		FMessage::Printf(Log, TEXT("WARNING: Loop Orientation of surface %d could failed\n"), Face.Pin()->GetId());
+		bSucceed = false;
+		FMessage::Printf(Log, TEXT("WARNING: Loop Orientation of surface %d is doubtful\n"), Face.Pin()->GetId());
 	}
 
 	if ((WrongOrientationCount > GoodOrientationCount) == bExternalLoop)
@@ -477,6 +483,7 @@ void FTopologicalLoop::Orient()
 	}
 #endif
 
+	return bSucceed;
 }
 
 void FTopologicalLoop::SwapOrientation()
diff --git a/Engine/Source/Runtime/Datasmith/CADKernel/Private/CADKernel/Topo/TopologicalVertex.cpp b/Engine/Source/Runtime/Datasmith/CADKernel/Private/CADKernel/Topo/TopologicalVertex.cpp
--- a/Engine/Source/Runtime/Datasmith/CADKernel/Private/CADKernel/Topo/TopologicalVertex.cpp
+++ b/Engine/Source/Runtime/Datasmith/CADKernel/Private/CADKernel/Topo/TopologicalVertex.cpp
@@ -67,7 +67,7 @@ bool FTopologicalVertex::IsBorderVertex()
 TSharedPtr<FEntityGeom> FTopologicalVertex::ApplyMatrix(const FMatrixH& InMatrix) const
 {
 	FPoint transformedPoint = InMatrix.Multiply(Coordinates);
-	return FEntity::MakeShared<FTopologicalVertex>(transformedPoint);
+	return FTopologicalVertex::Make(transformedPoint);
 }
 
 void FTopologicalVertex::GetConnectedEdges(const FTopologicalVertex& OtherVertex, TArray<FTopologicalEdge*>& OutEdges) const
diff --git a/Engine/Source/Runtime/Datasmith/CADKernel/Private/CADKernel/Topo/TopologyReport.cpp b/Engine/Source/Runtime/Datasmith/CADKernel/Private/CADKernel/Topo/TopologyReport.cpp
new file mode 100644
--- /dev/null
+++ b/Engine/Source/Runtime/Datasmith/CADKernel/Private/CADKernel/Topo/TopologyReport.cpp
@@ -0,0 +1,194 @@
+//CopyrightEpicGames,Inc.AllRightsReserved.
+#pragma once
+
+#include "TopologyReport.h"
+
+#include "CADKernel/Topo/Body.h"
+#include "CADKernel/Topo/Shell.h"
+#include "CADKernel/Topo/TopologicalFace.h"
+#include "CADKernel/Topo/TopologicalEdge.h"
+
+#include "CADKernel/UI/Message.h"
+
+namespace CADKernel
+{
+
+bool FTopologyReport::HasMarker(const FTopologicalEntity* Entity)
+{
+	if (!Entity->HasMarker1())
+	{
+		Entity->SetMarker1();
+		Entities.Add(Entity);
+		return false;
+	}
+	return true;
+}
+
+void FTopologyReport::Add(const FBody* Body)
+{
+	if (!HasMarker(Body))
+	{
+		BodyCount++;
+	}
+}
+
+void FTopologyReport::Add(const FShell* Shell)
+{
+	if (!HasMarker(Shell))
+	{
+		ShellCount++;
+	}
+}
+
+void FTopologyReport::Add(const FTopologicalFace* Face)
+{
+	if (!HasMarker(Face))
+	{
+		FaceCount++;
+	}
+}
+
+void FTopologyReport::Add(const FTopologicalEdge* Edge)
+{
+	EdgeCount++;
+
+	const TSharedRef<const FTopologicalEdge>& ActiveEdge = Edge->GetLinkActiveEdge();
+	if (!HasMarker(&*ActiveEdge))
+	{
+		ActiveEdge->SetMarker1();
+
+		CoedgeCount++;
+
+		switch (ActiveEdge->GetTwinsEntityCount())
+		{
+		case 1:
+			BorderEdgeCount++;
+			Edges.Add(&ActiveEdge.Get());
+			break;
+		case 2:
+			SurfaceEdgeCount++;
+			break;
+		default:
+			NonManifoldEdgeCount++;
+			Edges.Add(&ActiveEdge.Get());
+			break;
+		}
+	}
+}
+
+void FTopologyReport::CountLoops()
+{
+	LoopCount = 0;
+	ChainCount = 0;
+
+	for (const FTopologicalEdge* Edge : Edges)
+	{
+		Edge->SetMarker2();
+	}
+
+	for (const FTopologicalEdge* Edge : Edges)
+	{
+		if (Edge->HasMarker1())
+		{
+			continue;
+		}
+		Edge->SetMarker1();
+
+		if (Edge->GetTwinsEntityCount() == 2)
+		{
+			continue;
+		}
+
+		const FTopologicalVertex* FirstVertex = &Edge->GetStartVertex()->GetLinkActiveEntity().Get();
+		const FTopologicalVertex* NextVertex = &Edge->GetEndVertex()->GetLinkActiveEntity().Get();
+
+		const FTopologicalEdge* NextEdge = Edge;
+
+		TArray<FTopologicalEdge*> ConnectedEdges;
+
+		bool bIsCycle = true;
+		while (NextVertex != FirstVertex)
+		{
+			ConnectedEdges.Empty();
+			NextVertex->GetConnectedEdges(ConnectedEdges);
+
+			int32 BorderCount = 0;
+			const FTopologicalEdge* VertexEdges[2] = { nullptr, nullptr };
+			for (const FTopologicalEdge* ConnectedEdge : ConnectedEdges)
+			{
+				if (ConnectedEdge->HasMarker2())
+				{
+					if (BorderCount < 2)
+					{
+						VertexEdges[BorderCount++] = ConnectedEdge;
+					}
+					else
+					{
+						BorderCount++;
+						break;
+					}
+				}
+			}
+
+			if (BorderCount != 2)
+			{
+				if (!bIsCycle)
+				{
+					break;
+				}
+				else
+				{
+					Swap(NextVertex, FirstVertex);
+					bIsCycle = false;
+				}
+			}
+			else
+			{
+				NextEdge = VertexEdges[0] == NextEdge ? VertexEdges[1] : VertexEdges[0];
+				NextEdge->SetMarker1();
+				NextVertex = &NextEdge->GetOtherVertex(*NextVertex)->GetLinkActiveEntity().Get();
+			}
+		}
+
+		if (bIsCycle)
+		{
+			LoopCount++;
+		}
+		else
+		{
+			ChainCount++;
+		}
+	}
+
+	for (const FTopologicalEdge* Edge : Edges)
+	{
+		Edge->ResetMarkers();
+	}
+
+}
+
+
+void FTopologyReport::Print()
+{
+	for (const FTopologicalEntity* Entity : Entities)
+	{
+		Entity->ResetMarker1();
+	}
+
+	CountLoops();
+
+	FMessage::FillReportFile(TEXT("Body"), BodyCount);
+	FMessage::FillReportFile(TEXT("Shell"), ShellCount);
+	FMessage::FillReportFile(TEXT("Face"), FaceCount);
+	FMessage::FillReportFile(TEXT("Edge"), EdgeCount);
+	FMessage::FillReportFile(TEXT(""), TEXT(""));
+	FMessage::FillReportFile(TEXT("CoEdge"), CoedgeCount);
+	FMessage::FillReportFile(TEXT("Surface"), SurfaceEdgeCount);
+	FMessage::FillReportFile(TEXT("Border"), BorderEdgeCount);
+	FMessage::FillReportFile(TEXT("NManif"), NonManifoldEdgeCount);
+	FMessage::FillReportFile(TEXT(""), TEXT(""));
+	FMessage::FillReportFile(TEXT("Loop"), LoopCount);
+	FMessage::FillReportFile(TEXT("Chain"), ChainCount);
+}
+
+}
diff --git a/Engine/Source/Runtime/Datasmith/CADKernel/Private/CADKernel/Topo/TopologyReport.h b/Engine/Source/Runtime/Datasmith/CADKernel/Private/CADKernel/Topo/TopologyReport.h
new file mode 100644
--- /dev/null
+++ b/Engine/Source/Runtime/Datasmith/CADKernel/Private/CADKernel/Topo/TopologyReport.h
@@ -0,0 +1,56 @@
+//CopyrightEpicGames,Inc.AllRightsReserved.
+#pragma once
+
+#include "CADKernel/Core/Chrono.h"
+#include "CADKernel/Core/Types.h"
+#include "CADKernel/UI/Message.h"
+
+namespace CADKernel
+{
+
+class FBody;
+class FTopologicalEdge;
+class FTopologicalFace;
+class FTopologicalEntity;
+class FShell;
+
+class FTopologyReport
+{
+private:
+	int32 BodyCount = 0;
+	int32 ShellCount = 0;
+	int32 FaceCount = 0;
+	int32 EdgeCount = 0;
+
+	int32 CoedgeCount = 0;
+	int32 NonManifoldEdgeCount = 0;
+	int32 SurfaceEdgeCount = 0;
+	int32 BorderEdgeCount = 0;
+
+	int32 LoopCount = 0;
+	int32 ChainCount = 0;
+
+	TArray<const FTopologicalEdge*> Edges;
+	TArray<const FTopologicalEntity*> Entities;
+
+	bool HasMarker(const FTopologicalEntity* Entity);
+
+public:
+
+	FTopologyReport()
+	{
+		Entities.Reserve(100000);
+		Edges.Reserve(20000);
+	}
+
+	void Add(const FBody* Body);
+	void Add(const FTopologicalEdge* Edge);
+	void Add(const FTopologicalFace* Face);
+	void Add(const FShell* Shell);
+
+	void CountLoops();
+
+	void Print();
+};
+
+}
diff --git a/Engine/Source/Runtime/Datasmith/CADKernel/Private/CADKernel/UI/Message.cpp b/Engine/Source/Runtime/Datasmith/CADKernel/Private/CADKernel/UI/Message.cpp
--- a/Engine/Source/Runtime/Datasmith/CADKernel/Private/CADKernel/UI/Message.cpp
+++ b/Engine/Source/Runtime/Datasmith/CADKernel/Private/CADKernel/UI/Message.cpp
@@ -11,174 +11,166 @@
 
 namespace CADKernel
 {
-const char* VerboseLevelConstNames[] = {
-	("NO_VERBOSE"),
-	("SPY"),
-	("LOG"),
-	("DBG"),
-	nullptr
-};
-
-const char* VerboseConstDescHelp[] = {
-	": nothing is printed",
-	": only commands and arguments will be showed",
-	": SPY + messages",
-	": debug mode, for developers only",
-	nullptr
-};
-
-int32 FMessage::NumberOfIndentation = 0;
-int32 FMessage::OldPercent = -1;
-
-void FMessage::VPrintf(EVerboseLevel Level, const TCHAR* Text, ...)
-{
-	EVerboseLevel VerboseLevel = (EVerboseLevel)(int32)FSystem::Get().GetVerboseLevel();
-	TSharedPtr<FArchive> LogFile = FSystem::Get().GetLogFile();
-	EVerboseLevel LogLevel = FSystem::Get().GetLogLevel();
-	TSharedPtr<FArchive> SpyFile = FSystem::Get().GetSpyFile();
-
-	FString Indentation;
-	TCHAR* Buffer = NULL;
-	int32 Result = -1;
-	if (VerboseLevel >= Level || (LogFile.IsValid() && Level <= EVerboseLevel::Debug) || (SpyFile.IsValid() && Level == EVerboseLevel::Spy))
+	const char* VerboseLevelConstNames[] = {
+		("NO_VERBOSE"),
+		("SPY"),
+		("LOG"),
+		("DBG"),
+		nullptr
+	};
+
+	const char* VerboseConstDescHelp[] = {
+		": nothing is printed",
+		": only commands and arguments will be showed",
+		": SPY + messages",
+		": debug mode, for developers only",
+		nullptr
+	};
+
+	int32 FMessage::NumberOfIndentation = 0;
+	int32 FMessage::OldPercent = -1;
+
+	void FMessage::VPrintf(EVerboseLevel Level, const TCHAR* Text, ...)
 	{
-		int32 BufferSize = 1024;
+		EVerboseLevel VerboseLevel = (EVerboseLevel)(int32)FSystem::Get().GetVerboseLevel();
+		TSharedPtr<FArchive> LogFile = FSystem::Get().GetLogFile();
+		EVerboseLevel LogLevel = FSystem::Get().GetLogLevel();
+		TSharedPtr<FArchive> SpyFile = FSystem::Get().GetSpyFile();
+
+		FString Indentation;
+		TCHAR* Buffer = NULL;
+		int32 Result = -1;
+		if (VerboseLevel >= Level || (LogFile.IsValid() && Level <= EVerboseLevel::Debug) || (SpyFile.IsValid() && Level == EVerboseLevel::Spy))
+		{
+			int32 BufferSize = 1024;
+
+			// do the usual VARARGS shenanigans
+			while (Result == -1)
+			{
+				FMemory::Free(Buffer);
+				Buffer = (TCHAR*)FMemory::Malloc(BufferSize * sizeof(TCHAR));
+				GET_VARARGS_RESULT(Buffer, BufferSize, BufferSize - 1, Text, Text, Result);
+				BufferSize *= 2;
+			};
+			Buffer[Result] = 0;
+
+			Indentation.Reserve(NumberOfIndentation * 3);
+			for (int32 iIndent = 0; iIndent < NumberOfIndentation; iIndent++)
+			{
+				Indentation += TEXT(" - ");
+			}
+		}
 
-		// do the usual VARARGS shenanigans
-		while (Result == -1)
+		if (VerboseLevel >= Level)
 		{
-			FMemory::Free(Buffer);
-			Buffer = (TCHAR*)FMemory::Malloc(BufferSize * sizeof(TCHAR));
-			GET_VARARGS_RESULT(Buffer, BufferSize, BufferSize - 1, Text, Text, Result);
-			BufferSize *= 2;
-		};
-		Buffer[Result] = 0;
-
-		Indentation.Reserve(NumberOfIndentation * 3);
-		for (int32 iIndent = 0; iIndent < NumberOfIndentation; iIndent++)
+			FSystem::Get().GetConsole().Print(*Indentation, Level);
+			FSystem::Get().GetConsole().Print(Buffer, Level);
+		}
+
+		if ((LogFile.IsValid() && Level <= LogLevel) || (SpyFile.IsValid() && Level <= Spy))
 		{
-			Indentation += TEXT(" - ");
+			FTCHARToUTF8 UTF8Indentation(*Indentation, Indentation.Len());
+			FTCHARToUTF8 UTF8String(Buffer, Result);
+			if (LogFile.IsValid() && Level <= LogLevel)
+			{
+				LogFile->Serialize((UTF8CHAR*)UTF8Indentation.Get(), UTF8Indentation.Length() * sizeof(UTF8CHAR));
+				LogFile->Serialize((UTF8CHAR*)UTF8String.Get(), UTF8String.Length() * sizeof(UTF8CHAR));
+				LogFile->Flush();
+			}
+
+			if (SpyFile.IsValid() && Level <= Spy)
+			{
+				SpyFile->Serialize((UTF8CHAR*)UTF8Indentation.Get(), UTF8Indentation.Length() * sizeof(UTF8CHAR));
+				SpyFile->Serialize((UTF8CHAR*)UTF8String.Get(), UTF8String.Length() * sizeof(UTF8CHAR));
+				SpyFile->Flush();
+			}
 		}
-	}
 
-	if (VerboseLevel >= Level)
-	{
-		FSystem::Get().GetConsole().Print(*Indentation, Level);
-		FSystem::Get().GetConsole().Print(Buffer, Level);
+		if (Buffer)
+		{
+			FMemory::Free(Buffer);
+		}
 	}
 
-	if ((LogFile.IsValid() && Level <= LogLevel) || (SpyFile.IsValid() && Level <= Spy))
+	void FMessage::VReportPrintF(FString Header, const TCHAR* Text, ...)
 	{
-		FTCHARToUTF8 UTF8Indentation(*Indentation, Indentation.Len());
-		FTCHARToUTF8 UTF8String(Buffer, Result);
-		if (LogFile.IsValid() && Level <= LogLevel)
+#ifdef CADKERNEL_DEV
+		TSharedPtr<FArchive> ReportFile = FSystem::Get().GetReportFile();
+		TSharedPtr<FArchive> ReportHeaderFile = FSystem::Get().GetReportHeaderFile();
+		if (!ReportFile.IsValid())
 		{
-			LogFile->Serialize((UTF8CHAR*)UTF8Indentation.Get(), UTF8Indentation.Length() * sizeof(UTF8CHAR));
-			LogFile->Serialize((UTF8CHAR*)UTF8String.Get(), UTF8String.Length() * sizeof(UTF8CHAR));
-			LogFile->Flush();
+			return;
 		}
 
-		if (SpyFile.IsValid() && Level <= Spy)
+		TCHAR* Buffer = NULL;
+		int32 Result = -1;
 		{
-			SpyFile->Serialize((UTF8CHAR*)UTF8Indentation.Get(), UTF8Indentation.Length() * sizeof(UTF8CHAR));
-			SpyFile->Serialize((UTF8CHAR*)UTF8String.Get(), UTF8String.Length() * sizeof(UTF8CHAR));
-			SpyFile->Flush();
+			int32 BufferSize = 1024;
+
+			// do the usual VARARGS shenanigans
+			while (Result == -1)
+			{
+				FMemory::Free(Buffer);
+				Buffer = (TCHAR*)FMemory::Malloc(BufferSize * sizeof(TCHAR));
+				GET_VARARGS_RESULT(Buffer, BufferSize, BufferSize - 1, Text, Text, Result);
+				BufferSize *= 2;
+			};
+			Buffer[Result] = 0;
 		}
-	}
 
-	if (Buffer)
-	{
-		FMemory::Free(Buffer);
-	}
-}
+		//FSystem::Get().GetConsole().Print(*Header, Log);
+		//FSystem::Get().GetConsole().Print(TEXT(": "), Log);
+		//FSystem::Get().GetConsole().Print(Buffer, Log);
+		//FSystem::Get().GetConsole().Print(TEXT("\n"), Log);
 
-#if defined(CADKERNEL_DEV) || defined(CADKERNEL_STDA)
-void FMessage::VQaPrintF(const TCHAR* Header, const TCHAR* Text, ...)
-{
-	TSharedPtr<FArchive> QaDataFile = FSystem::Get().GetQaDataFile();
-	TSharedPtr<FArchive> QaHeaderFile = FSystem::Get().GetQaHeaderFile();
-	if (!QaDataFile.IsValid())
-	{
-		return;
-	}
-
-	FString Indentation;
-	TCHAR* Buffer = NULL;
-	int32 Result = -1;
+		FTCHARToUTF8 BufferUtf8(Buffer, Result);
+		ReportFile->Serialize((UTF8CHAR*)BufferUtf8.Get(), BufferUtf8.Length() * sizeof(UTF8CHAR));
+		ReportFile->Serialize(",", sizeof(UTF8CHAR));
+		ReportFile->Flush();
 
-	{
-		int32 BufferSize = 1024;
+		if (ReportHeaderFile.IsValid())
+		{
+			FTCHARToUTF8 HeaderUtf8(*Header, Header.Len());
+			ReportHeaderFile->Serialize((UTF8CHAR*)HeaderUtf8.Get(), HeaderUtf8.Length() * sizeof(UTF8CHAR));
+			ReportHeaderFile->Serialize(",", sizeof(UTF8CHAR));
+			ReportHeaderFile->Flush();
+		}
 
-		// do the usual VARARGS shenanigans
-		while (Result == -1)
+		if (Buffer)
 		{
 			FMemory::Free(Buffer);
-			Buffer = (TCHAR*)FMemory::Malloc(BufferSize * sizeof(TCHAR));
-			GET_VARARGS_RESULT(Buffer, BufferSize, BufferSize - 1, Text, Text, Result);
-			BufferSize *= 2;
-		};
-		Buffer[Result] = 0;
-
-		Indentation.Reserve(NumberOfIndentation * 3);
-		for (int32 iIndent = 0; iIndent < NumberOfIndentation; iIndent++)
-		{
-			Indentation += TEXT(" - ");
 		}
+#endif
 	}
 
-	FSystem::Get().GetConsole().Print(Header, Log);
-	FSystem::Get().GetConsole().Print(TEXT("\n"), Log);
-	FSystem::Get().GetConsole().Print(Buffer, Log);
-	FSystem::Get().GetConsole().Print(TEXT("\n"), Log);
-
-	FString HeaderStr = Header;
-	FTCHARToUTF8 UTF8String(Header, HeaderStr.Len());
-	QaDataFile->Serialize((UTF8CHAR*)UTF8String.Get(), UTF8String.Length() * sizeof(UTF8CHAR));
-	QaDataFile->Flush();
-
-	if (QaHeaderFile.IsValid())
-	{
-		FTCHARToUTF8 UTF8String(Buffer, Result);
-		QaHeaderFile->Serialize((UTF8CHAR*)UTF8String.Get(), UTF8String.Length() * sizeof(UTF8CHAR));
-		QaHeaderFile->Flush();
-	}
-
-	if (Buffer)
+	void FChrono::PrintClockElapse(EVerboseLevel Level, const TCHAR* Indent, const TCHAR* Process, FDuration Duration, ETimeUnit Unit)
 	{
-		FMemory::Free(Buffer);
-	}
-}
-#endif
-
-
-void FChrono::PrintClockElapse(EVerboseLevel Level, const TCHAR* Indent, const TCHAR* Process, FDuration Duration, ETimeUnit Unit)
-{
 #ifdef CADKERNEL_DEV
-	if (FSystem::Get().GetVerboseLevel() < Level)
-	{
-		return;
-	}
+		if (FSystem::Get().GetVerboseLevel() < Level)
+		{
+			return;
+		}
 
-	TCHAR* SUnit = TEXT("seconds");
-	long long Time = 0;
-	switch (Unit)
-	{
-	case ETimeUnit::NanoSeconds:
-		SUnit = TEXT("nanoseconds");
-		Time = FChrono::ConvertInto<std::chrono::nanoseconds>(Duration);
-		break;
-	case ETimeUnit::MicroSeconds:
-		SUnit = TEXT("microseconds");
-		Time = FChrono::ConvertInto<std::chrono::microseconds>(Duration);
-		break;
-	case ETimeUnit::MilliSeconds:
-		SUnit = TEXT("milliseconds");
-		Time = FChrono::ConvertInto<std::chrono::milliseconds>(Duration);
-		break;
-	default:
-		Time = FChrono::ConvertInto<std::chrono::seconds>(Duration);
-	}
-	FMessage::Printf(Level, TEXT("%sSpeed test (%s) -----> %lld %s\n"), Indent, Process, Time, SUnit);
+		TCHAR* SUnit = TEXT("seconds");
+		long long Time = 0;
+		switch (Unit)
+		{
+		case ETimeUnit::NanoSeconds:
+			SUnit = TEXT("nanoseconds");
+			Time = FChrono::ConvertInto<std::chrono::nanoseconds>(Duration);
+			break;
+		case ETimeUnit::MicroSeconds:
+			SUnit = TEXT("microseconds");
+			Time = FChrono::ConvertInto<std::chrono::microseconds>(Duration);
+			break;
+		case ETimeUnit::MilliSeconds:
+			SUnit = TEXT("milliseconds");
+			Time = FChrono::ConvertInto<std::chrono::milliseconds>(Duration);
+			break;
+		default:
+			Time = FChrono::ConvertInto<std::chrono::seconds>(Duration);
+		}
+		FMessage::Printf(Level, TEXT("%sSpeed test (%s) -----> %lld %s\n"), Indent, Process, Time, SUnit);
 #endif
-}
+	}
 }
\ No newline at end of file
diff --git a/Engine/Source/Runtime/Datasmith/CADKernel/Public/CADKernel/Core/System.h b/Engine/Source/Runtime/Datasmith/CADKernel/Public/CADKernel/Core/System.h
--- a/Engine/Source/Runtime/Datasmith/CADKernel/Public/CADKernel/Core/System.h
+++ b/Engine/Source/Runtime/Datasmith/CADKernel/Public/CADKernel/Core/System.h
@@ -39,8 +39,8 @@ namespace CADKernel
 		TSharedPtr<FArchive> SpyFile;
 
 #if defined(CADKERNEL_DEV) || defined(CADKERNEL_STDA)
-		TSharedPtr<FArchive> QaDataFile;
-		TSharedPtr<FArchive> QaHeaderFile;
+		TSharedPtr<FArchive> ReportFile;
+		TSharedPtr<FArchive> ReportHeaderFile;
 #endif
 
 		EVerboseLevel VerboseLevel;
@@ -126,15 +126,15 @@ namespace CADKernel
 
 
 #if defined(CADKERNEL_DEV) || defined(CADKERNEL_STDA)
-		void DefineQaDataFile(const FString& InLogFilePath);
-		TSharedPtr<FArchive> GetQaDataFile() const
+		void DefineReportFile(const FString& InLogFilePath);
+		TSharedPtr<FArchive> GetReportFile() const
 		{
-			return QaDataFile;
+			return ReportFile;
 		}
 
-		TSharedPtr<FArchive> GetQaHeaderFile() const
+		TSharedPtr<FArchive> GetReportHeaderFile() const
 		{
-			return QaHeaderFile;
+			return ReportHeaderFile;
 		}
 #endif
 
diff --git a/Engine/Source/Runtime/Datasmith/CADKernel/Public/CADKernel/Geo/GeoEnum.h b/Engine/Source/Runtime/Datasmith/CADKernel/Public/CADKernel/Geo/GeoEnum.h
--- a/Engine/Source/Runtime/Datasmith/CADKernel/Public/CADKernel/Geo/GeoEnum.h
+++ b/Engine/Source/Runtime/Datasmith/CADKernel/Public/CADKernel/Geo/GeoEnum.h
@@ -5,123 +5,129 @@
 
 namespace CADKernel
 {
-	enum class ECurve : uint8
-	{
-		Bezier = 0,
-		BoundedCurve,
-		Composite,
-		Ellipse,
-		Hyperbola,
-		Nurbs,
-		Offset,
-		Parabola,
-		Polyline3D,
-		Polyline2D,
-		Restriction,
-		Segment,
-		Surfacic,
-		SurfacicPolyline,
-		None,
-	};
+
+enum class ECurve : uint8
+{
+	Bezier = 0,
+	BoundedCurve,
+	Composite,
+	Ellipse,
+	Hyperbola,
+	Nurbs,
+	Offset,
+	Parabola,
+	Polyline3D,
+	Polyline2D,
+	Restriction,
+	Segment,
+	Surfacic,
+	SurfacicPolyline,
+	Spline,
+	None,
+};
 
 #ifdef CADKERNEL_DEV
-	extern const TCHAR* CurvesTypesNames[];
+extern const TCHAR* CurvesTypesNames[];
 #endif
 
-	enum class ESurface : uint8
-	{
-		Bezier = 0,
-		Composite,
-		Cone,
-		Coons,
-		Cylinder,
-		Nurbs,
-		Offset,
-		Plane,
-		Revolution,
-		Ruled,
-		Sphere,
-		TabulatedCylinder,
-		Torus,
-		None,
-	};
+enum class ESurface : uint8
+{
+	Bezier = 0,
+	Blend01,
+	Blend02,
+	Blend03,
+	Composite,
+	Cone,
+	Coons,
+	Cylinder,
+	Nurbs,
+	Offset,
+	Plane,
+	Revolution,
+	Ruled,
+	Sphere,
+	TabulatedCylinder,
+	Torus,
+	None,
+};
 
 #ifdef CADKERNEL_DEV
-	extern const TCHAR* SurfacesTypesNames[];
+extern const TCHAR* SurfacesTypesNames[];
 #endif
 
-	enum class ERupture : uint8
-	{
-		Continuity = 0, // Positional rupture i.e. points of the curve that are not G0 
-		Tangency,       // Tangency rupture i.e. points of the curve that are not G1 (or G0) 
-		Curvature,      // Curvature rupture i.e. points of the curve that are not G2 (or G1, G0)
-	};
+enum class ERupture : uint8
+{
+	Continuity = 0, // Positional rupture i.e. points of the curve that are not G0 
+	Tangency,       // Tangency rupture i.e. points of the curve that are not G1 (or G0) 
+	Curvature,      // Curvature rupture i.e. points of the curve that are not G2 (or G1, G0)
+};
 
-	enum EIso : uint8
-	{
-		IsoU = 0,
-		IsoV = 1,
-		UndefinedIso
-	};
+enum EIso : uint8
+{
+	IsoU = 0,
+	IsoV = 1,
+	UndefinedIso
+};
 
 #ifdef CADKERNEL_DEV
-	extern const TCHAR* IsoNames[];
+extern const TCHAR* IsoNames[];
 #endif
 
-	constexpr EIso Other(const EIso Iso)
-	{
-		return Iso == EIso::IsoU ? EIso::IsoV : EIso::IsoU;
-	}
+constexpr EIso Other(const EIso Iso)
+{
+	return Iso == EIso::IsoU ? EIso::IsoV : EIso::IsoU;
+}
 
-	enum EOrientation : uint8
-	{
-		Back = 0,
-		Front = 1
-	};
+enum EOrientation : uint8
+{
+	Back = 0,
+	Front = 1
+};
 
 #ifdef CADKERNEL_DEV
-	extern const TCHAR* OrientationNames[];
+extern const TCHAR* OrientationNames[];
 #endif
 
-	enum ELimit : uint8
-	{
-		Start = 0,
-		End = 1
-	};
+enum ELimit : uint8
+{
+	Start = 0,
+	End = 1
+};
 
-	inline EOrientation SameOrientation(bool bIsSameOrientation)
+inline EOrientation SameOrientation(bool bIsSameOrientation)
+{
+	if (bIsSameOrientation)
 	{
-		if (bIsSameOrientation)
-		{
-			return EOrientation::Front;
-		}
-		else
-		{
-			return EOrientation::Back;
-		}
+		return EOrientation::Front;
 	}
+	else
+	{
+		return EOrientation::Back;
+	}
+}
 
-	inline EOrientation GetReverseOrientation(EOrientation Orientation)
+inline EOrientation GetReverseOrientation(EOrientation Orientation)
+{
+	if (Orientation == EOrientation::Front)
 	{
-		if (Orientation == EOrientation::Front)
-		{
-			return EOrientation::Back;
-		}
-		else
-		{
-			return EOrientation::Front;
-		}
+		return EOrientation::Back;
 	}
+	else
+	{
+		return EOrientation::Front;
+	}
+}
 
-	inline void SwapOrientation(EOrientation& Orientation)
+inline void SwapOrientation(EOrientation& Orientation)
+{
+	if (Orientation == EOrientation::Front)
 	{
-		if (Orientation == EOrientation::Front)
-		{
-			Orientation = EOrientation::Back;
-		}
-		else
-		{
-			Orientation = EOrientation::Front;
-		}
+		Orientation = EOrientation::Back;
 	}
-} 
+	else
+	{
+		Orientation = EOrientation::Front;
+	}
+}
+
+} // ns
diff --git a/Engine/Source/Runtime/Datasmith/CADKernel/Public/CADKernel/Geo/Sampling/PolylineTools.h b/Engine/Source/Runtime/Datasmith/CADKernel/Public/CADKernel/Geo/Sampling/PolylineTools.h
--- a/Engine/Source/Runtime/Datasmith/CADKernel/Public/CADKernel/Geo/Sampling/PolylineTools.h
+++ b/Engine/Source/Runtime/Datasmith/CADKernel/Public/CADKernel/Geo/Sampling/PolylineTools.h
@@ -15,609 +15,625 @@
 
 namespace CADKernel
 {
-	namespace PolylineTools
+
+namespace PolylineTools
+{
+
+inline bool IsDichotomyToBePreferred(int32 InPolylineSize, int32 ResultSize)
+{
+	double MeanLinearIteration = InPolylineSize / ResultSize;
+	double MaxDichotomyIteration = FMath::Log2((double)InPolylineSize);
+
+	if (MeanLinearIteration > MaxDichotomyIteration)
 	{
-		inline bool IsDichotomyToBePreferred(int32 InPolylineSize, int32 ResultSize)
-		{
-			double MeanLinearIteration = InPolylineSize / ResultSize;
-			double MaxDichotomyIteration = FMath::Log2((double) InPolylineSize);
+		return true;
+	}
+	return false;
+}
 
-			if (MeanLinearIteration > MaxDichotomyIteration)
-			{
-				return true;
-			}
-			return false;
-		}
+template<typename PointType>
+inline PointType LinearInterpolation(const TArray<PointType>& Array, const int32 Index, const double Coordinate)
+{
+	ensureCADKernel(Index + 1 < Array.Num());
+	return PointOnSegment(Array[Index], Array[Index + 1], Coordinate);
+}
 
-		template<typename PointType>
-		inline PointType LinearInterpolation(const TArray<PointType>& Array, const int32 Index, const double Coordinate)
-		{
-			ensureCADKernel(Index + 1 < Array.Num());
-			return PointOnSegment(Array[Index], Array[Index + 1], Coordinate);
-		}
+inline double SectionCoordinate(const TArray<double>& Array, const int32 Index, const double Coordinate)
+{
+	ensureCADKernel(Index + 1 < Array.Num());
+	const double DeltaU = Array[Index + 1] - Array[Index];
+	if (FMath::IsNearlyZero(DeltaU))
+	{
+		return 0;
+	}
+	return (Coordinate - Array[Index]) / DeltaU;
+}
+
+template<typename PointType>
+double ComputeLength(TArray<PointType>& Polyline)
+{
+	double Length = 0;
+	for (int32 Index = 1; Index < Polyline.Num(); ++Index)
+	{
+		Length += Polyline[Index - 1].Distance(Polyline[Index]);
+	}
+	return Length;
+}
+
+/**
+ * Progressively deforms a polyline (or a control polygon) so that its end is in the desired position
+ */
+template<typename PointType>
+void ExtendTo(TArray<PointType>& Polyline, const PointType& DesiredEnd)
+{
+	double DistanceStart = Polyline[0].SquareDistance(DesiredEnd);
+	double DistanceEnd = Polyline.Last().SquareDistance(DesiredEnd);
+
+	if (DistanceStart < DistanceEnd)
+	{
+		Algo::Reverse(Polyline);
+	}
+
+	double PolylineLength = ComputeLength(Polyline);
+
+	PointType Delta = DesiredEnd - Polyline.Last();
+	Delta /= PolylineLength;
+
+	PolylineLength = 0;
+	PolylineLength = Polyline[1].Distance(Polyline[0]);
+	for (int32 Index = 1; Index < Polyline.Num() - 1; ++Index)
+	{
+		double LengthNextSegment = Polyline[Index].Distance(Polyline[Index + 1]);
+		Polyline[Index] += Delta * PolylineLength;
+		PolylineLength += LengthNextSegment;
+	}
+	Polyline.Last() = DesiredEnd;
+
+	if (DistanceStart < DistanceEnd)
+	{
+		Algo::Reverse(Polyline);
+	}
+}
 
-		inline double SectionCoordinate(const TArray<double>& Array, const int32 Index, const double Coordinate)
+template<class PointType>
+PointType ComputePoint(const TArray<double>& PolylineCoordinates, const TArray<PointType>& PolylinePoints, const int32 Index, const double PointCoordinate)
+{
+	double Delta = PolylineCoordinates[Index + 1] - PolylineCoordinates[Index];
+	if (FMath::IsNearlyZero(Delta, (double)SMALL_NUMBER))
+	{
+		return PolylinePoints[Index];
+	}
+
+	return PolylinePoints[Index] + (PolylinePoints[Index + 1] - PolylinePoints[Index]) * (PointCoordinate - PolylineCoordinates[Index]) / Delta;
+};
+
+} // ns PolylineTools
+
+template<class PointType>
+class TPolylineApproximator
+{
+protected:
+	const TArray<double>& PolylineCoordinates;
+	const TArray<PointType>& PolylinePoints;
+
+public:
+	TPolylineApproximator(const TArray<double>& InPolylineCoordinates, const TArray<PointType>& InPolylinePoints)
+		: PolylineCoordinates(InPolylineCoordinates)
+		, PolylinePoints(InPolylinePoints)
+	{
+	}
+
+protected:
+
+	double ComputeCurvilinearCoordinatesOfPolyline(const FLinearBoundary& InBoundary, TArray<double>& OutCurvilinearCoordinates, int32 BoundaryIndices[2]) const
+	{
+		GetStartEndIndex(InBoundary, BoundaryIndices);
+
+		OutCurvilinearCoordinates.Reserve(BoundaryIndices[1] - BoundaryIndices[0] + 2);
+
+		double LastSegmentLength;
+		double Length = 0;
+		ensureCADKernel(BoundaryIndices[1] > BoundaryIndices[0]);
+
+		OutCurvilinearCoordinates.Add(0);
+		if (BoundaryIndices[1] > BoundaryIndices[0] + 1)
 		{
-			ensureCADKernel(Index + 1 < Array.Num());
-			const double DeltaU = Array[Index + 1] - Array[Index];
-			if (FMath::IsNearlyZero(DeltaU))
+			PointType StartPoint = ComputePoint(BoundaryIndices[0], InBoundary.Min);
+			PointType EndPoint = ComputePoint(BoundaryIndices[1], InBoundary.Max);
+			Length = StartPoint.Distance(PolylinePoints[BoundaryIndices[0] + 1]);
+			LastSegmentLength = EndPoint.Distance(PolylinePoints[BoundaryIndices[1]]);
+
+			OutCurvilinearCoordinates.Add(Length);
+			for (int32 Index = BoundaryIndices[0] + 1; Index < BoundaryIndices[1]; ++Index)
 			{
-				return 0;
+				Length += PolylinePoints[Index].Distance(PolylinePoints[Index + 1]);
+				OutCurvilinearCoordinates.Add(Length);
 			}
-			return (Coordinate - Array[Index]) / DeltaU;
+			Length += LastSegmentLength;
+			OutCurvilinearCoordinates.Add(Length);
 		}
-
-		template<typename PointType>
-		double ComputeLength(TArray<PointType>& Polyline)
+		else
 		{
-			double Length = 0;
-			for (int32 Index = 1; Index < Polyline.Num(); ++Index)
-			{
-				Length += Polyline[Index - 1].Distance(Polyline[Index]);
-			}
-			return Length;
+			PointType StartPoint = ComputePoint(BoundaryIndices[0], InBoundary.Min);
+			PointType EndPoint = ComputePoint(BoundaryIndices[1], InBoundary.Max);
+			Length = StartPoint.Distance(EndPoint);
+			OutCurvilinearCoordinates.Add(Length);
 		}
+		return Length;
+	}
 
-		/**
-		 * Progressively deforms a polyline (or a control polygon) so that its end is in the desired position
-		 */
-		template<typename PointType>
-		void ExtendTo(TArray<PointType>& Polyline, const PointType& DesiredEnd)
+	PointType ComputePoint(const int32 Index, const double PointCoordinate) const
+	{
+		double Delta = PolylineCoordinates[Index + 1] - PolylineCoordinates[Index];
+		if (FMath::IsNearlyZero(Delta, (double)SMALL_NUMBER))
 		{
-			double DistanceStart = Polyline[0].SquareDistance(DesiredEnd);
-			double DistanceEnd = Polyline.Last().SquareDistance(DesiredEnd);
+			return PolylinePoints[Index];
+		}
 
-			if (DistanceStart < DistanceEnd)
-			{
-				Algo::Reverse(Polyline);
-			}
+		return PolylinePoints[Index] + (PolylinePoints[Index + 1] - PolylinePoints[Index]) * (PointCoordinate - PolylineCoordinates[Index]) / Delta;
+	};
 
-			double PolylineLength = ComputeLength(Polyline);
+	/**
+	 * Project a Set of points on a restricted polyline (StartIndex & EndIndex define the polyline boundary)
+	 * the points are projected on all segments of the polyline, the closest are selected
+	 */
+	double ProjectPointToPolyline(int32 BoundaryIndices[2], const PointType& InPointToProject, PointType& OutProjectedPoint) const
+	{
+		double MinDistance = HUGE_VAL;
+		double UForMinDistance = 0;
 
-			PointType Delta = DesiredEnd - Polyline.Last();
-			Delta /= PolylineLength;
+		double ParamU = 0.;
+		int32 SegmentIndex = 0;
 
-			PolylineLength = 0;
-			PolylineLength = Polyline[1].Distance(Polyline[0]);
-			for (int32 Index = 1; Index < Polyline.Num() - 1; ++Index)
+		for (int32 Index = BoundaryIndices[0]; Index <= BoundaryIndices[1]; ++Index)
+		{
+			FPoint ProjectPoint = ProjectPointOnSegment(InPointToProject, PolylinePoints[Index], PolylinePoints[Index + 1], ParamU, true);
+			double SquareDistance = FMath::Square(ProjectPoint[0] - InPointToProject[0]);
+			if (SquareDistance > MinDistance)
 			{
-				double LengthNextSegment = Polyline[Index].Distance(Polyline[Index + 1]);
-				Polyline[Index] += Delta * PolylineLength;
-				PolylineLength += LengthNextSegment;
+				continue;
 			}
-			Polyline.Last() = DesiredEnd;
-
-			if (DistanceStart < DistanceEnd)
+			SquareDistance += FMath::Square(ProjectPoint[1] - InPointToProject[1]);
+			if (SquareDistance > MinDistance)
 			{
-				Algo::Reverse(Polyline);
+				continue;
 			}
-		}
-
-		template<class PointType>
-		PointType ComputePoint(const TArray<double>& PolylineCoordinates, const TArray<PointType>& PolylinePoints, const int32 Index, const double PointCoordinate)
-		{
-			double Delta = PolylineCoordinates[Index + 1] - PolylineCoordinates[Index];
-			if (FMath::IsNearlyZero(Delta, (double)KINDA_SMALL_NUMBER))
+			SquareDistance += FMath::Square(ProjectPoint[2] - InPointToProject[2]);
+			if (SquareDistance > MinDistance)
 			{
-				return PolylinePoints[Index];
+				continue;
 			}
+			MinDistance = SquareDistance;
+			UForMinDistance = ParamU;
+			SegmentIndex = Index;
+			OutProjectedPoint = ProjectPoint;
+		}
 
-			return PolylinePoints[Index] + (PolylinePoints[Index + 1] - PolylinePoints[Index]) * (PointCoordinate - PolylineCoordinates[Index]) / Delta;
-		};
+		return PolylineTools::LinearInterpolation(PolylineCoordinates, SegmentIndex, UForMinDistance);
+	}
 
+public:
 
+	void GetStartEndIndex(const FLinearBoundary& InBoundary, int32 BoundaryIndices[2]) const
+	{
+		FDichotomyFinder Finder(PolylineCoordinates);
+		BoundaryIndices[0] = Finder.Find(InBoundary.Min);
+		BoundaryIndices[1] = Finder.Find(InBoundary.Max);
 	}
 
-	template<class PointType>
-	class TPolylineApproximator
+	/**
+	 * Evaluate the point of the polyline at the input InCoordinate
+	 * If the input coordinate is outside the boundary of the polyline, the coordinate of the nearest boundary is used.
+	 */
+	PointType ApproximatePoint(const double InCoordinate) const
 	{
-	protected:
-		const TArray<double>& PolylineCoordinates;
-		const TArray<PointType>& PolylinePoints;
-
-	public:
-		TPolylineApproximator(const TArray<double>& InPolylineCoordinates, const TArray<PointType>& InPolylinePoints)
-			: PolylineCoordinates(InPolylineCoordinates)
-			, PolylinePoints(InPolylinePoints)
-		{
-		}
+		FDichotomyFinder Finder(PolylineCoordinates);
+		int32 Index = Finder.Find(InCoordinate);
+		return ComputePoint(Index, InCoordinate);
+	}
 
-	protected:
+	/**
+	 * Evaluate the point of the polyline at the input Coordinate
+	 * If the input coordinate is outside the boundary of the polyline, the coordinate of the nearest boundary is used.
+	 */
+	template<class CurvePointType>
+	void ApproximatePoint(double InCoordinate, CurvePointType& OutPoint, int32 InDerivativeOrder) const
+	{
+		FDichotomyFinder Finder(PolylineCoordinates);
+		int32 Index = Finder.Find(InCoordinate);
 
-		double ComputeCurvilinearCoordinatesOfPolyline(const FLinearBoundary& InBoundary, TArray<double>& OutCurvilinearCoordinates, int32 BoundaryIndices[2]) const
-		{
-			GetStartEndIndex(InBoundary, BoundaryIndices);
+		OutPoint.DerivativeOrder = InDerivativeOrder;
 
-			OutCurvilinearCoordinates.Reserve(BoundaryIndices[1] - BoundaryIndices[0] + 2);
+		double DeltaU = PolylineCoordinates[Index + 1] - PolylineCoordinates[Index];
+		if (FMath::IsNearlyZero(DeltaU))
+		{
+			OutPoint.Point = PolylinePoints[Index];
+			OutPoint.Gradient = FPoint::ZeroPoint;
+			OutPoint.Laplacian = FPoint::ZeroPoint;
+			return;
+		}
 
-			double LastSegmentLength;
-			double Length = 0;
-			ensureCADKernel(BoundaryIndices[1] > BoundaryIndices[0]);
+		double SectionCoordinate = (InCoordinate - PolylineCoordinates[Index]) / DeltaU;
 
-			OutCurvilinearCoordinates.Add(0);
-			if (BoundaryIndices[1] > BoundaryIndices[0] + 1)
-			{
-				PointType StartPoint = ComputePoint(BoundaryIndices[0], InBoundary.Min);
-				PointType EndPoint = ComputePoint(BoundaryIndices[1], InBoundary.Max);
-				Length = StartPoint.Distance(PolylinePoints[BoundaryIndices[0] + 1]);
-				LastSegmentLength = EndPoint.Distance(PolylinePoints[BoundaryIndices[1]]);
+		FPoint Tangent = PolylinePoints[Index + 1] - PolylinePoints[Index];
+		OutPoint.Point = PolylinePoints[Index] + Tangent * SectionCoordinate;
 
-				OutCurvilinearCoordinates.Add(Length);
-				for (int32 Index = BoundaryIndices[0] + 1; Index < BoundaryIndices[1]; ++Index)
-				{
-					Length += PolylinePoints[Index].Distance(PolylinePoints[Index + 1]);
-					OutCurvilinearCoordinates.Add(Length);
-				}
-				Length += LastSegmentLength;
-				OutCurvilinearCoordinates.Add(Length);
-			}
-			else
-			{
-				PointType StartPoint = ComputePoint(BoundaryIndices[0], InBoundary.Min);
-				PointType EndPoint = ComputePoint(BoundaryIndices[1], InBoundary.Max);
-				Length = StartPoint.Distance(EndPoint);
-				OutCurvilinearCoordinates.Add(Length);
-			}
-			return Length;
+		if (InDerivativeOrder > 0)
+		{
+			OutPoint.Gradient = Tangent;
+			OutPoint.Laplacian = FPoint::ZeroPoint;
 		}
+	}
 
-		PointType ComputePoint(const int32 Index, const double PointCoordinate) const
+	template<class CurvePointType>
+	void ApproximatePoints(const TArray<double>& InCoordinates, TArray<CurvePointType>& OutPoints, int32 InDerivativeOrder = 0) const
+	{
+		if (!InCoordinates.Num())
 		{
-			double Delta = PolylineCoordinates[Index + 1] - PolylineCoordinates[Index];
-			if (FMath::IsNearlyZero(Delta, (double)KINDA_SMALL_NUMBER))
-			{
-				return PolylinePoints[Index];
-			}
-
-			return PolylinePoints[Index] + (PolylinePoints[Index + 1] - PolylinePoints[Index]) * (PointCoordinate - PolylineCoordinates[Index]) / Delta;
-		};
+			ensureCADKernel(false);
+			return;
+		}
 
-		/**
-		 * Project a Set of points on a restricted polyline (StartIndex & EndIndex define the polyline boundary)
-		 * the points are projected on all segments of the polyline, the closest are selected
-		 */
-		double ProjectPointToPolyline(int32 BoundaryIndices[2], const PointType& InPointToProject, PointType& OutProjectedPoint) const
+		TFunction<void(FIndexOfCoordinateFinder&)> ComputePoints = [&](FIndexOfCoordinateFinder& Finder)
 		{
-			double MinDistance = HUGE_VAL;
-			double UForMinDistance = 0;
+			for (int32 IPoint = 0; IPoint < InCoordinates.Num(); ++IPoint)
+			{
+				int32 Index = Finder.Find(InCoordinates[IPoint]);
 
-			double ParamU = 0.;
-			int32 SegmentIndex = 0;
+				OutPoints[IPoint].DerivativeOrder = InDerivativeOrder;
 
-			for (int32 Index = BoundaryIndices[0]; Index <= BoundaryIndices[1]; ++Index)
-			{
-				FPoint ProjectPoint = ProjectPointOnSegment(InPointToProject, PolylinePoints[Index], PolylinePoints[Index + 1], ParamU, true);
-				double SquareDistance = FMath::Square(ProjectPoint[0] - InPointToProject[0]);
-				if (SquareDistance > MinDistance)
-				{
-					continue;
-				}
-				SquareDistance += FMath::Square(ProjectPoint[1] - InPointToProject[1]);
-				if (SquareDistance > MinDistance)
+				double DeltaU = PolylineCoordinates[Index + 1] - PolylineCoordinates[Index];
+				if (FMath::IsNearlyZero(DeltaU))
 				{
-					continue;
+					OutPoints[IPoint].Point = PolylinePoints[Index];
+					OutPoints[IPoint].Gradient = FPoint::ZeroPoint;
+					OutPoints[IPoint].Laplacian = FPoint::ZeroPoint;
+					return;
 				}
-				SquareDistance += FMath::Square(ProjectPoint[2] - InPointToProject[2]);
-				if (SquareDistance > MinDistance)
+
+				double SectionCoordinate = (InCoordinates[IPoint] - PolylineCoordinates[Index]) / DeltaU;
+
+				FPoint Tangent = PolylinePoints[Index + 1] - PolylinePoints[Index];
+				OutPoints[IPoint].Point = PolylinePoints[Index] + Tangent * SectionCoordinate;
+
+				if (InDerivativeOrder > 0)
 				{
-					continue;
+					OutPoints[IPoint].Gradient = Tangent;
+					OutPoints[IPoint].Laplacian = FPoint::ZeroPoint;
 				}
-				MinDistance = SquareDistance;
-				UForMinDistance = ParamU;
-				SegmentIndex = Index;
-				OutProjectedPoint = ProjectPoint;
 			}
+		};
 
-			return PolylineTools::LinearInterpolation(PolylineCoordinates, SegmentIndex, UForMinDistance);
-		}
+		FDichotomyFinder DichotomyFinder(PolylineCoordinates);
 
-	public:
+		int32 StartIndex = DichotomyFinder.Find(InCoordinates[0]);
+		int32 EndIndex = DichotomyFinder.Find(InCoordinates.Last());
+		bool bUseDichotomy = PolylineTools::IsDichotomyToBePreferred(EndIndex - StartIndex, InCoordinates.Num());
 
-		void GetStartEndIndex(const FLinearBoundary& InBoundary, int32 BoundaryIndices[2]) const
+		OutPoints.Empty(InCoordinates.Num());
+		if (bUseDichotomy)
 		{
-			FDichotomyFinder Finder(PolylineCoordinates);
-			BoundaryIndices[0] = Finder.Find(InBoundary.Min);
-			BoundaryIndices[1] = Finder.Find(InBoundary.Max);
+			DichotomyFinder.StartLower = StartIndex;
+			DichotomyFinder.StartUpper = EndIndex;
+			ComputePoints(DichotomyFinder);
 		}
-
-		/**
-		 * Evaluate the point of the polyline at the input InCoordinate
-		 * If the input coordinate is outside the boundary of the polyline, the coordinate of the nearest boundary is used.
-		 */
-		PointType ApproximatePoint(const double InCoordinate) const
+		else
 		{
-			FDichotomyFinder Finder(PolylineCoordinates);
-			int32 Index = Finder.Find(InCoordinate);
-			return ComputePoint(Index, InCoordinate);
+			FLinearFinder LinearFinder(PolylineCoordinates, StartIndex);
+			ComputePoints(LinearFinder);
 		}
+	}
 
-		/**
-		 * Evaluate the point of the polyline at the input Coordinate
-		 * If the input coordinate is outside the boundary of the polyline, the coordinate of the nearest boundary is used.
-		 */
-		template<class CurvePointType>
-		void ApproximatePoint(double InCoordinate, CurvePointType& OutPoint, int32 InDerivativeOrder) const
-		{
-			FDichotomyFinder Finder(PolylineCoordinates);
-			int32 Index = Finder.Find(InCoordinate);
 
-			OutPoint.DerivativeOrder = InDerivativeOrder;
+	/**
+	 * Evaluate the points of the polyline associated to the increasing array of input Coordinates
+	 * If the input coordinate is outside the boundary of the polyline, the coordinate of the nearest boundary is used.
+	 */
+	void ApproximatePoints(const TArray<double>& InCoordinates, TArray<PointType>& OutPoints) const
+	{
+		if (!InCoordinates.Num())
+		{
+			return;
+		}
 
-			double DeltaU = PolylineCoordinates[Index + 1] - PolylineCoordinates[Index];
-			if (FMath::IsNearlyZero(DeltaU))
+		TFunction<void(FIndexOfCoordinateFinder&)> ComputePoints = [&](FIndexOfCoordinateFinder& Finder)
+		{
+			for (double Coordinate : InCoordinates)
 			{
-				OutPoint.Point = PolylinePoints[Index];
-				OutPoint.Gradient = FPoint::ZeroPoint;
-				OutPoint.Laplacian = FPoint::ZeroPoint;
-				return;
+				int32 Index = Finder.Find(Coordinate);
+				OutPoints.Emplace(ComputePoint(Index, Coordinate));
 			}
+		};
 
-			double SectionCoordinate = (InCoordinate - PolylineCoordinates[Index]) / DeltaU;
+		FDichotomyFinder DichotomyFinder(PolylineCoordinates);
 
-			FPoint Tangent = PolylinePoints[Index + 1] - PolylinePoints[Index];
-			OutPoint.Point = PolylinePoints[Index] + Tangent * SectionCoordinate;
+		int32 StartIndex = DichotomyFinder.Find(InCoordinates[0]);
+		int32 EndIndex = DichotomyFinder.Find(InCoordinates.Last());
+		bool bUseDichotomy = PolylineTools::IsDichotomyToBePreferred(EndIndex - StartIndex, InCoordinates.Num());
 
-			if (InDerivativeOrder > 0)
-			{
-				OutPoint.Gradient = Tangent;
-				OutPoint.Laplacian = FPoint::ZeroPoint;
-			}
+		OutPoints.Empty(InCoordinates.Num());
+		if (bUseDichotomy)
+		{
+			DichotomyFinder.StartLower = StartIndex;
+			DichotomyFinder.StartUpper = EndIndex;
+			ComputePoints(DichotomyFinder);
 		}
-
-		template<class CurvePointType>
-		void ApproximatePoints(const TArray<double>& InCoordinates, TArray<CurvePointType>& OutPoints, int32 InDerivativeOrder = 0) const
+		else
 		{
-			if (!InCoordinates.Num())
-			{
-				ensureCADKernel(false);
-				return;
-			}
-
-			TFunction<void(FIndexOfCoordinateFinder&)> ComputePoints = [&](FIndexOfCoordinateFinder& Finder)
-			{
-				for (int32 IPoint = 0; IPoint < InCoordinates.Num(); ++IPoint)
-				{
-					int32 Index = Finder.Find(InCoordinates[IPoint]);
-
-					OutPoints[IPoint].DerivativeOrder = InDerivativeOrder;
+			FLinearFinder LinearFinder(PolylineCoordinates, StartIndex);
+			ComputePoints(LinearFinder);
+		}
+	}
 
-					double DeltaU = PolylineCoordinates[Index + 1] - PolylineCoordinates[Index];
-					if (FMath::IsNearlyZero(DeltaU))
-					{
-						OutPoints[IPoint].Point = PolylinePoints[Index];
-						OutPoints[IPoint].Gradient = FPoint::ZeroPoint;
-						OutPoints[IPoint].Laplacian = FPoint::ZeroPoint;
-						return;
-					}
+	void SamplePolyline(const FLinearBoundary& InBoundary, const double DesiredSegmentLength, TArray<double>& OutCoordinates) const
+	{
+		int32 BoundaryIndices[2];
 
-					double SectionCoordinate = (InCoordinates[IPoint] - PolylineCoordinates[Index]) / DeltaU;
+		TArray<double> CurvilinearCoordinates;
+		double CurveLength = ComputeCurvilinearCoordinatesOfPolyline(InBoundary, CurvilinearCoordinates, BoundaryIndices);
 
-					FPoint Tangent = PolylinePoints[Index + 1] - PolylinePoints[Index];
-					OutPoints[IPoint].Point = PolylinePoints[Index] + Tangent * SectionCoordinate;
+		int32 SegmentNum = (int32)FMath::Max(CurveLength / DesiredSegmentLength + 0.5, 1.0);
 
-					if (InDerivativeOrder > 0)
-					{
-						OutPoints[IPoint].Gradient = Tangent;
-						OutPoints[IPoint].Laplacian = FPoint::ZeroPoint;
-					}
-				}
-			};
+		double SectionLength = CurveLength / (double)(SegmentNum);
 
-			FDichotomyFinder DichotomyFinder(PolylineCoordinates);
+		OutCoordinates.Empty();
+		OutCoordinates.Reserve(SegmentNum + 1);
+		OutCoordinates.Add(InBoundary.Min);
 
-			int32 StartIndex = DichotomyFinder.Find(InCoordinates[0]);
-			int32 EndIndex = DichotomyFinder.Find(InCoordinates.Last());
-			bool bUseDichotomy = PolylineTools::IsDichotomyToBePreferred(EndIndex - StartIndex, InCoordinates.Num());
+		TFunction<double(const int32, const int32, const double, const double)> ComputeSamplePointCoordinate = [&](const int32 IndexCurvilinear, const int32 IndexCoordinate, const double Length, const double Coordinate)
+		{
+			return Coordinate + (PolylineCoordinates[IndexCoordinate] - Coordinate) * (Length - CurvilinearCoordinates[IndexCurvilinear - 1u]) / (CurvilinearCoordinates[IndexCurvilinear] - CurvilinearCoordinates[IndexCurvilinear - 1u]);
+		};
 
-			OutPoints.Empty(InCoordinates.Num());
-			if (bUseDichotomy)
-			{
-				DichotomyFinder.StartLower = StartIndex;
-				DichotomyFinder.StartUpper = EndIndex;
-				ComputePoints(DichotomyFinder);
-			}
-			else
+		double CurvilinearLength = SectionLength;
+		double LastCoordinate = InBoundary.Min;
+		for (int32 IndexCurvilinear = 1, IndexCoordinate = BoundaryIndices[0] + 1; IndexCurvilinear < CurvilinearCoordinates.Num(); ++IndexCurvilinear, ++IndexCoordinate)
+		{
+			while (CurvilinearLength < CurvilinearCoordinates[IndexCurvilinear] + SMALL_NUMBER)
 			{
-				FLinearFinder LinearFinder(PolylineCoordinates, StartIndex);
-				ComputePoints(LinearFinder);
+				double Coordinate = ComputeSamplePointCoordinate(IndexCurvilinear, IndexCoordinate, CurvilinearLength, LastCoordinate);
+				OutCoordinates.Add(Coordinate);
+				CurvilinearLength += SectionLength;
+				if (CurvilinearLength + SMALL_NUMBER > CurveLength)
+				{
+					OutCoordinates.Add(InBoundary.Max);
+					break;
+				}
 			}
+			LastCoordinate = PolylineCoordinates[IndexCoordinate];
 		}
+	}
 
+	/**
+	 * Project a Set of points on a restricted polyline (StartIndex & EndIndex define the polyline boundary)
+	 * the points are projected on all segments of the polyline, the closest are selected
+	 */
+	double ProjectPointToPolyline(const FLinearBoundary& InBoundary, const PointType& PointOnEdge, PointType& OutProjectedPoint) const
+	{
+		int32 BoundaryIndices[2];
+		GetStartEndIndex(InBoundary, BoundaryIndices);
+		return ProjectPointToPolyline(BoundaryIndices, PointOnEdge, OutProjectedPoint);
+	}
 
-		/**
-		 * Evaluate the points of the polyline associated to the increasing array of input Coordinates
-		 * If the input coordinate is outside the boundary of the polyline, the coordinate of the nearest boundary is used.
-		 */
-		void ApproximatePoints(const TArray<double>& InCoordinates, TArray<PointType>& OutPoints) const
-		{
-			if (!InCoordinates.Num())
-			{
-				return;
-			}
+	/**
+	 * Project a Set of points on a restricted polyline (StartIndex & EndIndex define the polyline boundary)
+	 * Each points are projected on all segments of the restricted polyline, the closest are selected
+	 */
+	void ProjectPointsToPolyline(const FLinearBoundary& InBoundary, const TArray<PointType>& InPointsToProject, TArray<double>& OutProjectedPointCoords, TArray<PointType>& OutProjectedPoints) const
+	{
+		OutProjectedPointCoords.Empty(InPointsToProject.Num());
+		OutProjectedPoints.Empty(InPointsToProject.Num());
 
-			TFunction<void(FIndexOfCoordinateFinder&)> ComputePoints = [&](FIndexOfCoordinateFinder& Finder)
-			{
-				for (double Coordinate : InCoordinates)
-				{
-					int32 Index = Finder.Find(Coordinate);
-					OutPoints.Emplace(ComputePoint(Index, Coordinate));
-				}
-			};
+		int32 BoundaryIndices[2];
+		GetStartEndIndex(InBoundary, BoundaryIndices);
 
-			FDichotomyFinder DichotomyFinder(PolylineCoordinates);
+		for (const PointType& Point : InPointsToProject)
+		{
+			PointType ProjectedPoint;
+			double Coordinate = ProjectPointToPolyline(BoundaryIndices, Point, ProjectedPoint);
 
-			int32 StartIndex = DichotomyFinder.Find(InCoordinates[0]);
-			int32 EndIndex = DichotomyFinder.Find(InCoordinates.Last());
-			bool bUseDichotomy = PolylineTools::IsDichotomyToBePreferred(EndIndex - StartIndex, InCoordinates.Num());
+			OutProjectedPointCoords.Emplace(Coordinate);
+			OutProjectedPoints.Emplace(ProjectedPoint);
+		}
+	}
 
-			OutPoints.Empty(InCoordinates.Num());
-			if (bUseDichotomy)
+	/**
+	 * Project each point of a coincidental polyline on the Polyline.
+	 */
+	void ProjectCoincidentalPolyline(const TArray<PointType>& InPointsToProject, bool bSameOrientation, TArray<double>& OutProjectedPointCoords) const
+	{
+#ifdef DEBUG_PROJECT_COINCIDENTAL_POLYLINE
+		static bool bDisplay = false;
+		if (bDisplay)
+		{
+			F3DDebugSession Session(TEXT("ProjectCoincidentalPolyline"));
 			{
-				DichotomyFinder.StartLower = StartIndex;
-				DichotomyFinder.StartUpper = EndIndex;
-				ComputePoints(DichotomyFinder);
+				F3DDebugSession _(TEXT("ProjectCoincidentalPolyline"));
+				DisplayPolyline(InPointsToProject, EVisuProperty::BlueCurve);
 			}
-			else
 			{
-				FLinearFinder LinearFinder(PolylineCoordinates, StartIndex);
-				ComputePoints(LinearFinder);
+				F3DDebugSession _(TEXT("ProjectCoincidentalPolyline"));
+				DisplayPolyline(PolylinePoints, EVisuProperty::YellowCurve);
 			}
 		}
+#endif
+		int32 StartIndex = 0;
+		const int32 EndIndex = PolylinePoints.Num() - 1;
 
-		void SamplePolyline(const FLinearBoundary& InBoundary, const double DesiredSegmentLength, TArray<double>& OutCoordinates) const
+		TFunction<double(const PointType&)> ProjectPointToPolyline = [&](const PointType& InPointToProject)
 		{
-			int32 BoundaryIndices[2];
-
-			TArray<double> CurvilinearCoordinates;
-			double CurveLength = ComputeCurvilinearCoordinatesOfPolyline(InBoundary, CurvilinearCoordinates, BoundaryIndices);
-
-			int32 SegmentNum = (int32)FMath::Max(CurveLength / DesiredSegmentLength + 0.5, 1.0);
-
-			double SectionLength = CurveLength / (double)(SegmentNum);
-
-			OutCoordinates.Empty();
-			OutCoordinates.Reserve(SegmentNum + 1);
-			OutCoordinates.Add(InBoundary.Min);
+			double MinDistance = HUGE_VAL;
+			double UForMinDistance = 0;
 
-			TFunction<double(const int32, const int32, const double, const double)> ComputeSamplePointCoordinate = [&](const int32 IndexCurvilinear, const int32 IndexCoordinate, const double Length, const double Coordinate)
+			double ParamU = 0.;
+			int32 SegmentIndex = 0;
+			for (int32 Index = StartIndex; Index < EndIndex; ++Index)
 			{
-				return Coordinate + (PolylineCoordinates[IndexCoordinate] - Coordinate) * (Length - CurvilinearCoordinates[IndexCurvilinear - 1u]) / (CurvilinearCoordinates[IndexCurvilinear] - CurvilinearCoordinates[IndexCurvilinear - 1u]);
-			};
+				PointType ProjectPoint = ProjectPointOnSegment(InPointToProject, PolylinePoints[Index], PolylinePoints[Index + 1], ParamU, true);
 
-			double CurvilinearLength = SectionLength;
-			double LastCoordinate = InBoundary.Min;
-			for (int32 IndexCurvilinear = 1, IndexCoordinate = BoundaryIndices[0] + 1; IndexCurvilinear < CurvilinearCoordinates.Num(); ++IndexCurvilinear, ++IndexCoordinate)
-			{
-				while (CurvilinearLength < CurvilinearCoordinates[IndexCurvilinear] + SMALL_NUMBER)
+				double SquareDistance = ProjectPoint.SquareDistance(InPointToProject);
+				if (SquareDistance > MinDistance)
 				{
-					double Coordinate = ComputeSamplePointCoordinate(IndexCurvilinear, IndexCoordinate, CurvilinearLength, LastCoordinate);
-					OutCoordinates.Add(Coordinate);
-					CurvilinearLength += SectionLength;
-					if (CurvilinearLength + SMALL_NUMBER > CurveLength)
-					{
-						OutCoordinates.Add(InBoundary.Max);
-						break;
-					}
+					break;
 				}
-				LastCoordinate = PolylineCoordinates[IndexCoordinate];
+				MinDistance = SquareDistance;
+				UForMinDistance = ParamU;
+				SegmentIndex = Index;
 			}
-		}
-
-		/**
-		 * Project a Set of points on a restricted polyline (StartIndex & EndIndex define the polyline boundary)
-		 * the points are projected on all segments of the polyline, the closest are selected
-		 */
-		double ProjectPointToPolyline(const FLinearBoundary& InBoundary, const PointType& PointOnEdge, PointType& OutProjectedPoint) const
-		{
-			int32 BoundaryIndices[2];
-			GetStartEndIndex(InBoundary, BoundaryIndices);
-			return ProjectPointToPolyline(BoundaryIndices, PointOnEdge, OutProjectedPoint);
-		}
+			StartIndex = SegmentIndex;
+			return PolylineTools::LinearInterpolation(PolylineCoordinates, SegmentIndex, UForMinDistance);
+		};
 
-		/**
-		 * Project a Set of points on a restricted polyline (StartIndex & EndIndex define the polyline boundary)
-		 * Each points are projected on all segments of the restricted polyline, the closest are selected
-		 */
-		void ProjectPointsToPolyline(const FLinearBoundary& InBoundary, const TArray<PointType>& InPointsToProject, TArray<double>& OutProjectedPointCoords, TArray<PointType>& OutProjectedPoints) const
+		if (bSameOrientation)
 		{
 			OutProjectedPointCoords.Empty(InPointsToProject.Num());
-			OutProjectedPoints.Empty(InPointsToProject.Num());
-
-			int32 BoundaryIndices[2];
-			GetStartEndIndex(InBoundary, BoundaryIndices);
-
 			for (const PointType& Point : InPointsToProject)
 			{
-				PointType ProjectedPoint;
-				double Coordinate = ProjectPointToPolyline(BoundaryIndices, Point, ProjectedPoint);
-
+				double Coordinate = ProjectPointToPolyline(Point);
 				OutProjectedPointCoords.Emplace(Coordinate);
-				OutProjectedPoints.Emplace(ProjectedPoint);
 			}
 		}
-
-		/**
-		 * Project each point of a coincidental polyline on the Polyline. 
-		 */
-		void ProjectCoincidentalPolyline(const TArray<PointType>& InPointsToProject, bool bSameOrientation, TArray<double>& OutProjectedPointCoords) const
+		else
 		{
-			int32 StartIndex = 0;
-			const int32 EndIndex = PolylinePoints.Num() - 1;
-
-			TFunction<double(const PointType&)> ProjectPointToPolyline = [&](const PointType & InPointToProject)
-			{
-				double MinDistance = HUGE_VAL;
-				double UForMinDistance = 0;
-
-				double ParamU = 0.;
-				int32 SegmentIndex = 0;
-				for (int32 Index = StartIndex; Index < EndIndex; ++Index)
-				{
-					PointType ProjectPoint = ProjectPointOnSegment(InPointToProject, PolylinePoints[Index], PolylinePoints[Index + 1], ParamU, true);
-
-					double SquareDistance = ProjectPoint.SquareDistance(InPointToProject);
-					if (SquareDistance > MinDistance)
-					{
-						break;
-					}
-					MinDistance = SquareDistance;
-					UForMinDistance = ParamU;
-					SegmentIndex = Index;
-				}
-				StartIndex = SegmentIndex;
-				return PolylineTools::LinearInterpolation(PolylineCoordinates, SegmentIndex, UForMinDistance);
-			};
-
-			if (bSameOrientation)
-			{
-				OutProjectedPointCoords.Empty(InPointsToProject.Num());
-				for (const PointType& Point : InPointsToProject)
-				{
-					double Coordinate = ProjectPointToPolyline(Point);
-					OutProjectedPointCoords.Emplace(Coordinate);
-				}
-			}
-			else
+			OutProjectedPointCoords.SetNum(InPointsToProject.Num());
+			for (int32 Index = InPointsToProject.Num() - 1, Pndex = 0; Index >= 0; --Index, ++Pndex)
 			{
-				OutProjectedPointCoords.SetNum(InPointsToProject.Num());
-				for (int32 Index = InPointsToProject.Num()-1, Pndex = 0; Index >= 0; --Index, ++Pndex)
-				{
-					OutProjectedPointCoords[Pndex] = ProjectPointToPolyline(InPointsToProject[Index]);
-				}
+				OutProjectedPointCoords[Pndex] = ProjectPointToPolyline(InPointsToProject[Index]);
 			}
 		}
+	}
 
 
-		/**
-		 * Append to the OutPoints array the sub polyline bounded by InBoundary according to the orientation
-		 */
-		void GetSubPolyline(const FLinearBoundary& InBoundary, const EOrientation Orientation, TArray<PointType>& OutPoints) const
-		{
-			int32 BoundaryIndices[2];
-			GetStartEndIndex(InBoundary, BoundaryIndices);
+	/**
+	 * Append to the OutPoints array the sub polyline bounded by InBoundary according to the orientation
+	 */
+	void GetSubPolyline(const FLinearBoundary& InBoundary, const EOrientation Orientation, TArray<PointType>& OutPoints) const
+	{
+		int32 BoundaryIndices[2];
+		GetStartEndIndex(InBoundary, BoundaryIndices);
 
-			int32 NewSize = OutPoints.Num() + BoundaryIndices[1] - BoundaryIndices[0] + 2;
-			OutPoints.Reserve(NewSize);
+		int32 NewSize = OutPoints.Num() + BoundaryIndices[1] - BoundaryIndices[0] + 2;
+		OutPoints.Reserve(NewSize);
 
-			int32 PolylineStartIndex = BoundaryIndices[0];
-			int32 PolylineEndIndex = BoundaryIndices[1];
-			if (FMath::IsNearlyEqual(PolylineCoordinates[BoundaryIndices[0] + 1], InBoundary.Min, (double)KINDA_SMALL_NUMBER))
-			{
-				PolylineStartIndex++;
-			}
-			if (FMath::IsNearlyEqual(PolylineCoordinates[BoundaryIndices[1]], InBoundary.Max, (double)KINDA_SMALL_NUMBER))
-			{
-				PolylineEndIndex--;
-			}
+		int32 PolylineStartIndex = BoundaryIndices[0];
+		int32 PolylineEndIndex = BoundaryIndices[1];
+		if (FMath::IsNearlyEqual(PolylineCoordinates[BoundaryIndices[0] + 1], InBoundary.Min, (double)SMALL_NUMBER))
+		{
+			PolylineStartIndex++;
+		}
+		if (FMath::IsNearlyEqual(PolylineCoordinates[BoundaryIndices[1]], InBoundary.Max, (double)SMALL_NUMBER))
+		{
+			PolylineEndIndex--;
+		}
 
-			if (Orientation)
-			{
-				OutPoints.Emplace(ComputePoint(BoundaryIndices[0], InBoundary.Min));
-				if(PolylineEndIndex - PolylineStartIndex > 0)
-				{
-					OutPoints.Append(PolylinePoints.GetData() + PolylineStartIndex + 1, PolylineEndIndex - PolylineStartIndex);
-				}
-				OutPoints.Emplace(ComputePoint(BoundaryIndices[1], InBoundary.Max));
-			}
-			else
+		if (Orientation)
+		{
+			OutPoints.Emplace(ComputePoint(BoundaryIndices[0], InBoundary.Min));
+			if (PolylineEndIndex - PolylineStartIndex > 0)
 			{
-				OutPoints.Emplace(ComputePoint(BoundaryIndices[1], InBoundary.Max));
-				for (int32 Index = PolylineEndIndex; Index > PolylineStartIndex; --Index)
-				{
-					OutPoints.Emplace(PolylinePoints[Index]);
-				}
-				OutPoints.Emplace(ComputePoint(BoundaryIndices[0], InBoundary.Min));
+				OutPoints.Append(PolylinePoints.GetData() + PolylineStartIndex + 1, PolylineEndIndex - PolylineStartIndex);
 			}
+			OutPoints.Emplace(ComputePoint(BoundaryIndices[1], InBoundary.Max));
 		}
-
-		/**
-		 * Get the subset of point defining the sub polyline bounded by InBoundary
-		 * OutCoordinates and OutPoints are emptied before the process
-		 */
-		void GetSubPolyline(const FLinearBoundary& InBoundary, TArray<double>& OutCoordinates, TArray<PointType>& OutPoints) const
+		else
 		{
-			int32 StartIndex = 0;
-			int32 EndIndex = 0;
-			GetStartEndIndex(PolylineCoordinates, InBoundary, StartIndex, EndIndex);
-
-			OutCoordinates.Empty(EndIndex - StartIndex + 2);
-			OutPoints.Empty(EndIndex - StartIndex + 2);
-
-			if (FMath::IsNearlyEqual(PolylineCoordinates[StartIndex + 1], InBoundary.Min, (double) KINDA_SMALL_NUMBER))
+			OutPoints.Emplace(ComputePoint(BoundaryIndices[1], InBoundary.Max));
+			for (int32 Index = PolylineEndIndex; Index > PolylineStartIndex; --Index)
 			{
-				StartIndex++;
+				OutPoints.Emplace(PolylinePoints[Index]);
 			}
+			OutPoints.Emplace(ComputePoint(BoundaryIndices[0], InBoundary.Min));
+		}
+	}
 
-			if (!FMath::IsNearlyEqual(PolylineCoordinates[EndIndex], InBoundary.Max, (double) KINDA_SMALL_NUMBER))
-			{
-				EndIndex--;
-			}
+	/**
+	 * Get the subset of point defining the sub polyline bounded by InBoundary
+	 * OutCoordinates and OutPoints are emptied before the process
+	 */
+	void GetSubPolyline(const FLinearBoundary& InBoundary, TArray<double>& OutCoordinates, TArray<PointType>& OutPoints) const
+	{
+		int32 StartIndex = 0;
+		int32 EndIndex = 0;
+		GetStartEndIndex(PolylineCoordinates, InBoundary, StartIndex, EndIndex);
 
-			OutPoints.Emplace(ComputePoint(StartIndex, InBoundary.Min));
-			OutPoints.Append(PolylinePoints.GetData() + StartIndex + 1, EndIndex - StartIndex);
-			OutPoints.Emplace(ComputePoint(EndIndex, InBoundary.Max));
+		OutCoordinates.Empty(EndIndex - StartIndex + 2);
+		OutPoints.Empty(EndIndex - StartIndex + 2);
 
-			OutCoordinates.Add(InBoundary.Min);
-			OutCoordinates.Append(PolylineCoordinates.GetData() + StartIndex + 1, EndIndex - StartIndex);
-			OutCoordinates.Add(InBoundary.Max);
+		if (FMath::IsNearlyEqual(PolylineCoordinates[StartIndex + 1], InBoundary.Min, (double)SMALL_NUMBER))
+		{
+			StartIndex++;
 		}
 
-		double ComputeLength() const
+		if (!FMath::IsNearlyEqual(PolylineCoordinates[EndIndex], InBoundary.Max, (double)SMALL_NUMBER))
 		{
-			return PolylineTools::ComputeLength(PolylinePoints);
+			EndIndex--;
 		}
 
-		double ComputeLengthOfSubPolyline(const int BoundaryIndex[2], const FLinearBoundary& InBoundary) const
+		OutPoints.Emplace(ComputePoint(StartIndex, InBoundary.Min));
+		OutPoints.Append(PolylinePoints.GetData() + StartIndex + 1, EndIndex - StartIndex);
+		OutPoints.Emplace(ComputePoint(EndIndex, InBoundary.Max));
+
+		OutCoordinates.Add(InBoundary.Min);
+		OutCoordinates.Append(PolylineCoordinates.GetData() + StartIndex + 1, EndIndex - StartIndex);
+		OutCoordinates.Add(InBoundary.Max);
+	}
+
+	double ComputeLength() const
+	{
+		return PolylineTools::ComputeLength(PolylinePoints);
+	}
+
+	double ComputeLengthOfSubPolyline(const int BoundaryIndex[2], const FLinearBoundary& InBoundary) const
+	{
+		double Length = 0;
+		if (BoundaryIndex[1] > BoundaryIndex[0])
 		{
-			double Length = 0;
-			if (BoundaryIndex[1] > BoundaryIndex[0])
-			{
-				PointType StartPoint = ComputePoint(BoundaryIndex[0], InBoundary.Min);
-				PointType EndPoint = ComputePoint(BoundaryIndex[1], InBoundary.Max);
-				Length = StartPoint.Distance(PolylinePoints[BoundaryIndex[0] + 1]);
-				Length += EndPoint.Distance(PolylinePoints[BoundaryIndex[1]]);
+			PointType StartPoint = ComputePoint(BoundaryIndex[0], InBoundary.Min);
+			PointType EndPoint = ComputePoint(BoundaryIndex[1], InBoundary.Max);
+			Length = StartPoint.Distance(PolylinePoints[BoundaryIndex[0] + 1]);
+			Length += EndPoint.Distance(PolylinePoints[BoundaryIndex[1]]);
 
-				if (BoundaryIndex[1] > BoundaryIndex[0] + 1)
+			if (BoundaryIndex[1] > BoundaryIndex[0] + 1)
+			{
+				for (int32 Index = BoundaryIndex[0] + 1; Index < BoundaryIndex[1]; ++Index)
 				{
-					for (int32 Index = BoundaryIndex[0] + 1; Index < BoundaryIndex[1]; ++Index)
-					{
-						Length += PolylinePoints[Index].Distance(PolylinePoints[Index + 1]);
-					}
+					Length += PolylinePoints[Index].Distance(PolylinePoints[Index + 1]);
 				}
 			}
-			else
-			{
-				PointType StartPoint = ComputePoint(BoundaryIndex[0], InBoundary.Min);
-				PointType EndPoint = ComputePoint(BoundaryIndex[1], InBoundary.Max);
-				Length = StartPoint.Distance(EndPoint);
-			}
-			return Length;
 		}
-
-		double ComputeLengthOfSubPolyline(const FLinearBoundary& InBoundary) const
+		else
 		{
-			int BoundaryIndices[2];
-			GetStartEndIndex(InBoundary, BoundaryIndices);
-			return ComputeLengthOfSubPolyline(BoundaryIndices, InBoundary);
+			PointType StartPoint = ComputePoint(BoundaryIndex[0], InBoundary.Min);
+			PointType EndPoint = ComputePoint(BoundaryIndex[1], InBoundary.Max);
+			Length = StartPoint.Distance(EndPoint);
 		}
+		return Length;
+	}
+
+	double ComputeLengthOfSubPolyline(const FLinearBoundary& InBoundary) const
+	{
+		int BoundaryIndices[2];
+		GetStartEndIndex(InBoundary, BoundaryIndices);
+		return ComputeLengthOfSubPolyline(BoundaryIndices, InBoundary);
+	}
 
-		void ComputeBoundingBox(const int BoundaryIndex[2], const FLinearBoundary& InBoundary, TAABB<PointType>& Aabb)
+	void ComputeBoundingBox(const int BoundaryIndex[2], const FLinearBoundary& InBoundary, TAABB<PointType>& Aabb)
+	{
+		Aabb.Empty();
+		if (BoundaryIndex[1] > BoundaryIndex[0])
 		{
-			Aabb.Empty();
-			if (BoundaryIndex[1] > BoundaryIndex[0])
+			for (int32 Index = BoundaryIndex[0] + 1; Index <= BoundaryIndex[1]; ++Index)
 			{
-				for (int32 Index = BoundaryIndex[0] + 1; Index <= BoundaryIndex[1]; ++Index)
-				{
-					Aabb += PolylinePoints[Index];
-				}
+				Aabb += PolylinePoints[Index];
 			}
-
-			PointType StartPoint = ComputePoint(BoundaryIndex[0], InBoundary.Min);
-			Aabb += StartPoint;
-			PointType EndPoint = ComputePoint(BoundaryIndex[1], InBoundary.Max);
-			Aabb += EndPoint;
 		}
-	};
+
+		PointType StartPoint = ComputePoint(BoundaryIndex[0], InBoundary.Min);
+		Aabb += StartPoint;
+		PointType EndPoint = ComputePoint(BoundaryIndex[1], InBoundary.Max);
+		Aabb += EndPoint;
+	}
+};
 
 } // namespace CADKernel
 
diff --git a/Engine/Source/Runtime/Datasmith/CADKernel/Public/CADKernel/Geo/Surfaces/ConeSurface.h b/Engine/Source/Runtime/Datasmith/CADKernel/Public/CADKernel/Geo/Surfaces/ConeSurface.h
--- a/Engine/Source/Runtime/Datasmith/CADKernel/Public/CADKernel/Geo/Surfaces/ConeSurface.h
+++ b/Engine/Source/Runtime/Datasmith/CADKernel/Public/CADKernel/Geo/Surfaces/ConeSurface.h
@@ -71,7 +71,9 @@ namespace CADKernel
 			ToleranceU /= ScaleU;
 
 			FPoint Point2DV{ 0, 1, 0 };
-			double ToleranceV = Tolerance3D / ComputeScaleAlongAxis(Point2DV, Matrix, Origin);
+			double ScaleV = ComputeScaleAlongAxis(Point2DV, Matrix, Origin);
+			double ToleranceV = Tolerance3D / FMath::Sqrt(FMath::Square(DeltaVR) + 1);
+			ToleranceV /= ScaleV;
 
 			MinToleranceIso.Set(ToleranceU, ToleranceV);
 		}
diff --git a/Engine/Source/Runtime/Datasmith/CADKernel/Public/CADKernel/Geo/Surfaces/NURBSSurface.h b/Engine/Source/Runtime/Datasmith/CADKernel/Public/CADKernel/Geo/Surfaces/NURBSSurface.h
--- a/Engine/Source/Runtime/Datasmith/CADKernel/Public/CADKernel/Geo/Surfaces/NURBSSurface.h
+++ b/Engine/Source/Runtime/Datasmith/CADKernel/Public/CADKernel/Geo/Surfaces/NURBSSurface.h
@@ -7,210 +7,248 @@
 namespace CADKernel
 {
 
-struct FNurbsSurfaceHomogeneousData
-{
-	bool bSwapUV;
-	int32 PoleUCount;
-	int32 PoleVCount;
+	struct FNurbsSurfaceHomogeneousData
+	{
+		bool bSwapUV;
+		int32 PoleUCount;
+		int32 PoleVCount;
 
-	int32 UDegree;
-	int32 VDegree;
+		int32 UDegree;
+		int32 VDegree;
 
-	TArray<double> UNodalVector;
-	TArray<double> VNodalVector;
+		TArray<double> UNodalVector;
+		TArray<double> VNodalVector;
 
-	bool bIsRational;
-	TArray<double> HomogeneousPoles;
-};
+		bool bIsRational;
+		TArray<double> HomogeneousPoles;
+	};
 
-class CADKERNEL_API FNURBSSurface : public FSurface
-{
-	friend FEntity;
-
-protected:
-	int32 PoleUCount;
-	int32 PoleVCount;
-
-	int32 UDegree;
-	int32 VDegree;
-
-	TArray<double> UNodalVector;
-	TArray<double> VNodalVector;
-
-	TArray<double> Weights;
-
-	TArray<FPoint> Poles;
-
-	bool bIsRational;
-
-	/**
-	 * Data generated at initialization which are not serialized
-	 */
-	TArray<double> HomogeneousPoles;
-
-	/**
-	 * Build a Non uniform B-Spline surface
-	 * @param NodalVectorU Its size is the number of poles in U + the surface degree in U + 1 (PoleUNum + UDegre + 1)
-	 * @param NodalVectorV Its size is the number of poles in V + the surface degree in V + 1 (PoleVNum + VDegre + 1)
-	 */
-	FNURBSSurface(const double InToleranceGeometric, int32 InPoleUCount, int32 InPoleVCount, int32 InDegreU, int32 InDegreV, const TArray<double>& InNodalVectorU, const TArray<double>& InNodalVectorV, const TArray<FPoint>& InPoles)
-		: FSurface(InToleranceGeometric)
-		, PoleUCount(InPoleUCount)
-		, PoleVCount(InPoleVCount)
-		, UDegree(InDegreU)
-		, VDegree(InDegreV)
-		, UNodalVector(InNodalVectorU)
-		, VNodalVector(InNodalVectorV)
-		, Poles(InPoles)
-		, bIsRational(false)
-	{
-		Finalize();
-	}
-
-	/**
-	 * Build a Non uniform rational B-Spline surface
-	 * @param NodalVectorU Its size is the number of poles in U + the surface degree in U + 1 (PoleUNum + UDegre + 1)
-	 * @param NodalVectorV Its size is the number of poles in V + the surface degree in V + 1 (PoleVNum + VDegre + 1)
-	 */
-	FNURBSSurface(const double InToleranceGeometric, int32 InPoleUCount, int32 InPoleVCount, int32 InDegreU, int32 InDegreV, const TArray<double>& InNodalVectorU, const TArray<double>& InNodalVectorV, const TArray<FPoint>& InPoles, const TArray<double>& InWeights)
-		: FSurface(InToleranceGeometric)
-		, PoleUCount(InPoleUCount)
-		, PoleVCount(InPoleVCount)
-		, UDegree(InDegreU)
-		, VDegree(InDegreV)
-		, UNodalVector(InNodalVectorU)
-		, VNodalVector(InNodalVectorV)
-		, Weights(InWeights)
-		, Poles(InPoles)
-		, bIsRational(true)
+	struct FNurbsSurfaceData
 	{
-		ComputeMinToleranceIso();
-		Finalize();
-	}
-
-	/**
-	 * Build a Non uniform B-Spline surface
-	 * @param NodalVectorU Its size is the number of poles in U + the surface degree in U + 1 (PoleUNum + UDegre + 1)
-	 * @param NodalVectorV Its size is the number of poles in V + the surface degree in V + 1 (PoleVNum + VDegre + 1)
-	 */
-	FNURBSSurface(const double InToleranceGeometric, FNurbsSurfaceHomogeneousData NurbsData)
-		: FSurface(InToleranceGeometric)
-	{
-		FillNurbs(NurbsData);
-	}
+		bool bSwapUV;
+		int32 PoleUCount;
+		int32 PoleVCount;
+
+		int32 UDegree;
+		int32 VDegree;
 
-	FNURBSSurface() = default;
+		TArray<double> UNodalVector;
+		TArray<double> VNodalVector;
 
-public:
+		// if Weights.Num() == 0  => bIsRational = false
+		TArray<double> Weights;
+		TArray<FPoint> Poles;
+	};
 
-	virtual void Serialize(FCADKernelArchive& Ar) override
+	class CADKERNEL_API FNURBSSurface : public FSurface
 	{
-		FSurface::Serialize(Ar);
-		Ar << PoleUCount;
-		Ar << PoleVCount;
-		Ar << UDegree;
-		Ar << VDegree;
-		Ar.Serialize(UNodalVector);
-		Ar.Serialize(VNodalVector);
-		Ar.Serialize(Weights);
-		Ar.Serialize(Poles);
-		Ar << bIsRational;
-
-		if (Ar.IsLoading())
+		friend FEntity;
+
+	protected:
+		int32 PoleUCount;
+		int32 PoleVCount;
+
+		int32 UDegree;
+		int32 VDegree;
+
+		TArray<double> UNodalVector;
+		TArray<double> VNodalVector;
+
+		TArray<double> Weights;
+
+		TArray<FPoint> Poles;
+
+		bool bIsRational;
+
+		/**
+		 * Data generated at initialization which are not serialized
+		 */
+		TArray<double> HomogeneousPoles;
+
+		/**
+		 * Build a Non uniform B-Spline surface
+		 * @param NodalVectorU Its size is the number of poles in U + the surface degree in U + 1 (PoleUNum + UDegre + 1)
+		 * @param NodalVectorV Its size is the number of poles in V + the surface degree in V + 1 (PoleVNum + VDegre + 1)
+		 */
+		FNURBSSurface(const double InToleranceGeometric, int32 InPoleUCount, int32 InPoleVCount, int32 InDegreU, int32 InDegreV, const TArray<double>& InNodalVectorU, const TArray<double>& InNodalVectorV, const TArray<FPoint>& InPoles)
+			: FSurface(InToleranceGeometric)
+			, PoleUCount(InPoleUCount)
+			, PoleVCount(InPoleVCount)
+			, UDegree(InDegreU)
+			, VDegree(InDegreV)
+			, UNodalVector(InNodalVectorU)
+			, VNodalVector(InNodalVectorV)
+			, Poles(InPoles)
+			, bIsRational(false)
 		{
 			Finalize();
 		}
-	}
 
-	ESurface GetSurfaceType() const
-	{
-		return ESurface::Nurbs;
-	}
+		/**
+		 * Build a Non uniform rational B-Spline surface
+		 * @param NodalVectorU Its size is the number of poles in U + the surface degree in U + 1 (PoleUNum + UDegre + 1)
+		 * @param NodalVectorV Its size is the number of poles in V + the surface degree in V + 1 (PoleVNum + VDegre + 1)
+		 */
+		FNURBSSurface(const double InToleranceGeometric, int32 InPoleUCount, int32 InPoleVCount, int32 InDegreU, int32 InDegreV, const TArray<double>& InNodalVectorU, const TArray<double>& InNodalVectorV, const TArray<FPoint>& InPoles, const TArray<double>& InWeights)
+			: FSurface(InToleranceGeometric)
+			, PoleUCount(InPoleUCount)
+			, PoleVCount(InPoleVCount)
+			, UDegree(InDegreU)
+			, VDegree(InDegreV)
+			, UNodalVector(InNodalVectorU)
+			, VNodalVector(InNodalVectorV)
+			, Weights(InWeights)
+			, Poles(InPoles)
+			, bIsRational(true)
+		{
+			ComputeMinToleranceIso();
+			Finalize();
+		}
+
+		/**
+		 * Build a Non uniform B-Spline surface
+		 * @param NodalVectorU Its size is the number of poles in U + the surface degree in U + 1 (PoleUNum + UDegre + 1)
+		 * @param NodalVectorV Its size is the number of poles in V + the surface degree in V + 1 (PoleVNum + VDegre + 1)
+		 */
+		FNURBSSurface(const double InToleranceGeometric, FNurbsSurfaceData NurbsData)
+			: FSurface(InToleranceGeometric)
+			, PoleUCount(NurbsData.PoleUCount)
+			, PoleVCount(NurbsData.PoleVCount)
+			, UDegree(NurbsData.UDegree)
+			, VDegree(NurbsData.VDegree)
+			, UNodalVector(NurbsData.UNodalVector)
+			, VNodalVector(NurbsData.VNodalVector)
+			, Weights(NurbsData.Weights)
+			, Poles(NurbsData.Poles)
+			, bIsRational(true)
+		{
+			ComputeMinToleranceIso();
+			Finalize();
+		}
+
+		/**
+		 * Build a Non uniform B-Spline surface
+		 * @param NodalVectorU Its size is the number of poles in U + the surface degree in U + 1 (PoleUNum + UDegre + 1)
+		 * @param NodalVectorV Its size is the number of poles in V + the surface degree in V + 1 (PoleVNum + VDegre + 1)
+		 */
+		FNURBSSurface(const double InToleranceGeometric, FNurbsSurfaceHomogeneousData NurbsData)
+			: FSurface(InToleranceGeometric)
+		{
+			FillNurbs(NurbsData);
+		}
+
+		FNURBSSurface() = default;
+
+	public:
+
+		virtual void Serialize(FCADKernelArchive& Ar) override
+		{
+			FSurface::Serialize(Ar);
+			Ar << PoleUCount;
+			Ar << PoleVCount;
+			Ar << UDegree;
+			Ar << VDegree;
+			Ar.Serialize(UNodalVector);
+			Ar.Serialize(VNodalVector);
+			Ar.Serialize(Weights);
+			Ar.Serialize(Poles);
+			Ar << bIsRational;
+
+			if (Ar.IsLoading())
+			{
+				Finalize();
+			}
+		}
+
+		ESurface GetSurfaceType() const
+		{
+			return ESurface::Nurbs;
+		}
 
 #ifdef CADKERNEL_DEV
-	virtual FInfoEntity& GetInfo(FInfoEntity&) const override;
+		virtual FInfoEntity& GetInfo(FInfoEntity&) const override;
 #endif
 
-	constexpr const int32 GetDegree(EIso Iso) const
-	{
-		switch (Iso)
+		constexpr const int32 GetDegree(EIso Iso) const
 		{
-		case EIso::IsoU:
-			return UDegree;
-		case EIso::IsoV:
-		default:
-			return VDegree;
+			switch (Iso)
+			{
+			case EIso::IsoU:
+				return UDegree;
+			case EIso::IsoV:
+			default:
+				return VDegree;
+			}
 		}
-	}
 
-	constexpr const int32 GetPoleCount(EIso Iso) const
-	{
-		switch (Iso)
+		constexpr const int32 GetPoleCount(EIso Iso) const
 		{
-		case EIso::IsoU:
-			return PoleUCount;
-		case EIso::IsoV:
-		default:
-			return PoleVCount;
+			switch (Iso)
+			{
+			case EIso::IsoU:
+				return PoleUCount;
+			case EIso::IsoV:
+			default:
+				return PoleVCount;
+			}
 		}
-	}
 
-	const TArray<FPoint>& GetPoles() const
-	{
-		return Poles;
-	}
+		const TArray<FPoint>& GetPoles() const
+		{
+			return Poles;
+		}
 
-	const TArray<double>& GetWeights() const
-	{
-		return Weights;
-	}
+		const TArray<double>& GetWeights() const
+		{
+			return Weights;
+		}
 
-	TArray<double> GetHPoles() const
-	{
-		return HomogeneousPoles;
-	}
+		TArray<double> GetHPoles() const
+		{
+			return HomogeneousPoles;
+		}
 
-	constexpr const TArray<double>& GetNodalVector(EIso Iso) const
-	{
-		switch (Iso)
+		constexpr const TArray<double>& GetNodalVector(EIso Iso) const
 		{
-		case EIso::IsoU:
-			return UNodalVector;
-		case EIso::IsoV:
-		default:
-			return VNodalVector;
+			switch (Iso)
+			{
+			case EIso::IsoU:
+				return UNodalVector;
+			case EIso::IsoV:
+			default:
+				return VNodalVector;
+			}
 		}
-	}
 
-	bool IsRational() const
-	{
-		return bIsRational;
-	}
+		bool IsRational() const
+		{
+			return bIsRational;
+		}
 
 
-	virtual TSharedPtr<FEntityGeom> ApplyMatrix(const FMatrixH& InMatrix) const override;
+		virtual TSharedPtr<FEntityGeom> ApplyMatrix(const FMatrixH& InMatrix) const override;
 
-	virtual void EvaluatePoint(const FPoint2D& InSurfacicCoordinate, FSurfacicPoint& OutPoint3D, int32 InDerivativeOrder = 0) const override
-	{
-		BSpline::EvaluatePoint(*this, InSurfacicCoordinate, OutPoint3D, InDerivativeOrder);
-	}
+		virtual void EvaluatePoint(const FPoint2D& InSurfacicCoordinate, FSurfacicPoint& OutPoint3D, int32 InDerivativeOrder = 0) const override
+		{
+			BSpline::EvaluatePoint(*this, InSurfacicCoordinate, OutPoint3D, InDerivativeOrder);
+		}
 
-	void EvaluatePointGrid(const FCoordinateGrid& InSurfacicCoordinates, FSurfacicSampling& OutPoints, bool bComputeNormals) const override
-	{
-		BSpline::EvaluatePointGrid(*this, InSurfacicCoordinates, OutPoints, bComputeNormals);
-	}
+		void EvaluatePointGrid(const FCoordinateGrid& InSurfacicCoordinates, FSurfacicSampling& OutPoints, bool bComputeNormals) const override
+		{
+			BSpline::EvaluatePointGrid(*this, InSurfacicCoordinates, OutPoints, bComputeNormals);
+		}
 
-	virtual void LinesNotDerivables(const FSurfacicBoundary& InBoundary, int32 InDerivativeOrder, FCoordinateGrid& OutNotDerivableCoordinates) const override
-	{
-		BSpline::FindNotDerivableParameters(*this, InDerivativeOrder, InBoundary, OutNotDerivableCoordinates);
-	}
+		virtual void LinesNotDerivables(const FSurfacicBoundary& InBoundary, int32 InDerivativeOrder, FCoordinateGrid& OutNotDerivableCoordinates) const override
+		{
+			BSpline::FindNotDerivableParameters(*this, InDerivativeOrder, InBoundary, OutNotDerivableCoordinates);
+		}
 
-	void ComputeMinToleranceIso();
+		void ComputeMinToleranceIso();
 
-private:
-	void Finalize();
-	void FillNurbs(FNurbsSurfaceHomogeneousData& NurbsData);
-};
+	private:
+		void Finalize();
+		void FillNurbs(FNurbsSurfaceHomogeneousData& NurbsData);
+	};
 
 } // namespace CADKernel
 
diff --git a/Engine/Source/Runtime/Datasmith/CADKernel/Public/CADKernel/Math/Point.h b/Engine/Source/Runtime/Datasmith/CADKernel/Public/CADKernel/Math/Point.h
--- a/Engine/Source/Runtime/Datasmith/CADKernel/Public/CADKernel/Math/Point.h
+++ b/Engine/Source/Runtime/Datasmith/CADKernel/Public/CADKernel/Math/Point.h
@@ -222,7 +222,7 @@ public:
 
 	double Length() const
 	{
-		return sqrt((X * X) + (Y * Y) + (Z * Z));
+		return FMath::Sqrt((X * X) + (Y * Y) + (Z * Z));
 	}
 
 	double SquareLength() const
@@ -253,7 +253,7 @@ public:
 
 	double Distance(const FPoint& Point) const
 	{
-		return sqrt(FMath::Square(Point.X - X) + FMath::Square(Point.Y - Y) + FMath::Square(Point.Z - Z));
+		return FMath::Sqrt(FMath::Square(Point.X - X) + FMath::Square(Point.Y - Y) + FMath::Square(Point.Z - Z));
 	}
 
 	double SquareDistance(const FPoint& Point) const
@@ -446,7 +446,7 @@ public:
 
 	double Distance(const FPoint2D& Point) const
 	{
-		return sqrt(FMath::Square(Point.U - U) + FMath::Square(Point.V - V));
+		return FMath::Sqrt(FMath::Square(Point.U - U) + FMath::Square(Point.V - V));
 	}
 
 	double SquareDistance(const FPoint2D& Point) const
@@ -456,7 +456,7 @@ public:
 
 	double Length() const
 	{
-		return sqrt(FMath::Square(U) + FMath::Square(V));
+		return FMath::Sqrt(FMath::Square(U) + FMath::Square(V));
 	}
 
 	double SquareLength() const
diff --git a/Engine/Source/Runtime/Datasmith/CADKernel/Public/CADKernel/Mesh/Meshers/MesherReport.h b/Engine/Source/Runtime/Datasmith/CADKernel/Public/CADKernel/Mesh/Meshers/MesherReport.h
--- a/Engine/Source/Runtime/Datasmith/CADKernel/Public/CADKernel/Mesh/Meshers/MesherReport.h
+++ b/Engine/Source/Runtime/Datasmith/CADKernel/Public/CADKernel/Mesh/Meshers/MesherReport.h
@@ -27,6 +27,14 @@ public:
 
 	void PrintReport() const
 	{
+		FMessage::FillReportFile(TEXT("SelfInterFail"), CountOfRemoveLoopSelfIntersectionFailure);
+		FMessage::FillReportFile(TEXT("CrossingFail"), CountOfCrossingLoopsFailure);
+		FMessage::FillReportFile(TEXT(""), TEXT(""));
+		FMessage::FillReportFile(TEXT("MeshingFail"), CountOfMeshingFailure);
+		FMessage::FillReportFile(TEXT("DegenGrid"), CountOfDegeneratedGrid);
+		FMessage::FillReportFile(TEXT("DegenLoop"), CountOfDegeneratedLoop);
+		FMessage::FillReportFile(TEXT("CycleFailure"), CountOfCycleMeshingFailure);
+		FMessage::FillReportFile(TEXT(""), TEXT(""));
 	}
 
 	void AddRemoveCrossingLoopsFailure()
@@ -108,6 +116,16 @@ struct CADKERNEL_API FMesherChronos
 
 	void PrintReport() const
 	{
+		FMessage::FillReportFile(TEXT("Criteria"), ApplyCriteriaDuration);
+		FMessage::FillReportFile(TEXT("Find Quad"), IsolateQuadPatchDuration);
+		FMessage::FillReportFile(TEXT("GenPoint"), GlobalGeneratePointCloudDuration);
+		FMessage::FillReportFile(TEXT("FindThin"), GlobalThinZones);
+		FMessage::FillReportFile(TEXT("MeshThin"), GlobalMeshThinZones);
+		FMessage::FillReportFile(TEXT("MeshEdges"), GlobalMeshEdges);
+		FMessage::FillReportFile(TEXT("Triangul"),  GlobalTriangulateDuration);
+		FMessage::FillReportFile(TEXT("MeshTotal"), GlobalDuration);
+		FMessage::FillReportFile(TEXT(""), TEXT(""));
+
 	}
 };
 
diff --git a/Engine/Source/Runtime/Datasmith/CADKernel/Public/CADKernel/Mesh/Meshers/ParametricMesher.h b/Engine/Source/Runtime/Datasmith/CADKernel/Public/CADKernel/Mesh/Meshers/ParametricMesher.h
--- a/Engine/Source/Runtime/Datasmith/CADKernel/Public/CADKernel/Mesh/Meshers/ParametricMesher.h
+++ b/Engine/Source/Runtime/Datasmith/CADKernel/Public/CADKernel/Mesh/Meshers/ParametricMesher.h
@@ -115,6 +115,8 @@ namespace CADKernel
 
 		void GenerateCloud(FGrid& Grid);
 
+		void PrintReport();
+
 	protected:
 
 		void MeshEntities();
diff --git a/Engine/Source/Runtime/Datasmith/CADKernel/Public/CADKernel/Topo/Body.h b/Engine/Source/Runtime/Datasmith/CADKernel/Public/CADKernel/Topo/Body.h
--- a/Engine/Source/Runtime/Datasmith/CADKernel/Public/CADKernel/Topo/Body.h
+++ b/Engine/Source/Runtime/Datasmith/CADKernel/Public/CADKernel/Topo/Body.h
@@ -105,6 +105,8 @@ namespace CADKernel
 			}
 		}
 
+		virtual void FillTopologyReport(FTopologyReport& Report) const override;
+
 		virtual TSharedPtr<FEntityGeom> ApplyMatrix(const FMatrixH& InMatrix) const override;
 	};
 }
diff --git a/Engine/Source/Runtime/Datasmith/CADKernel/Public/CADKernel/Topo/Linkable.h b/Engine/Source/Runtime/Datasmith/CADKernel/Public/CADKernel/Topo/Linkable.h
--- a/Engine/Source/Runtime/Datasmith/CADKernel/Public/CADKernel/Topo/Linkable.h
+++ b/Engine/Source/Runtime/Datasmith/CADKernel/Public/CADKernel/Topo/Linkable.h
@@ -7,199 +7,161 @@
 
 namespace CADKernel
 {
-	class FModelMesh;
 
-	template<typename EntityType, typename LinkType>
-	class CADKERNEL_API TLinkable : public FTopologicalEntity
+class FModelMesh;
+
+template<typename EntityType, typename LinkType>
+class CADKERNEL_API TLinkable : public FTopologicalEntity
+{
+protected:
+	mutable TSharedPtr<LinkType> TopologicalLink;
+
+public:
+	TLinkable() = default;
+
+	void Finalize()
 	{
-	protected:
-		mutable TSharedPtr<LinkType> TopologicalLink;
+		ensureCADKernel(!TopologicalLink.IsValid());
+		TopologicalLink = FEntity::MakeShared<LinkType>((EntityType&)(*this));
+	}
 
-	public:
-		TLinkable() = default;
+	virtual void Serialize(FCADKernelArchive& Ar) override
+	{
+		FEntityGeom::Serialize(Ar);
+		SerializeIdent(Ar, TopologicalLink);
+	}
 
-		virtual void Serialize(FCADKernelArchive& Ar) override
-		{
-			FEntityGeom::Serialize(Ar);
-			SerializeIdent(Ar, TopologicalLink);
-		}
+	const TSharedRef<const EntityType> GetLinkActiveEntity() const
+	{
+		ensureCADKernel(TopologicalLink.IsValid());
+		EntityType* ActiveEntity = TopologicalLink->GetActiveEntity();
+		return StaticCastSharedRef<EntityType>(ActiveEntity->AsShared());
+	}
 
-		const TSharedRef<const EntityType> GetLinkActiveEntity() const
-		{
-			if (!TopologicalLink.IsValid())
-			{
-				return StaticCastSharedRef<const EntityType>(AsShared());
-			}
-			EntityType* ActiveEntity = TopologicalLink->GetActiveEntity();
-			return StaticCastSharedRef<EntityType>(ActiveEntity->AsShared());
-		}
+	TSharedRef<EntityType> GetLinkActiveEntity()
+	{
+		ensureCADKernel(TopologicalLink.IsValid());
+		EntityType* ActiveEntity = TopologicalLink->GetActiveEntity();
+		return StaticCastSharedRef<EntityType>(ActiveEntity->AsShared());
+	}
 
-		TSharedRef<EntityType> GetLinkActiveEntity()
+	bool IsActiveEntity() const
+	{
+		ensureCADKernel(TopologicalLink.IsValid());
+
+		if (TopologicalLink->GetTwinsEntitieNum() == 1)
 		{
-			if (!TopologicalLink.IsValid())
-			{
-				return StaticCastSharedRef<EntityType>(AsShared());
-			}
-			EntityType* ActiveEntity = TopologicalLink->GetActiveEntity();
-			return StaticCastSharedRef<EntityType>(ActiveEntity->AsShared());
+			return true;
 		}
 
-		bool IsActiveEntity() const
-		{
-			if (!TopologicalLink.IsValid())
-			{
-				return true;
-			}
+		return (TopologicalLink->GetActiveEntity() == this);
+	}
 
-			if (TopologicalLink->GetTwinsEntitieNum() == 1)
-			{
-				return true;
-			}
+	void Activate()
+	{
+		ensureCADKernel(TopologicalLink.IsValid());
+		TopologicalLink->ActivateEntity(*this);
+	}
 
-			return (TopologicalLink->GetActiveEntity() == this);
-		}
+	virtual TSharedPtr<LinkType> GetLink() const
+	{
+		ensureCADKernel(TopologicalLink.IsValid());
+		return TopologicalLink;
+	}
 
-		void Activate()
-		{
-			if (TopologicalLink.IsValid())
-			{
-				TopologicalLink->ActivateEntity(*this);
-			}
-		}
+	virtual TSharedPtr<LinkType> GetLink()
+	{
+		ensureCADKernel(TopologicalLink.IsValid());
+		return TopologicalLink;
+	}
 
-		virtual TSharedPtr<LinkType> GetLink() const
-		{
-			ensureCADKernel(TopologicalLink.IsValid());
-			return TopologicalLink;
-		}
+	void ResetTopologicalLink()
+	{
+		TopologicalLink.Reset();
+		TopologicalLink = FEntity::MakeShared<LinkType>((EntityType&)(*this));
+	}
 
-		virtual TSharedPtr<LinkType> GetLink()
+	bool IsLinkedTo(TSharedRef<EntityType> Entity) const
+	{
+		if (this == &*Entity)
 		{
-			if (!TopologicalLink.IsValid())
-			{
-				TopologicalLink = FEntity::MakeShared<LinkType>((EntityType&)(*this));
-			}
-			return TopologicalLink;
+			return true;
 		}
+		return (Entity->TopologicalLink == TopologicalLink);
+	}
 
-		void ResetTopologicalLink()
-		{
-			TopologicalLink = TSharedPtr<LinkType>();
-		}
+	int32 GetTwinsEntityCount() const
+	{
+		ensureCADKernel(TopologicalLink.IsValid());
+		return TopologicalLink->GetTwinsEntitieNum();
+	}
 
-		bool IsLinkedTo(TSharedRef<EntityType> Entity) const
-		{
-			if (this == &*Entity)
-			{
-				return true;
-			}
-			if(TopologicalLink.IsValid())
-			{
-				return (Entity->TopologicalLink == TopologicalLink);
-			}
-			return false;
-		}
+	bool HasTwin() const
+	{
+		return GetTwinsEntityCount() != 1;
+	}
 
-		int32 GetTwinsEntityCount() const
-		{
-			if (TopologicalLink.IsValid())
-			{
-				return TopologicalLink->GetTwinsEntitieNum();
-			}
-			return 1;
-		}
+	const TArray<EntityType*>& GetTwinsEntities() const
+	{
+		ensureCADKernel(TopologicalLink.IsValid());
+		return TopologicalLink->GetTwinsEntities();
+	}
 
-		bool HasTwin() const
-		{
-			return GetTwinsEntityCount() != 1;
-		}
+	virtual void RemoveFromLink()
+	{
+		ensureCADKernel(TopologicalLink.IsValid());
+		TopologicalLink->RemoveEntity((EntityType&)*this);
+		TopologicalLink.Reset();
+	}
 
-		const TArray<EntityType*>& GetTwinsEntities() const
-		{
-			if (!TopologicalLink.IsValid())
-			{
-				TopologicalLink = FEntity::MakeShared<LinkType>((EntityType&)(*this));
-			}
-			return TopologicalLink->GetTwinsEntities();
-		}
+	const bool IsThinZone() const
+	{
+		return ((States & EHaveStates::ThinZone) == EHaveStates::ThinZone);
+	}
 
-		virtual void RemoveFromLink()
-		{
-			if (TopologicalLink.IsValid())
-			{
-				TopologicalLink->RemoveEntity((EntityType&)*this);
-				TopologicalLink.Reset();
-			}
-		}
+	virtual void SetThinZone()
+	{
+		States |= EHaveStates::ThinZone;
+	}
 
-		const bool IsThinZone() const
-		{
-			return ((States & EHaveStates::ThinZone) == EHaveStates::ThinZone);
-		}
+	virtual void ResetThinZone()
+	{
+		States &= ~EHaveStates::ThinZone;
+	}
 
-		virtual void SetThinZone()
-		{
-			States |= EHaveStates::ThinZone;
-		}
+protected:
 
-		virtual void ResetThinZone()
+	void MakeLink(EntityType& Twin)
+	{
+		TSharedPtr<LinkType> Link1 = TopologicalLink;
+		TSharedPtr<LinkType> Link2 = Twin.TopologicalLink;
+
+		ensureCADKernel(Link1.IsValid() && Link2.IsValid());
+
+		if (Link1 == Link2)
 		{
-			States &= ~EHaveStates::ThinZone;
+			return;
 		}
 
-	protected:
-
-		void MakeLink(EntityType& Twin)
+		if (Link2->GetTwinsEntities().Num() > Link1->GetTwinsEntities().Num())
 		{
-			TSharedPtr<LinkType> Link1 = TopologicalLink;
-			TSharedPtr<LinkType> Link2 = Twin.TopologicalLink;
-
-			if (!Link1.IsValid() && !Link2.IsValid())
-			{
-				TopologicalLink = FEntity::MakeShared<LinkType>((EntityType&)(*this));
-				TopologicalLink->AddEntity(&Twin);
-				Twin.SetTopologicalLink(TopologicalLink);
-			}
-			else if (Link1.IsValid() && Link2.IsValid())
-			{
-				if (Link1 == Link2)
-				{
-					return;
-				}
-
-				if (Link2->GetTwinsEntities().Num() > Link1->GetTwinsEntities().Num())
-				{
-					Swap(Link1, Link2);
-				}
-
-				Link1->AddEntities(Link2->GetTwinsEntities());
-				for (EntityType* Entity : Link2->GetTwinsEntities())
-				{
-					Entity->SetTopologicalLink(Link1);
-				}
-				Link2->Delete();
-			}
-			else
-			{
-				if (Link1.IsValid())
-				{
-					Link1->AddEntity(&Twin);
-					Twin.SetTopologicalLink(Link1);
-				}
-				else
-				{
-					Link2->AddEntity(this);
-					SetTopologicalLink(Link2);
-				}
-			}
+			Swap(Link1, Link2);
 		}
 
-	protected:
-		void SetTopologicalLink(TSharedPtr<LinkType> Link)
+		Link1->AddEntities(Link2->GetTwinsEntities());
+		for (EntityType* Entity : Link2->GetTwinsEntities())
 		{
-			TopologicalLink = Link;
+			Entity->SetTopologicalLink(Link1);
 		}
-	};
+		Link2->Delete();
+	}
+
+protected:
+	void SetTopologicalLink(TSharedPtr<LinkType> Link)
+	{
+		TopologicalLink = Link;
+	}
+};
 
 } // namespace CADKernel
 
diff --git a/Engine/Source/Runtime/Datasmith/CADKernel/Public/CADKernel/Topo/Model.h b/Engine/Source/Runtime/Datasmith/CADKernel/Public/CADKernel/Topo/Model.h
--- a/Engine/Source/Runtime/Datasmith/CADKernel/Public/CADKernel/Topo/Model.h
+++ b/Engine/Source/Runtime/Datasmith/CADKernel/Public/CADKernel/Topo/Model.h
@@ -178,6 +178,8 @@ namespace CADKernel
 		 */
 		void CheckTopology();
 
+		virtual void FillTopologyReport(FTopologyReport& Report) const override;
+
 		/**
 		 * Fore each body
 		 */
diff --git a/Engine/Source/Runtime/Datasmith/CADKernel/Public/CADKernel/Topo/Shell.h b/Engine/Source/Runtime/Datasmith/CADKernel/Public/CADKernel/Topo/Shell.h
--- a/Engine/Source/Runtime/Datasmith/CADKernel/Public/CADKernel/Topo/Shell.h
+++ b/Engine/Source/Runtime/Datasmith/CADKernel/Public/CADKernel/Topo/Shell.h
@@ -143,6 +143,7 @@ namespace CADKernel
 		virtual void SpreadBodyOrientation() override;
 
 		void CheckTopology(TArray<FFaceSubset>& SubShells);
+		virtual void FillTopologyReport(FTopologyReport& Report) const override;
 
 		/**
 		 * @return true if the shell has at least one border edge 
diff --git a/Engine/Source/Runtime/Datasmith/CADKernel/Public/CADKernel/Topo/TopologicalEntity.h b/Engine/Source/Runtime/Datasmith/CADKernel/Public/CADKernel/Topo/TopologicalEntity.h
--- a/Engine/Source/Runtime/Datasmith/CADKernel/Public/CADKernel/Topo/TopologicalEntity.h
+++ b/Engine/Source/Runtime/Datasmith/CADKernel/Public/CADKernel/Topo/TopologicalEntity.h
@@ -7,6 +7,7 @@ namespace CADKernel
 {
 	class FModelMesh;
 	class FTopologicalFace;
+	class FTopologyReport;
 
 	class CADKERNEL_API FTopologicalEntity : public FEntityGeom
 	{
@@ -25,6 +26,8 @@ namespace CADKernel
 		 */
 		virtual void SpreadBodyOrientation() {}
 
+		virtual void FillTopologyReport(FTopologyReport& Report) const {}
+
 		virtual void Serialize(FCADKernelArchive& Ar) override
 		{
 			FEntityGeom::Serialize(Ar);
diff --git a/Engine/Source/Runtime/Datasmith/CADKernel/Public/CADKernel/Topo/TopologicalFace.h b/Engine/Source/Runtime/Datasmith/CADKernel/Public/CADKernel/Topo/TopologicalFace.h
--- a/Engine/Source/Runtime/Datasmith/CADKernel/Public/CADKernel/Topo/TopologicalFace.h
+++ b/Engine/Source/Runtime/Datasmith/CADKernel/Public/CADKernel/Topo/TopologicalFace.h
@@ -182,7 +182,7 @@ namespace CADKernel
 		/**
 		 * Trimmed the face with an outer boundary (first boundary of the array) and inners boundaries
 		 */
-		void AddLoops(const TArray<TSharedPtr<FTopologicalLoop>>& Loops);
+		void AddLoops(const TArray<TSharedPtr<FTopologicalLoop>>& Loops, int32& DoubtfulLoopOrientationCount);
 
 		/**
 		 * Trimmed the face with curves i.e. Edges will be build from the cures to make boundaries.
@@ -292,6 +292,8 @@ namespace CADKernel
 		 */
 		void RemoveLinksWithNeighbours();
 
+		virtual void FillTopologyReport(FTopologyReport& Report) const override;
+
 		// ======   Meshing Function   ======
 
 		TSharedRef<FFaceMesh> GetOrCreateMesh(const TSharedRef<FModelMesh>& ShellMesh);
diff --git a/Engine/Source/Runtime/Datasmith/CADKernel/Public/CADKernel/Topo/TopologicalLoop.h b/Engine/Source/Runtime/Datasmith/CADKernel/Public/CADKernel/Topo/TopologicalLoop.h
--- a/Engine/Source/Runtime/Datasmith/CADKernel/Public/CADKernel/Topo/TopologicalLoop.h
+++ b/Engine/Source/Runtime/Datasmith/CADKernel/Public/CADKernel/Topo/TopologicalLoop.h
@@ -149,7 +149,10 @@ namespace CADKernel
 			return Face.Pin().ToSharedRef();
 		}
 
-		void Orient();
+		/*
+		 * @return false if the orientation is doubtful
+		 */ 
+		bool Orient();
 		void SwapOrientation();
 
 		void SetAsInnerBoundary()
diff --git a/Engine/Source/Runtime/Datasmith/CADKernel/Public/CADKernel/Topo/TopologicalVertex.h b/Engine/Source/Runtime/Datasmith/CADKernel/Public/CADKernel/Topo/TopologicalVertex.h
--- a/Engine/Source/Runtime/Datasmith/CADKernel/Public/CADKernel/Topo/TopologicalVertex.h
+++ b/Engine/Source/Runtime/Datasmith/CADKernel/Public/CADKernel/Topo/TopologicalVertex.h
@@ -14,7 +14,6 @@ namespace CADKernel
 	class CADKERNEL_API FTopologicalVertex : public TLinkable<FTopologicalVertex, FVertexLink>
 	{
 		friend class FEntity;
-		friend class FTopologicalEdge;
 		friend class FVertexLink;
 
 	protected:
@@ -32,6 +31,13 @@ namespace CADKernel
 
 	public:
 
+		static TSharedRef<FTopologicalVertex> Make(const FPoint& InCoordinate)
+		{
+			TSharedRef<FTopologicalVertex> Vertex = FEntity::MakeShared<FTopologicalVertex>(InCoordinate);
+			Vertex->Finalize();
+			return Vertex;
+		}
+
 		virtual void Serialize(FCADKernelArchive& Ar) override
 		{
 			if(Ar.IsSaving())
diff --git a/Engine/Source/Runtime/Datasmith/CADKernel/Public/CADKernel/UI/Message.h b/Engine/Source/Runtime/Datasmith/CADKernel/Public/CADKernel/UI/Message.h
--- a/Engine/Source/Runtime/Datasmith/CADKernel/Public/CADKernel/UI/Message.h
+++ b/Engine/Source/Runtime/Datasmith/CADKernel/Public/CADKernel/UI/Message.h
@@ -2,81 +2,92 @@
 #pragma once
 
 #include "CADKernel/Core/Types.h"
+#include "CADKernel/Core/Chrono.h"
 
 #include "CADKernel/UI/Progress.h"
 
 namespace CADKernel
 {
-extern const char* VerboseLevelConstNames[];
-extern const char* VerboseConstDescHelp[];
+	extern const char* VerboseLevelConstNames[];
+	extern const char* VerboseConstDescHelp[];
 
 
-class CADKERNEL_API FMessage
-{
-	friend class FProgressManager;
+	class CADKERNEL_API FMessage
+	{
+		friend class FProgressManager;
 
-private:
-	static int32 NumberOfIndentation;
-	static int32 OldPercent;
+	private:
+		static int32 NumberOfIndentation;
+		static int32 OldPercent;
 
-	static void VPrintf(EVerboseLevel Level, const TCHAR* Text, ...);
+		static void VPrintf(EVerboseLevel Level, const TCHAR* Text, ...);
 
-#if defined(CADKERNEL_DEV) || defined(CADKERNEL_STDA)
-	static void VQaPrintF(const TCHAR* Header, const TCHAR* Text, ...);
-#endif
+		static void VReportPrintF(FString Header, const TCHAR* Text, ...);
 
-public:
+	public:
 
-	template <typename FmtType, typename... Types>
-	static void Printf(EVerboseLevel Level, const FmtType& Text, Types... Args)
-	{
-		static_assert(TAnd<TIsValidVariadicFunctionArg<Types>...>::Value, "Invalid argument(s) passed to FString::Printf");
-		VPrintf(Level, Text, Args...);
-	}
+		template <typename FmtType, typename... Types>
+		static void Printf(EVerboseLevel Level, const FmtType& Text, Types... Args)
+		{
+			static_assert(TAnd<TIsValidVariadicFunctionArg<Types>...>::Value, "Invalid argument(s) passed to FString::Printf");
+			VPrintf(Level, Text, Args...);
+		}
 
-	template <typename FmtType, typename... Types>
-	static void Error(const FmtType& Text, Types... Args)
-	{
-		static_assert(TAnd<TIsValidVariadicFunctionArg<Types>...>::Value, "Invalid argument(s) passed to FString::Printf");
-		FString CompleteText = TEXT("ERROR: ");
-		CompleteText += Text;
-		VPrintf(Log, *CompleteText, Args...);
-	}
-
-	template <typename FmtType, typename... Types>
-	static void Warning(const FmtType& Text, Types... Args)
-	{
-		static_assert(TAnd<TIsValidVariadicFunctionArg<Types>...>::Value, "Invalid argument(s) passed to FString::Printf");
-		FString CompleteText = TEXT("WARNING: ");
-		CompleteText += Text;
-		VPrintf(Log, *CompleteText, Args...);
-	}
-
-#if defined(CADKERNEL_DEV) || defined(CADKERNEL_STDA)
-	template <typename FmtType, typename... Types>
-	static void FillQaDataFile(const FmtType& Header, const FmtType& Text, Types... Args)
-	{
-		static_assert(TAnd<TIsValidVariadicFunctionArg<Types>...>::Value, "Invalid argument(s) passed to FString::Printf");
-		VQaPrintF(Header, Text, Args...);
-	}
+		template <typename FmtType, typename... Types>
+		static void Error(const FmtType& Text, Types... Args)
+		{
+			static_assert(TAnd<TIsValidVariadicFunctionArg<Types>...>::Value, "Invalid argument(s) passed to FString::Printf");
+			FString CompleteText = TEXT("ERROR: ");
+			CompleteText += Text;
+			VPrintf(Log, *CompleteText, Args...);
+		}
+
+		template <typename FmtType, typename... Types>
+		static void Warning(const FmtType& Text, Types... Args)
+		{
+			static_assert(TAnd<TIsValidVariadicFunctionArg<Types>...>::Value, "Invalid argument(s) passed to FString::Printf");
+			FString CompleteText = TEXT("WARNING: ");
+			CompleteText += Text;
+			VPrintf(Log, *CompleteText, Args...);
+		}
+
+		template <typename FmtType, typename... Types>
+		static void FillReportFile(const FmtType Header, const FmtType Text, Types... Args)
+		{
+			static_assert(TAnd<TIsValidVariadicFunctionArg<Types>...>::Value, "Invalid argument(s) passed to FString::Printf");
+			VReportPrintF(Header, Text, Args...);
+		}
+
+		static void FillReportFile(const TCHAR* Header, FDuration Duration)
+		{
+#ifdef CADKERNEL_DEV
+			int64 Time = FChrono::ConvertInto<std::chrono::microseconds>(Duration);
+#else
+			int64 Time = Duration;
 #endif
+			VReportPrintF(Header, TEXT("%lld"), Time);
+		}
 
-	static void Indent(int32 InNumberOfIndent = 1)
-	{
-		NumberOfIndentation += InNumberOfIndent;
-	}
+		static void FillReportFile(const TCHAR* Header, int32 Count)
+		{
+			VReportPrintF(Header, TEXT("%d"), Count);
+		}
 
-	static void Deindent(int32 InNumberOfIndent = 1)
-	{
-		NumberOfIndentation -= InNumberOfIndent;
-		if (NumberOfIndentation < 0)
+		static void Indent(int32 InNumberOfIndent = 1)
 		{
-			NumberOfIndentation = 0;
+			NumberOfIndentation += InNumberOfIndent;
 		}
-	}
 
-};
+		static void Deindent(int32 InNumberOfIndent = 1)
+		{
+			NumberOfIndentation -= InNumberOfIndent;
+			if (NumberOfIndentation < 0)
+			{
+				NumberOfIndentation = 0;
+			}
+		}
 
+	};
 } // namespace CADKernel
 
 #define ERROR_FUNCTION_CALL_NOT_EXPECTED \
