Update CADKernel

diff --git a/Engine/Plugins/Enterprise/DatasmithCADImporter/Source/CADInterfaces/Private/CoreTechFileParserImpl.cpp b/Engine/Plugins/Enterprise/DatasmithCADImporter/Source/CADInterfaces/Private/CoreTechFileParserImpl.cpp
--- a/Engine/Plugins/Enterprise/DatasmithCADImporter/Source/CADInterfaces/Private/CoreTechFileParserImpl.cpp
+++ b/Engine/Plugins/Enterprise/DatasmithCADImporter/Source/CADInterfaces/Private/CoreTechFileParserImpl.cpp
@@ -718,7 +718,8 @@ namespace CADLibrary
 		}
 
 		{
-			CADKernel::FSession CADKernelSession(0.00001 / CADFileData.GetImportParameters().GetMetricUnit());
+			double GeometricTolerance = 0.00001 / CADFileData.GetImportParameters().GetMetricUnit();
+			CADKernel::FSession CADKernelSession(GeometricTolerance);
 
 			TSharedRef<CADKernel::FModel> CADKernelModel = CADKernelSession.GetModel();
 
@@ -751,7 +752,7 @@ namespace CADLibrary
 			// Tessellate the body
 			TSharedRef<CADKernel::FModelMesh> CADKernelModelMesh = CADKernel::FEntity::MakeShared<CADKernel::FModelMesh>();
 
-			FCADKernelTools::DefineMeshCriteria(CADKernelModelMesh, CADFileData.GetImportParameters());
+			FCADKernelTools::DefineMeshCriteria(CADKernelModelMesh, CADFileData.GetImportParameters(), GeometricTolerance);
 
 			CADKernel::FParametricMesher Mesher(CADKernelModelMesh);
 			Mesher.MeshEntity(CADKernelModel);
@@ -800,7 +801,7 @@ namespace CADLibrary
 		// Tessellate the body
 		TSharedRef<CADKernel::FModelMesh> CADKernelModelMesh = CADKernel::FEntity::MakeShared<CADKernel::FModelMesh>();
 
-		FCADKernelTools::DefineMeshCriteria(CADKernelModelMesh, CADFileData.GetImportParameters());
+		FCADKernelTools::DefineMeshCriteria(CADKernelModelMesh, CADFileData.GetImportParameters(), CADKernelSession.GetGeometricTolerance());
 
 		CADKernel::FParametricMesher Mesher(CADKernelModelMesh);
 		Mesher.MeshEntity(CADKernelBody);
diff --git a/Engine/Plugins/Enterprise/DatasmithCADImporter/Source/CADLibrary/Private/CADKernelTools.cpp b/Engine/Plugins/Enterprise/DatasmithCADImporter/Source/CADLibrary/Private/CADKernelTools.cpp
--- a/Engine/Plugins/Enterprise/DatasmithCADImporter/Source/CADLibrary/Private/CADKernelTools.cpp
+++ b/Engine/Plugins/Enterprise/DatasmithCADImporter/Source/CADLibrary/Private/CADKernelTools.cpp
@@ -283,7 +283,7 @@ namespace CADLibrary
 		TSharedRef<FModelMesh> CADKernelModelMesh = FEntity::MakeShared<FModelMesh>();
 		FParametricMesher Mesher(CADKernelModelMesh);
 
-		DefineMeshCriteria(CADKernelModelMesh, ImportParameters);
+		DefineMeshCriteria(CADKernelModelMesh, ImportParameters, 0.01);
 		Mesher.MeshEntity(CADTopologicalEntity);
 
 		FMeshConversionContext Context(ImportParameters, MeshParameters);
@@ -383,16 +383,18 @@ namespace CADLibrary
 		}
 	}
 
-	void FCADKernelTools::DefineMeshCriteria(TSharedRef<CADKernel::FModelMesh>& MeshModel, const FImportParameters& ImportParameters)
+	void FCADKernelTools::DefineMeshCriteria(TSharedRef<CADKernel::FModelMesh>& MeshModel, const FImportParameters& ImportParameters, double GeometricTolerance)
 	{
 		{
 			TSharedPtr<CADKernel::FCriterion> CurvatureCriterion = CADKernel::FCriterion::CreateCriterion(CADKernel::ECriterion::CADCurvature);
 			MeshModel->AddCriterion(CurvatureCriterion);
+
+			TSharedPtr<CADKernel::FCriterion> MinSizeCriterion = CADKernel::FCriterion::CreateCriterion(CADKernel::ECriterion::MinSize, 2. * GeometricTolerance);
+			MeshModel->AddCriterion(MinSizeCriterion);
 		}
 
 		if (ImportParameters.GetMaxEdgeLength() > SMALL_NUMBER)
 		{
-
 			TSharedPtr<CADKernel::FCriterion> MaxSizeCriterion = CADKernel::FCriterion::CreateCriterion(CADKernel::ECriterion::MaxSize, ImportParameters.GetMaxEdgeLength() / ImportParameters.GetScaleFactor());
 			MeshModel->AddCriterion(MaxSizeCriterion);
 		}
diff --git a/Engine/Plugins/Enterprise/DatasmithCADImporter/Source/CADLibrary/Public/CADKernelTools.h b/Engine/Plugins/Enterprise/DatasmithCADImporter/Source/CADLibrary/Public/CADKernelTools.h
--- a/Engine/Plugins/Enterprise/DatasmithCADImporter/Source/CADLibrary/Public/CADKernelTools.h
+++ b/Engine/Plugins/Enterprise/DatasmithCADImporter/Source/CADLibrary/Public/CADKernelTools.h
@@ -25,7 +25,7 @@ namespace CADLibrary
 	class CADLIBRARY_API FCADKernelTools
 	{
 	public:
-		static void DefineMeshCriteria(TSharedRef<CADKernel::FModelMesh>& MeshModel, const FImportParameters& ImportParameters);
+		static void DefineMeshCriteria(TSharedRef<CADKernel::FModelMesh>& MeshModel, const FImportParameters& ImportParameters, double GeometricTolerance);
 
 		static void GetBodyTessellation(const TSharedRef<CADKernel::FModelMesh>& ModelMesh, const TSharedRef<CADKernel::FBody>& Body, FBodyMesh& OutBodyMesh, uint32 DefaultMaterialHash, TFunction<void(FObjectDisplayDataId, FObjectDisplayDataId, int32)> SetFaceMainMaterial);
 
diff --git a/Engine/Source/Runtime/Datasmith/CADKernel/Private/CADKernel/Mesh/Criteria/CriteriaGrid.cpp b/Engine/Source/Runtime/Datasmith/CADKernel/Private/CADKernel/Mesh/Criteria/CriteriaGrid.cpp
--- a/Engine/Source/Runtime/Datasmith/CADKernel/Private/CADKernel/Mesh/Criteria/CriteriaGrid.cpp
+++ b/Engine/Source/Runtime/Datasmith/CADKernel/Private/CADKernel/Mesh/Criteria/CriteriaGrid.cpp
@@ -68,21 +68,21 @@ void FCriteriaGrid::ApplyCriteria(const TArray<TSharedPtr<FCriterion>>& Criteria
 	{
 		for (int32 IndexU = 0; IndexU < GetCoordinateCount(EIso::IsoU) - 1; ++IndexU)
 		{
-			const FPoint& Point = GetPoint(IndexU, IndexV);
-			const FPoint& PointU = GetPoint(IndexU + 1, IndexV);
-			const FPoint& PointV = GetPoint(IndexU, IndexV + 1);
-			const FPoint& PointUV = GetPoint(IndexU + 1, IndexV + 1);
-			const FPoint& PointUMid = GetIntermediateU(IndexU, IndexV);
-			const FPoint& PointVMid = GetIntermediateV(IndexU, IndexV);
-			const FPoint& PointUVMid = GetIntermediateUV(IndexU, IndexV);
+			const FPoint& Point_U0_V0 = GetPoint(IndexU, IndexV);
+			const FPoint& Point_U1_V1 = GetPoint(IndexU + 1, IndexV + 1);
+			const FPoint& Point_Um_V0 = GetIntermediateU(IndexU, IndexV);
+			const FPoint& Point_Um_V1 = GetIntermediateU(IndexU, IndexV+1);
+			const FPoint& Point_U0_Vm = GetIntermediateV(IndexU, IndexV);
+			const FPoint& Point_U1_Vm = GetIntermediateV(IndexU+1, IndexV);
+			const FPoint& Point_Um_Vm = GetIntermediateUV(IndexU, IndexV);
 
 			// Evaluate Sag
 			double LengthU;
-			double SagU = FCriterion::EvaluateSag(Point, PointU, PointUMid, LengthU);
+			double SagU = FCriterion::EvaluateSag(Point_U0_Vm, Point_U1_Vm, Point_Um_Vm, LengthU);
 			double LengthV;
-			double SagV = FCriterion::EvaluateSag(Point, PointV, PointVMid, LengthV);
+			double SagV = FCriterion::EvaluateSag(Point_Um_V0, Point_Um_V1, Point_Um_Vm, LengthV);
 			double LengthUV;
-			double SagUV = FCriterion::EvaluateSag(Point, PointUV, PointUVMid, LengthUV);
+			double SagUV = FCriterion::EvaluateSag(Point_U0_V0, Point_U1_V1, Point_Um_Vm, LengthUV);
 
 			for (const TSharedPtr<FCriterion>& Criterion : Criteria)
 			{
@@ -106,99 +106,107 @@ void FCriteriaGrid::ApplyCriteria(const TArray<TSharedPtr<FCriterion>>& Criteria
 	}
 }
 
-void FCriteriaGrid::Display()
+void FCriteriaGrid::Display() const
 {
-	Open3DDebugSession(TEXT("Grid"));
+	F3DDebugSession _(TEXT("Grid"));
 
-	Open3DDebugSession(TEXT("CriteriaGrid Point 3d"));
-	for (int32 VIndex = 0; VIndex < GetCoordinateCount(EIso::IsoV); ++VIndex)
 	{
-		for (int32 UIndex = 0; UIndex < GetCoordinateCount(EIso::IsoU); ++UIndex)
+		F3DDebugSession A(TEXT("CriteriaGrid Point 3d"));
+		for (int32 VIndex = 0; VIndex < GetCoordinateCount(EIso::IsoV); ++VIndex)
 		{
-			CADKernel::DisplayPoint(GetPoint(UIndex, VIndex), UIndex);
+			for (int32 UIndex = 0; UIndex < GetCoordinateCount(EIso::IsoU); ++UIndex)
+			{
+				CADKernel::DisplayPoint(GetPoint(UIndex, VIndex), UIndex);
+			}
 		}
 	}
-	Close3DDebugSession();
 
-	Open3DDebugSession(TEXT("CriteriaGrid IntermediateU"));
-	for (int32 VIndex = 0; VIndex < GetCoordinateCount(EIso::IsoV); ++VIndex)
 	{
-		for (int32 UIndex = 0; UIndex < GetCoordinateCount(EIso::IsoU) - 1; ++UIndex)
+		F3DDebugSession A(TEXT("CriteriaGrid IntermediateU"));
+		for (int32 VIndex = 0; VIndex < GetCoordinateCount(EIso::IsoV); ++VIndex)
 		{
-			CADKernel::DisplayPoint(GetIntermediateU(UIndex, VIndex), EVisuProperty::ControlPoint);
+			for (int32 UIndex = 0; UIndex < GetCoordinateCount(EIso::IsoU) - 1; ++UIndex)
+			{
+				CADKernel::DisplayPoint(GetIntermediateU(UIndex, VIndex), EVisuProperty::PurplePoint);
+			}
 		}
 	}
-	Close3DDebugSession();
 
-	Open3DDebugSession(TEXT("CriteriaGrid IntermediateV"));
-	for (int32 VIndex = 0; VIndex < GetCoordinateCount(EIso::IsoV) - 1; ++VIndex)
 	{
-		for (int32 UIndex = 0; UIndex < GetCoordinateCount(EIso::IsoU); ++UIndex)
+		F3DDebugSession A(TEXT("CriteriaGrid IntermediateV"));
+		for (int32 VIndex = 0; VIndex < GetCoordinateCount(EIso::IsoV) - 1; ++VIndex)
 		{
-			CADKernel::DisplayPoint(GetIntermediateV(UIndex, VIndex), EVisuProperty::ControlPoint);
+			for (int32 UIndex = 0; UIndex < GetCoordinateCount(EIso::IsoU); ++UIndex)
+			{
+				CADKernel::DisplayPoint(GetIntermediateV(UIndex, VIndex), EVisuProperty::PurplePoint);
+			}
 		}
 	}
-	Close3DDebugSession();
 
-	Open3DDebugSession(TEXT("CriteriaGrid IntermediateUV"));
-	for (int32 VIndex = 0; VIndex < GetCoordinateCount(EIso::IsoV) - 1; ++VIndex)
 	{
-		for (int32 UIndex = 0; UIndex < GetCoordinateCount(EIso::IsoU) - 1; ++UIndex)
+		F3DDebugSession A(TEXT("CriteriaGrid IntermediateUV"));
+		for (int32 VIndex = 0; VIndex < GetCoordinateCount(EIso::IsoV) - 1; ++VIndex)
 		{
-			CADKernel::DisplayPoint(GetIntermediateUV(UIndex, VIndex), EVisuProperty::ControlPoint);
+			for (int32 UIndex = 0; UIndex < GetCoordinateCount(EIso::IsoU) - 1; ++UIndex)
+			{
+				CADKernel::DisplayPoint(GetIntermediateUV(UIndex, VIndex), EVisuProperty::PurplePoint);
+			}
 		}
 	}
-	Close3DDebugSession();
 
-	Open3DDebugSession(TEXT("Loop 3D"));
-	for (const TSharedPtr<FTopologicalLoop>& Loop : Surface->GetLoops())
 	{
-		CADKernel::Display(*Loop);
+		F3DDebugSession A(TEXT("Loop 3D"));
+		for (const TSharedPtr<FTopologicalLoop>& Loop : Surface->GetLoops())
+		{
+			CADKernel::Display(*Loop);
+		}
 	}
-	Close3DDebugSession();
 
-	Open3DDebugSession(TEXT("Loop 2D"));
-	for (const TSharedPtr<FTopologicalLoop>& Loop : Surface->GetLoops())
 	{
-		CADKernel::Display2D(*Loop);
+		F3DDebugSession A(TEXT("Loop 2D"));
+		for (const TSharedPtr<FTopologicalLoop>& Loop : Surface->GetLoops())
+		{
+			CADKernel::Display2D(*Loop);
+		}
 	}
-	Close3DDebugSession();
 
-	Open3DDebugSession(TEXT("CriteriaGrid Point 2D"));
-	for (int32 VIndex = 0; VIndex < GetCoordinateCount(EIso::IsoV); ++VIndex)
 	{
-		for (int32 UIndex = 0; UIndex < GetCoordinateCount(EIso::IsoU); ++UIndex)
+		F3DDebugSession A(TEXT("CriteriaGrid Point 2D"));
+		for (int32 VIndex = 0; VIndex < GetCoordinateCount(EIso::IsoV); ++VIndex)
 		{
-			CADKernel::DisplayPoint(FPoint2D(GetCoordinate(EIso::IsoU, UIndex), GetCoordinate(EIso::IsoV, VIndex)));
+			for (int32 UIndex = 0; UIndex < GetCoordinateCount(EIso::IsoU); ++UIndex)
+			{
+				CADKernel::DisplayPoint(FPoint2D(GetCoordinate(EIso::IsoU, UIndex), GetCoordinate(EIso::IsoV, VIndex)));
+			}
 		}
 	}
-	Close3DDebugSession();
 
-	Open3DDebugSession(TEXT("CriteriaGrid Point 2D Intermediate"));
-	for (int32 VIndex = 0; VIndex < GetCoordinateCount(EIso::IsoV); ++VIndex)
 	{
-		for (int32 UIndex = 1; UIndex < GetCoordinateCount(EIso::IsoU); ++UIndex)
+		F3DDebugSession A(TEXT("CriteriaGrid Point 2D Intermediate"));
+		for (int32 VIndex = 0; VIndex < GetCoordinateCount(EIso::IsoV); ++VIndex)
 		{
-			CADKernel::DisplayPoint(FPoint2D((GetCoordinate(EIso::IsoU, UIndex) + GetCoordinate(EIso::IsoU, UIndex - 1))*0.5, GetCoordinate(EIso::IsoV, VIndex)), EVisuProperty::ControlPoint);
+			for (int32 UIndex = 1; UIndex < GetCoordinateCount(EIso::IsoU); ++UIndex)
+			{
+				CADKernel::DisplayPoint(FPoint2D((GetCoordinate(EIso::IsoU, UIndex) + GetCoordinate(EIso::IsoU, UIndex - 1)) * 0.5, GetCoordinate(EIso::IsoV, VIndex)), EVisuProperty::PurplePoint);
+			}
 		}
-	}
-	for (int32 VIndex = 1; VIndex < GetCoordinateCount(EIso::IsoV); ++VIndex)
-	{
-		for (int32 UIndex = 0; UIndex < GetCoordinateCount(EIso::IsoU); ++UIndex)
+		for (int32 VIndex = 1; VIndex < GetCoordinateCount(EIso::IsoV); ++VIndex)
 		{
-			CADKernel::DisplayPoint(FPoint2D(GetCoordinate(EIso::IsoU, UIndex), (GetCoordinate(EIso::IsoV, VIndex) + GetCoordinate(EIso::IsoV, VIndex - 1))*0.5), EVisuProperty::ControlPoint);
+			for (int32 UIndex = 0; UIndex < GetCoordinateCount(EIso::IsoU); ++UIndex)
+			{
+				CADKernel::DisplayPoint(FPoint2D(GetCoordinate(EIso::IsoU, UIndex), (GetCoordinate(EIso::IsoV, VIndex) + GetCoordinate(EIso::IsoV, VIndex - 1)) * 0.5), EVisuProperty::PurplePoint);
+			}
 		}
-	}
-	for (int32 VIndex = 1; VIndex < GetCoordinateCount(EIso::IsoV); ++VIndex)
-	{
-		for (int32 UIndex = 1; UIndex < GetCoordinateCount(EIso::IsoU); ++UIndex)
+		for (int32 VIndex = 1; VIndex < GetCoordinateCount(EIso::IsoV); ++VIndex)
 		{
-			CADKernel::DisplayPoint(FPoint2D((GetCoordinate(EIso::IsoU, UIndex) + GetCoordinate(EIso::IsoU, UIndex - 1))*0.5, (GetCoordinate(EIso::IsoV, VIndex) + GetCoordinate(EIso::IsoV, VIndex - 1))*0.5), EVisuProperty::ControlPoint);
+			for (int32 UIndex = 1; UIndex < GetCoordinateCount(EIso::IsoU); ++UIndex)
+			{
+				CADKernel::DisplayPoint(FPoint2D((GetCoordinate(EIso::IsoU, UIndex) + GetCoordinate(EIso::IsoU, UIndex - 1)) * 0.5, (GetCoordinate(EIso::IsoV, VIndex) + GetCoordinate(EIso::IsoV, VIndex - 1)) * 0.5), EVisuProperty::PurplePoint);
+			}
 		}
 	}
-	Close3DDebugSession();
 
-	Close3DDebugSession();
+	Wait();
 }
 
 } // namespace CADKernel
diff --git a/Engine/Source/Runtime/Datasmith/CADKernel/Private/CADKernel/Mesh/Criteria/CriteriaGrid.h b/Engine/Source/Runtime/Datasmith/CADKernel/Private/CADKernel/Mesh/Criteria/CriteriaGrid.h
--- a/Engine/Source/Runtime/Datasmith/CADKernel/Private/CADKernel/Mesh/Criteria/CriteriaGrid.h
+++ b/Engine/Source/Runtime/Datasmith/CADKernel/Private/CADKernel/Mesh/Criteria/CriteriaGrid.h
@@ -93,7 +93,7 @@ namespace CADKernel
 			return GetPoint(iU, iV, true, true, normal);
 		}
 
-		void Display();
+		void Display() const;
 	};
 }
 
