CADKernel::StringUtil move

diff --git a/Engine/Source/Runtime/Datasmith/CADKernel/Private/CADKernel/Utils/StringUtil.h b/Engine/Source/Runtime/Datasmith/CADKernel/Private/CADKernel/Utils/StringUtil.h
deleted file mode 100644
diff --git a/Engine/Source/Runtime/Datasmith/CADKernel/Public/CADKernel/Utils/StringUtil.h b/Engine/Source/Runtime/Datasmith/CADKernel/Public/CADKernel/Utils/StringUtil.h
new file mode 100644
--- /dev/null
+++ b/Engine/Source/Runtime/Datasmith/CADKernel/Public/CADKernel/Utils/StringUtil.h
@@ -0,0 +1,106 @@
+// Copyright Epic Games, Inc. All Rights Reserved.
+#pragma once
+
+#include "CADKernel/Core/Types.h"
+
+#include "CADKernel/Core/Types.h"
+#include "CADKernel/Math/MathConst.h"
+#include "CADKernel/Math/Point.h"
+
+#include <list>
+#include <locale>
+#include <string>
+#include <sstream>
+#include <stdio.h>
+#include <stdarg.h>
+
+#ifdef DEV_CADKERNEL_LIB
+#include "CADKernel/Utils/UnrealEnvironment.h"
+#endif
+
+namespace CADKernel
+{
+	namespace Utils
+	{
+		inline std::string ToString(const TCHAR* WideString)
+		{
+			return TCHAR_TO_UTF8(WideString);
+		}
+
+		template<typename ValueType>
+		inline std::string ToString(const ValueType& v)
+		{
+			std::stringstream stream;
+			stream << v;
+			return stream.str();
+		}
+
+		inline std::string ToString(const FPoint& p)
+		{
+			std::string Message = "dm.FPoint(";
+			Message += Utils::ToString(p.X) + "," + Utils::ToString(p.Y) + "," + Utils::ToString(p.Z) + ")";
+			return Message;
+		}
+
+		inline FString ToFString(const char* Text)
+		{
+			return UTF8_TO_TCHAR(Text);
+		}
+
+		inline FString ToFString(const double& Value)
+		{
+			return FString::Printf(TEXT("%f"), Value);
+		}
+
+		inline FString ToFString(const int32& Value)
+		{
+			return FString::Printf(TEXT("%d"), Value);
+		}
+
+		inline FString ToFString(const uint32& Value)
+		{
+			return FString::Printf(TEXT("%u"), Value);
+		}
+
+		inline FString ToFString(const bool& Value)
+		{
+			return FString::Printf(TEXT("%d"), Value);
+		}
+
+		inline FString ToFString(const FPoint& p)
+		{
+			return FString::Printf(TEXT("dm.FPoint(%f, %f, %f)"), p.X, p.Y, p.Z);
+		}
+
+		void Explode(const FString& FullString, const TCHAR* Separator, TArray<FString>& StringArray);
+
+		inline FString EscapeBackSlashes(const FString& InputString)
+		{
+			FString ret;
+			ret.Reserve(InputString.Len());
+			for (TCHAR Char : InputString)
+			{
+				if (Char == TEXT('\\')) {
+					ret.AppendChar(TEXT('\\'));
+				}
+				ret.AppendChar(Char);
+			}
+			return ret;
+		}
+
+		inline void RemoveUnwantedChar(FString& String, const TCHAR UnwantedChar)
+		{
+			FString NewString;
+			NewString.Reserve(String.Len());
+			for (const TCHAR& Char : String)
+			{
+				if (Char != UnwantedChar)
+				{
+					NewString.AppendChar(Char);
+				}
+			}
+			Move(String, NewString);
+		}
+
+	};
+}
