CADKernel Update : Fix Sew

diff --git a/Engine/Source/Runtime/Datasmith/CADKernel/Private/CADKernel/Core/Database.cpp b/Engine/Source/Runtime/Datasmith/CADKernel/Private/CADKernel/Core/Database.cpp
--- a/Engine/Source/Runtime/Datasmith/CADKernel/Private/CADKernel/Core/Database.cpp
+++ b/Engine/Source/Runtime/Datasmith/CADKernel/Private/CADKernel/Core/Database.cpp
@@ -125,6 +125,11 @@ void FDatabase::GetEntities(const TArray<FIdent>& EntityIds, TArray<TSharedPtr<F
 	for (FIdent EntityId : EntityIds)
 	{
 		TSharedPtr<FEntity> Entity = GetEntity(EntityId);
+		if (!Entity.IsValid() || Entity->IsDeleted())
+		{
+			continue;
+		}
+
 		if (AddedComponents.Find(EntityId) == nullptr)
 		{
 			AddedComponents.Add(EntityId);
@@ -182,6 +187,11 @@ void FDatabase::SerializeSelection(FCADKernelArchive& Ar, const TArray<FIdent>&
 	for (FIdent EntityId : SelectionIds)
 	{
 		TSharedPtr<FEntity> Entity = GetEntity(EntityId);
+		if (!Entity.IsValid() || Entity->IsDeleted())
+		{
+			continue;
+		}
+
 		switch(Entity->GetEntityType())
 		{
 		case EEntity::Model:
@@ -204,7 +214,7 @@ void FDatabase::SerializeSelection(FCADKernelArchive& Ar, const TArray<FIdent>&
 	while(NotYetSerialized.Dequeue(EntityIdToSave))
 	{
 		TSharedPtr<FEntity> Entity = GetEntity(EntityIdToSave);
-		if (!Entity.IsValid())
+		if (!Entity.IsValid() || Entity->IsDeleted())
 		{
 			continue;
 		}
@@ -243,14 +253,7 @@ void FDatabase::Serialize(FCADKernelArchive& Ar)
 	{
 		Progess.Increase();
 
-		if (!Entity.IsValid())
-		{
-			EEntity Type = EEntity::NullEntity;
-			Ar << Type;
-			continue;
-		}
-
-		if (Entity->IsDeleted())
+		if (!Entity.IsValid() || Entity->IsDeleted())
 		{
 			EEntity Type = EEntity::NullEntity;
 			Ar << Type;
@@ -309,7 +312,7 @@ void FDatabase::CleanArchiveEntities()
 	int32 ShellCount = 0;
 	for (TSharedPtr<FEntity> Entity : ArchiveEntities)
 	{
-		if (Entity.IsValid())
+		if (Entity.IsValid() && !Entity->IsDeleted())
 		{
 			switch (Entity->GetEntityType())
 			{
@@ -368,7 +371,7 @@ void FDatabase::CleanArchiveEntities()
 	// find independent faces and shells
 	for (TSharedPtr<FEntity> Entity : ArchiveEntities)
 	{
-		if (Entity.IsValid())
+		if (Entity.IsValid() && !Entity->IsDeleted())
 		{
 			switch (Entity->GetEntityType())
 			{
@@ -425,6 +428,11 @@ uint32 FDatabase::SpawnEntityIdent(const TArray<TSharedPtr<FEntity>>& SelectedEn
 	EntityCount = 0;
 	for (TSharedPtr<FEntity> Entity : SelectedEntities)
 	{
+		if (!Entity.IsValid() || Entity->IsDeleted())
+		{
+			continue;
+		}
+
 		Entity->SpawnIdent(*this);
 	}
 	bForceSpawning = false;
@@ -433,6 +441,11 @@ uint32 FDatabase::SpawnEntityIdent(const TArray<TSharedPtr<FEntity>>& SelectedEn
 
 uint32 FDatabase::SpawnEntityIdent(TSharedPtr<FEntity>& Entity, bool bInForceSpawning)
 {
+	if (!Entity.IsValid() || Entity->IsDeleted())
+	{
+		return 0;
+	}
+
 	bForceSpawning = bInForceSpawning;
 	EntityCount = 0;
 	Entity->SpawnIdent(*this);
@@ -459,6 +472,11 @@ void FDatabase::ExpandSelection(const TArray<TSharedPtr<FEntity>>& Entities, con
 	ExpandSelection(Entities, Filters, EntitySet);
 	for (TSharedPtr<FEntity> Entity : EntitySet)
 	{
+		if (!Entity.IsValid() || Entity->IsDeleted())
+		{
+			continue;
+		}
+
 		Selection.Add(Entity);
 	}
 }
@@ -467,12 +485,22 @@ void FDatabase::ExpandSelection(const TArray<TSharedPtr<FEntity>>& Entities, con
 {
 	for (TSharedPtr<FEntity> Entity : Entities)
 	{
+		if (!Entity.IsValid() || Entity->IsDeleted())
+		{
+			continue;
+		}
+
 		ExpandSelection(Entity, Filter, Selection);
 	}
 }
 
 void FDatabase::ExpandSelection(TSharedPtr<FEntity> Entity, const TSet<EEntity>& Filter, TSet<TSharedPtr<FEntity>>& Selection) const
 {
+	if (!Entity.IsValid() || Entity->IsDeleted())
+	{
+		return;
+	}
+
 	EEntity Type = Entity->GetEntityType();
 
 	if (Filter.Contains(Type))
@@ -594,7 +622,7 @@ void FDatabase::GetEntitiesOfTypes(const TSet<EEntity>& Filter, TArray<TSharedPt
 	OutEntities.Empty(100);
 	for (TSharedPtr<FEntity> Entity : DatabaseEntities)
 	{
-		if (!Entity.IsValid())
+		if (!Entity.IsValid() || Entity->IsDeleted())
 		{
 			continue;
 		}
@@ -610,7 +638,7 @@ TSharedPtr<FEntity> FDatabase::GetFirstEntityOfType(EEntity Type) const
 {
 	for (TSharedPtr<FEntity> Entity : DatabaseEntities)
 	{
-		if (!Entity.IsValid())
+		if (!Entity.IsValid() || Entity->IsDeleted())
 		{
 			continue;
 		}
diff --git a/Engine/Source/Runtime/Datasmith/CADKernel/Private/CADKernel/Core/System.cpp b/Engine/Source/Runtime/Datasmith/CADKernel/Private/CADKernel/Core/System.cpp
--- a/Engine/Source/Runtime/Datasmith/CADKernel/Private/CADKernel/Core/System.cpp
+++ b/Engine/Source/Runtime/Datasmith/CADKernel/Private/CADKernel/Core/System.cpp
@@ -137,12 +137,12 @@ void FSystem::InitializeCADKernel()
 	FSystem::Get().SetVerboseLevel(EVerboseLevel::Log);
 }
 
-FString FSystem::GetToolkitVersion() const
+FString FSystem::GetToolkitVersion()
 {
 	return UTF8_TO_TCHAR(TOOLKIT_VERSION_ASCII);
 }
 
-FString FSystem::GetCompilationDate() const
+FString FSystem::GetCompilationDate()
 {
 	return UTF8_TO_TCHAR(__DATE__);
 }
@@ -151,7 +151,8 @@ void FSystem::PrintHeader()
 {
 	FMessage::Printf(Log, TEXT("_______________________________________________________________________________\n"));
 	FMessage::Printf(Log, TEXT("\n"));
-	FMessage::Printf(Log, TEXT("\tDatasmith CAD Kernel Toolkit release %s (%s)\n"),* GetToolkitVersion(),* GetCompilationDate());
+
+	FMessage::Printf(Log, TEXT("\tDatasmith CAD Kernel Toolkit release %s (%s)\n"), UTF8_TO_TCHAR(TOOLKIT_VERSION_ASCII), UTF8_TO_TCHAR(__DATE__));
 	FMessage::Printf(Log, TEXT("\t" EPIC_COPYRIGHT "\n"));
 	FMessage::Printf(Log, TEXT("\n"));
 	FMessage::Printf(Log, TEXT("_______________________________________________________________________________\n"));
diff --git a/Engine/Source/Runtime/Datasmith/CADKernel/Private/CADKernel/Core/Version.h b/Engine/Source/Runtime/Datasmith/CADKernel/Private/CADKernel/Core/Version.h
--- a/Engine/Source/Runtime/Datasmith/CADKernel/Private/CADKernel/Core/Version.h
+++ b/Engine/Source/Runtime/Datasmith/CADKernel/Private/CADKernel/Core/Version.h
@@ -2,8 +2,8 @@
 #pragma once
 
 //** ** Version du Toolkit** **//
-#define TOOLKIT_VERSION 5, 0, 18632741
-#define TOOLKIT_VERSION_ASCII "5.0.18632741"
+#define TOOLKIT_VERSION 5, 0, 18687460
+#define TOOLKIT_VERSION_ASCII "5.0.18687460"
 
 #define RELEASE_DATE	 __DATE__
 
diff --git a/Engine/Source/Runtime/Datasmith/CADKernel/Private/CADKernel/Mesh/Meshers/IsoTriangulator.cpp b/Engine/Source/Runtime/Datasmith/CADKernel/Private/CADKernel/Mesh/Meshers/IsoTriangulator.cpp
--- a/Engine/Source/Runtime/Datasmith/CADKernel/Private/CADKernel/Mesh/Meshers/IsoTriangulator.cpp
+++ b/Engine/Source/Runtime/Datasmith/CADKernel/Private/CADKernel/Mesh/Meshers/IsoTriangulator.cpp
@@ -31,7 +31,7 @@ FIsoTriangulator::FIsoTriangulator(FGrid& InGrid, TSharedRef<FFaceMesh> EntityMe
 	, InnerSegmentsIntersectionTool(InGrid)
 	, InnerToLoopSegmentsIntersectionTool(InGrid)
 	, InnerToOuterSegmentsIntersectionTool(InGrid)
-	, GeometricTolerance(Grid.GetFace()->GetCarrierSurface()->Get3DTolerance())
+	, GeometricTolerance(InGrid.GetFace().GetCarrierSurface()->Get3DTolerance())
 	, SquareGeometricTolerance(FMath::Square(GeometricTolerance))
 	, SquareGeometricTolerance2(2. * SquareGeometricTolerance)
 	, MeshingTolerance(GeometricTolerance* GeometricToMeshingToleranceFactor)
@@ -44,7 +44,7 @@ FIsoTriangulator::FIsoTriangulator(FGrid& InGrid, TSharedRef<FFaceMesh> EntityMe
 #ifdef DEBUG_ISOTRIANGULATOR
 
 #ifdef DEBUG_ONLY_SURFACE_TO_DEBUG
-	if (InGrid.GetFace()->GetId() == FaceToDebug)
+	if (Grid.GetFace().GetId() == FaceToDebug)
 #endif
 	{
 		bDisplay = true;
@@ -61,7 +61,7 @@ bool FIsoTriangulator::Triangulate()
 {
 	EGridSpace DisplaySpace = EGridSpace::UniformScaled;
 #ifdef DEBUG_ISOTRIANGULATOR
-	F3DDebugSession _(bDisplay, FString::Printf(TEXT("Triangulate %d"), Grid.GetFace()->GetId()));
+	F3DDebugSession _(bDisplay, FString::Printf(TEXT("Triangulate %d"), Grid.GetFace().GetId()));
 #endif
 
 	FTimePoint StartTime = FChrono::Now();
@@ -87,7 +87,7 @@ bool FIsoTriangulator::Triangulate()
 	if (!LoopCleaner.Run())
 	{
 		MesherReport.Logs.AddDegeneratedLoop();
-		FMessage::Printf(EVerboseLevel::Log, TEXT("The meshing of the surface %d failed due to a degenerated loop\n"), Grid.GetFace()->GetId());
+		FMessage::Printf(EVerboseLevel::Log, TEXT("The meshing of the surface %d failed due to a degenerated loop\n"), Grid.GetFace().GetId());
 		return false;
 	}
 
@@ -406,7 +406,7 @@ void FIsoTriangulator::BuildThinZoneSegments()
 		ThinZoneSegments.Add(&Segment);
 	};
 
-	for (const TSharedPtr<FTopologicalLoop>& Loop : Grid.GetFace()->GetLoops())
+	for (const TSharedPtr<FTopologicalLoop>& Loop : Grid.GetFace().GetLoops())
 	{
 		for (const FOrientedEdge& OrientedEdge : Loop->GetEdges())
 		{
@@ -1566,7 +1566,7 @@ bool FIsoTriangulator::CanCycleBeMeshed(const TArray<FIsoSegment*>& Cycle, FInte
 	{
 		if (CycleIntersectionTool.DoesIntersect(*Segment))
 		{
-			FMessage::Printf(Log, TEXT("A cycle of the surface %d is in self intersecting. The mesh of this sector is canceled.\n"), Grid.GetFace()->GetId());
+			FMessage::Printf(Log, TEXT("A cycle of the surface %d is in self intersecting. The mesh of this sector is canceled.\n"), Grid.GetFace().GetId());
 
 #ifdef DEBUG_CAN_CYCLE_BE_FIXED_AND_MESHED
 			if (bDisplay)
diff --git a/Engine/Source/Runtime/Datasmith/CADKernel/Private/CADKernel/Mesh/Meshers/IsoTriangulator/IsoCell.h b/Engine/Source/Runtime/Datasmith/CADKernel/Private/CADKernel/Mesh/Meshers/IsoTriangulator/IsoCell.h
--- a/Engine/Source/Runtime/Datasmith/CADKernel/Private/CADKernel/Mesh/Meshers/IsoTriangulator/IsoCell.h
+++ b/Engine/Source/Runtime/Datasmith/CADKernel/Private/CADKernel/Mesh/Meshers/IsoTriangulator/IsoCell.h
@@ -41,7 +41,7 @@ namespace CADKernel
 		{
 			ensureCADKernel(InNodes.Num() > 0);
 
-			int32 FaceLoopCount = Grid.GetFace()->LoopCount();
+			int32 FaceLoopCount = Grid.GetFace().LoopCount();
 			int32 LoopCount = 0;
 			{
 				TArray<int32> LoopVertexCount;
diff --git a/Engine/Source/Runtime/Datasmith/CADKernel/Private/CADKernel/Mesh/Meshers/ParametricMesher.cpp b/Engine/Source/Runtime/Datasmith/CADKernel/Private/CADKernel/Mesh/Meshers/ParametricMesher.cpp
--- a/Engine/Source/Runtime/Datasmith/CADKernel/Private/CADKernel/Mesh/Meshers/ParametricMesher.cpp
+++ b/Engine/Source/Runtime/Datasmith/CADKernel/Private/CADKernel/Mesh/Meshers/ParametricMesher.cpp
@@ -137,7 +137,7 @@ void FParametricMesher::MeshEntities()
 
 	if (Faces.Num() > 1)
 	{
-		TArray<TSharedPtr<FTopologicalFace>> OtherEntities;
+		TArray<FTopologicalFace*> OtherEntities;
 
 		FMessage::Printf(Log, TEXT("  Isolate QuadPatch\n"));
 		FTimePoint IsolateQuadPatchStartTime = FChrono::Now();
@@ -206,21 +206,21 @@ void FParametricMesher::ApplyEdgeCriteria(FTopologicalEdge& Edge)
 	ActiveEdge.SetApplyCriteria();
 }
 
-void FParametricMesher::Mesh(TSharedRef<FTopologicalFace> Face)
+void FParametricMesher::Mesh(FTopologicalFace& Face)
 {
 #ifdef DEBUG_ONLY_SURFACE_TO_DEBUG
-	bDisplay = (Face->GetId() == FaceToDebug);
+	bDisplay = (Face.GetId() == FaceToDebug);
 #endif
 
-	ensureCADKernel(!Face->IsDeleted());
-	ensureCADKernel(!Face->IsMeshed());
+	ensureCADKernel(!Face.IsDeleted());
+	ensureCADKernel(!Face.IsMeshed());
 
-	FMessage::Printf(EVerboseLevel::Debug, TEXT("Meshing of surface %d\n"), Face->GetId());
+	FMessage::Printf(EVerboseLevel::Debug, TEXT("Meshing of surface %d\n"), Face.GetId());
 
 	FProgress _(1, TEXT("Meshing Entities : Mesh Surface"));
 
 #ifdef DEBUG_CADKERNEL
-	F3DDebugSession S(bDisplayDebugMeshStep, FString::Printf(TEXT("Mesh of surface %d"), Face->GetId()));
+	F3DDebugSession S(bDisplayDebugMeshStep, FString::Printf(TEXT("Mesh of surface %d"), Face.GetId()));
 #endif
 
 	FTimePoint StartTime = FChrono::Now();
@@ -235,26 +235,26 @@ void FParametricMesher::Mesh(TSharedRef<FTopologicalFace> Face)
 	if (Grid.IsDegenerated())
 	{
 		MesherReport.Logs.AddDegeneratedGrid();
-		FMessage::Printf(EVerboseLevel::Log, TEXT("The meshing of the surface %d failed due to a degenerated grid\n"), Face->GetId());
-		Face->SetMeshed();
+		FMessage::Printf(EVerboseLevel::Log, TEXT("The meshing of the surface %d failed due to a degenerated grid\n"), Face.GetId());
+		Face.SetMeshed();
 		return;
 	}
 
 
 	FTimePoint IsoTriangulerStartTime = FChrono::Now();
 
-	TSharedRef<FFaceMesh> SurfaceMesh = StaticCastSharedRef<FFaceMesh>(Face->GetOrCreateMesh(MeshModel));
+	TSharedRef<FFaceMesh> SurfaceMesh = StaticCastSharedRef<FFaceMesh>(Face.GetOrCreateMesh(MeshModel));
 
 	FIsoTriangulator IsoTrianguler(Grid, SurfaceMesh, MesherReport);
 	if (IsoTrianguler.Triangulate())
 	{
-		if (Face->IsBackOriented())
+		if (Face.IsBackOriented())
 		{
 			SurfaceMesh->InverseOrientation();
 		}
 		MeshModel->AddMesh(SurfaceMesh);
 	}
-	Face->SetMeshed();
+	Face.SetMeshed();
 
 	FDuration TriangulateDuration = FChrono::Elapse(IsoTriangulerStartTime);
 	FDuration Duration = FChrono::Elapse(StartTime);
@@ -285,7 +285,7 @@ void FParametricMesher::GenerateCloud(FGrid& Grid)
 		FTimePoint StartTime = FChrono::Now();
 		Grid.SearchThinZones();
 
-		if (Grid.GetFace()->HasThinZone())
+		if (Grid.GetFace().HasThinZone())
 		{
 #ifdef DEBUG_THIN_ZONES
 			{
@@ -315,15 +315,15 @@ void FParametricMesher::GenerateCloud(FGrid& Grid)
 
 void FParametricMesher::MeshFaceLoops(FGrid& Grid)
 {
-	TSharedRef<FTopologicalFace> Face = Grid.GetFace();
+	const FTopologicalFace& Face = Grid.GetFace();
 
 	FTimePoint StartTime = FChrono::Now();
 
-	for (const TSharedPtr<FTopologicalLoop>& Loop : Face->GetLoops())
+	for (const TSharedPtr<FTopologicalLoop>& Loop : Face.GetLoops())
 	{
 		for (const FOrientedEdge& Edge : Loop->GetEdges())
 		{
-			Mesh(*Edge.Entity, *Face);
+			Mesh(*Edge.Entity, Face);
 		}
 	}
 
@@ -391,9 +391,9 @@ static void FillImposedIsoCuttingPoints(TArray<double>& UEdgeSetOfIntersectionWi
 	}
 }
 
-void FParametricMesher::Mesh(TSharedRef<FTopologicalVertex> InVertex)
+void FParametricMesher::Mesh(FTopologicalVertex& InVertex)
 {
-	InVertex->GetOrCreateMesh(GetMeshModel());
+	InVertex.GetOrCreateMesh(GetMeshModel());
 }
 
 //#define DEBUG_INTERSECTEDGEISOS
@@ -401,7 +401,7 @@ void FParametricMesher::Mesh(TSharedRef<FTopologicalVertex> InVertex)
 void DebugIntersectEdgeIsos(const FTopologicalFace& Face, const TArray<double>& IsoCoordinates, EIso TypeIso);
 #endif
 
-void FParametricMesher::Mesh(FTopologicalEdge& InEdge, FTopologicalFace& Face)
+void FParametricMesher::Mesh(FTopologicalEdge& InEdge, const FTopologicalFace& Face)
 {
 	{
 		FTopologicalEdge& ActiveEdge = *InEdge.GetLinkActiveEntity();
@@ -765,21 +765,16 @@ void FParametricMesher::GenerateEdgeElements(FTopologicalEdge& Edge)
 	ActiveEdge.SetMeshed();
 }
 
-void FParametricMesher::IsolateQuadFace(TArray<FCostToFace>& QuadSurfaces, TArray<TSharedPtr<FTopologicalFace>>& OtherSurfaces) const
+void FParametricMesher::IsolateQuadFace(TArray<FCostToFace>& QuadSurfaces, TArray<FTopologicalFace*>& OtherSurfaces) const
 {
-	TArray<TSharedPtr<FTopologicalFace>> FlatQuadsAndTriangles;
+	TArray<FTopologicalFace*> FlatQuadsAndTriangles;
 	FlatQuadsAndTriangles.Reserve(Faces.Num());
 	QuadSurfaces.Reserve(Faces.Num() * 2);
 	OtherSurfaces.Reserve(Faces.Num());
 
 	for (const TSharedPtr<FTopologicalFace>& Face : Faces)
 	{
-		if (Face->IsDeleted())
-		{
-			continue;
-		}
-
-		if (Face->IsMeshed())
+		if (Face->IsDeleted() || Face->IsMeshed())
 		{
 			continue;
 		}
@@ -787,36 +782,37 @@ void FParametricMesher::IsolateQuadFace(TArray<FCostToFace>& QuadSurfaces, TArra
 		Face->ComputeSurfaceSideProperties();
 	}
 
-	for (const TSharedPtr<FTopologicalFace>& Face : Faces)
+	for (const TSharedPtr<FTopologicalFace>& FacePtr : Faces)
 	{
-		Face->DefineSurfaceType();
-		switch (Face->GetQuadType())
+		FTopologicalFace& Face = *FacePtr;
+		Face.DefineSurfaceType();
+		switch (Face.GetQuadType())
 		{
 		case EQuadType::Quadrangular:
 			double LocalMinCurvature;
 			double LocalMaxCurvature;
-			GetMinMax(Face->GetCurvature(EIso::IsoU).Max, Face->GetCurvature(EIso::IsoV).Max, LocalMinCurvature, LocalMaxCurvature);
+			GetMinMax(Face.GetCurvature(EIso::IsoU).Max, Face.GetCurvature(EIso::IsoV).Max, LocalMinCurvature, LocalMaxCurvature);
 			if (LocalMaxCurvature > ConstMinCurvature)
 			{
-				QuadSurfaces.Emplace(LocalMaxCurvature, Face.ToSharedRef());
+				QuadSurfaces.Emplace(LocalMaxCurvature, &Face);
 				if (LocalMinCurvature > ConstMinCurvature)
 				{
-					QuadSurfaces.Emplace(LocalMinCurvature, Face.ToSharedRef());
+					QuadSurfaces.Emplace(LocalMinCurvature, &Face);
 				}
 			}
 			else
 			{
-				FlatQuadsAndTriangles.Add(Face);
-				OtherSurfaces.Add(Face);
+				FlatQuadsAndTriangles.Add(&Face);
+				OtherSurfaces.Add(&Face);
 			}
 			break;
 		case EQuadType::Triangular:
-			FlatQuadsAndTriangles.Add(Face);
-			OtherSurfaces.Add(Face);
+			FlatQuadsAndTriangles.Add(&Face);
+			OtherSurfaces.Add(&Face);
 			break;
 		case EQuadType::Unset:
 		default:
-			OtherSurfaces.Add(Face);
+			OtherSurfaces.Add(&Face);
 		}
 	}
 
@@ -832,7 +828,7 @@ void FParametricMesher::IsolateQuadFace(TArray<FCostToFace>& QuadSurfaces, TArra
 		Open3DDebugSession(TEXT("Quad Entities"));
 		for (const FCostToFace& Quad : QuadSurfaces)
 		{
-			//FGraphicSession _(FString::Printf(TEXT("Face %d %f"), Quad.Face->GetId(), Quad.Cost));
+			//FGraphicSession _(FString::Printf(TEXT("Face %d %f"), Quad.Face.GetId(), Quad.Cost));
 			Display(Quad.Face);
 		}
 		Close3DDebugSession();
@@ -843,7 +839,7 @@ void FParametricMesher::IsolateQuadFace(TArray<FCostToFace>& QuadSurfaces, TArra
 		Open3DDebugSession(TEXT("Flat Quads & Triangles"));
 		for (const TSharedPtr<FTopologicalFace>& Face : FlatQuadsAndTriangles)
 		{
-			//FGraphicSession _(FString::Printf(TEXT("Face %d %f"), Face->GetId()));
+			//FGraphicSession _(FString::Printf(TEXT("Face %d %f"), Face.GetId()));
 			Display(Face);
 		}
 		Close3DDebugSession();
@@ -854,7 +850,7 @@ void FParametricMesher::IsolateQuadFace(TArray<FCostToFace>& QuadSurfaces, TArra
 		Open3DDebugSession(TEXT("Other Entities"));
 		for (const TSharedPtr<FTopologicalFace>& Face : OtherSurfaces)
 		{
-			//FGraphicSession _(FString::Printf(TEXT("Face %d %f"), Face->GetId()));
+			//FGraphicSession _(FString::Printf(TEXT("Face %d %f"), Face.GetId()));
 			Display(StaticCastSharedPtr<FTopologicalFace>(Face));
 		}
 		Close3DDebugSession();
@@ -862,7 +858,7 @@ void FParametricMesher::IsolateQuadFace(TArray<FCostToFace>& QuadSurfaces, TArra
 #endif
 }
 
-void FParametricMesher::LinkQuadSurfaceForMesh(TArray<FCostToFace>& QuadTrimmedSurfaceSet, TArray<TArray<TSharedPtr<FTopologicalFace>>>& OutStrips)
+void FParametricMesher::LinkQuadSurfaceForMesh(TArray<FCostToFace>& QuadTrimmedSurfaceSet, TArray<TArray<FTopologicalFace*>>& OutStrips)
 {
 	const double GeometricTolerance = 20. * MeshModel->GetGeometricTolerance();
 
@@ -870,33 +866,33 @@ void FParametricMesher::LinkQuadSurfaceForMesh(TArray<FCostToFace>& QuadTrimmedS
 
 	for (FCostToFace& Quad : QuadTrimmedSurfaceSet)
 	{
-		TSharedPtr<FTopologicalFace> Surface = Quad.Face;
-		const FSurfaceCurvature& Curvatures = Surface->GetCurvatures();
+		FTopologicalFace* Face = Quad.Face;
+		const FSurfaceCurvature& Curvatures = Face->GetCurvatures();
 
 		EIso Axe = (!RealCompare(Quad.Cost, Curvatures[EIso::IsoU].Max)) ? EIso::IsoU : EIso::IsoV;
 
 		if (Axe == EIso::IsoU)
 		{
-			if (Surface->HasMarker1())
+			if (Face->HasMarker1())
 			{
 				continue;
 			}
-			Surface->SetMarker1();
+			Face->SetMarker1();
 		}
 		else
 		{
-			if (Surface->HasMarker2())
+			if (Face->HasMarker2())
 			{
 				continue;
 			}
-			Surface->SetMarker2();
+			Face->SetMarker2();
 		}
 
-		TArray<TSharedPtr<FTopologicalFace>>* QuadStrip = &OutStrips.Emplace_GetRef();
+		TArray<FTopologicalFace*>* QuadStrip = &OutStrips.Emplace_GetRef();
 		QuadStrip->Reserve(QuadTrimmedSurfaceSet.Num());
-		QuadStrip->Add(Surface);
+		QuadStrip->Add(Face);
 
-		const TArray<FEdge2DProperties>& SideProperties = Surface->GetSideProperties();
+		const TArray<FEdge2DProperties>& SideProperties = Face->GetSideProperties();
 
 		int32 StartSideIndex = 0;
 		for (; StartSideIndex < 4; StartSideIndex++)
@@ -914,87 +910,86 @@ void FParametricMesher::LinkQuadSurfaceForMesh(TArray<FCostToFace>& QuadTrimmedS
 		bool bFirstStep = true;
 		int32 SideIndex = StartSideIndex;
 
-		while (Surface)
+		while (Face)
 		{
-			int32 EdgeIndex = Surface->GetStartEdgeIndexOfSide(SideIndex);
-			double SideLength = Surface->GetSideProperties()[SideIndex].Length3D;
-			TSharedPtr<FTopologicalEdge> Edge = Surface->GetLoops()[0]->GetEdges()[EdgeIndex].Entity;
+			int32 EdgeIndex = Face->GetStartEdgeIndexOfSide(SideIndex);
+			double SideLength = Face->GetSideProperties()[SideIndex].Length3D;
+			TSharedPtr<FTopologicalEdge> Edge = Face->GetLoops()[0]->GetEdges()[EdgeIndex].Entity;
 
-			Surface.Reset();
+			Face = nullptr;
 			FTopologicalEdge* NextEdge = Edge->GetFirstTwinEdge();
 			if (NextEdge)
 			{
-				Surface = (TSharedPtr<FTopologicalFace>)NextEdge->GetLoop()->GetFace();
-				ensureCADKernel(Surface);
+				Face = NextEdge->GetLoop()->GetFace();
 			}
 
-			if (Surface && (Surface->GetQuadType() == EQuadType::Quadrangular || Surface->GetQuadType() == EQuadType::Triangular))
+			if (Face && (Face->GetQuadType() == EQuadType::Quadrangular || Face->GetQuadType() == EQuadType::Triangular))
 			{
 				// check side length
-				int32 LocalEdgeIndex = Surface->GetLoops()[0]->GetEdgeIndex(*NextEdge);
-				SideIndex = Surface->GetSideIndex(LocalEdgeIndex);
-				double OtherSideLength = Surface->GetSideProperties()[SideIndex].Length3D;
+				int32 LocalEdgeIndex = Face->GetLoops()[0]->GetEdgeIndex(*NextEdge);
+				SideIndex = Face->GetSideIndex(LocalEdgeIndex);
+				double OtherSideLength = Face->GetSideProperties()[SideIndex].Length3D;
 
 				GetMinMax(OtherSideLength, SideLength);
 				if (SideLength - OtherSideLength > GeometricTolerance)
 				{
-					Surface = nullptr;
+					Face = nullptr;
 				}
 			}
 			else
 			{
-				Surface = nullptr;
+				Face = nullptr;
 			}
 
-			if (Surface)
+			if (Face)
 			{
 				// Set as processed in a direction
-				const TArray<FEdge2DProperties>& LocalSideProperties = Surface->GetSideProperties();
+				const TArray<FEdge2DProperties>& LocalSideProperties = Face->GetSideProperties();
 				if (LocalSideProperties[SideIndex].IsoType == EIso::IsoU)
 				{
-					if (Surface->HasMarker1())
+					if (Face->HasMarker1())
 					{
-						Surface = nullptr;
+						Face = nullptr;
 					}
 					else
 					{
-						Surface->SetMarker1();
+						Face->SetMarker1();
 					}
 				}
 				else
 				{
-					if (Surface->HasMarker2())
+					if (Face->HasMarker2())
 					{
-						Surface = nullptr;
+						Face = nullptr;
 					}
 					else
 					{
-						Surface->SetMarker2();
+						Face->SetMarker2();
 					}
 				}
 			}
 
-			if (Surface)
+			if (Face)
 			{
 				// it's a quad or a tri => add
-				if (Surface->GetQuadType() != EQuadType::Other)
+				if (Face->GetQuadType() != EQuadType::Other)
 				{
-					QuadStrip->Add(Surface);
+					QuadStrip->Add(Face);
 				}
 
-				if (Surface->GetQuadType() == EQuadType::Triangular)
+				if (Face->GetQuadType() == EQuadType::Triangular)
 				{
 					// stop
-					Surface = nullptr;
+					Face = nullptr;
 				}
 			}
 
-			if (!Surface)
+			if (!Face)
 			{
 				if (bFirstStep)
 				{
 					bFirstStep = false;
-					Surface = (*QuadStrip)[0];
+					Face = (*QuadStrip)[0];
 					SideIndex = (StartSideIndex + 2) % 4;
 					continue;
 				}
@@ -1035,13 +1030,13 @@ void FParametricMesher::MeshSurfaceByFront(TArray<FCostToFace>& QuadTrimmedSurfa
 
 	const double GeometricTolerance = 20. * MeshModel->GetGeometricTolerance();
 
-	TArray<TSharedPtr<FTopologicalFace>> CandidateFacesForMesh; // first in first out
+	TArray<FTopologicalFace*> CandidateFacesForMesh; // first in first out
 	CandidateFacesForMesh.Reserve(100);
 
-	TArray<TSharedPtr<FTopologicalFace>> SecondChoiceOfCandidateFacesForMesh; // first in first out
+	TArray<FTopologicalFace*> SecondChoiceOfCandidateFacesForMesh; // first in first out
 	SecondChoiceOfCandidateFacesForMesh.Reserve(100);
 
-	TFunction<void(TSharedRef<FTopologicalFace>)> MeshFace = [&](TSharedRef<FTopologicalFace> Face)
+	TFunction<void(FTopologicalFace&)> MeshFace = [&](FTopologicalFace& Face)
 	{
 #ifdef DISPLAYDEBUGMESHFACEBYFACESTEP
 		Open3DDebugSession(TEXT("Surface " + Utils::ToFString(Face->GetId()));
@@ -1051,22 +1046,22 @@ void FParametricMesher::MeshSurfaceByFront(TArray<FCostToFace>& QuadTrimmedSurfa
 		Mesh(Face);
 
 #ifdef DISPLAYDEBUGMESHFACEBYFACESTEP
-		Open3DDebugSession(TEXT("Mesh " + Utils::ToFString(Face->GetId()));
-		DisplayMesh((FFaceMesh&)Face->GetOrCreateMesh(GetMeshModel()));
+		Open3DDebugSession(TEXT("Mesh " + Utils::ToFString(Face.GetId()));
+		DisplayMesh((FFaceMesh&)Face.GetOrCreateMesh(GetMeshModel()));
 		Close3DDebugSession();
 		//Wait();
 #endif
 
-		if (Face->HasMarker1())
+		if (Face.HasMarker1())
 		{
-			CandidateFacesForMesh.RemoveSingle(Face);
+			CandidateFacesForMesh.RemoveSingle(&Face);
 		}
-		if (Face->HasMarker2())
+		if (Face.HasMarker2())
 		{
-			SecondChoiceOfCandidateFacesForMesh.RemoveSingle(Face);
+			SecondChoiceOfCandidateFacesForMesh.RemoveSingle(&Face);
 		}
 
-		const TSharedPtr<FTopologicalLoop> Loop = Face->GetLoops()[0];
+		const TSharedPtr<FTopologicalLoop> Loop = Face.GetLoops()[0];
 		for (const FOrientedEdge& OrientedEdge : Loop->GetEdges())
 		{
 			const TSharedPtr<FTopologicalEdge> Edge = OrientedEdge.Entity;
@@ -1078,13 +1073,8 @@ void FParametricMesher::MeshSurfaceByFront(TArray<FCostToFace>& QuadTrimmedSurfa
 					continue;
 				}
 
-				TSharedPtr<FTopologicalFace> NextFace = NextEdge->GetFace();
-				if (!NextFace.IsValid())
-				{
-					continue;
-				}
-
-				if (!NextFace->HasMarker3())
+				FTopologicalFace* NextFace = NextEdge->GetFace();
+				if ((NextFace == nullptr) || !NextFace->HasMarker3())
 				{
 					// not in the scope of surface to mesh
 					continue;
@@ -1134,9 +1124,9 @@ void FParametricMesher::MeshSurfaceByFront(TArray<FCostToFace>& QuadTrimmedSurfa
 		}
 	};
 
-	TFunction<void(TSharedRef<FTopologicalFace>)> MeshFacesByFront = [&](TSharedRef<FTopologicalFace> Face)
+	TFunction<void(FTopologicalFace&)> MeshFacesByFront = [&](FTopologicalFace& Face)
 	{
-		if (Face->IsMeshed())
+		if (Face.IsMeshed())
 		{
 			return;
 		}
@@ -1146,7 +1136,7 @@ void FParametricMesher::MeshSurfaceByFront(TArray<FCostToFace>& QuadTrimmedSurfa
 		while (CandidateFacesForMesh.Num() || SecondChoiceOfCandidateFacesForMesh.Num())
 		{
 			// the candidate are sorted according to the number of meshed side 
-			Algo::Sort(CandidateFacesForMesh, [](TSharedPtr<FTopologicalFace> Surface1, TSharedPtr<FTopologicalFace> Surface2) { return Surface1->MeshedSideNum() > Surface2->MeshedSideNum(); });
+			Algo::Sort(CandidateFacesForMesh, [](FTopologicalFace* Surface1, FTopologicalFace* Surface2) { return Surface1->MeshedSideNum() > Surface2->MeshedSideNum(); });
 
 			int32 IndexOfBestCandidate = -1;
 			double CandidateMeshedSideRatio = 0;
@@ -1168,7 +1158,7 @@ void FParametricMesher::MeshSurfaceByFront(TArray<FCostToFace>& QuadTrimmedSurfa
 				int32 Index = 0;
 				for (; Index < CandidateFacesForMesh.Num(); ++Index)
 				{
-					TSharedPtr<FTopologicalFace> CandidateSurface = CandidateFacesForMesh[Index];
+					FTopologicalFace* CandidateSurface = CandidateFacesForMesh[Index];
 					if (CandidateSurface->MeshedSideNum() < MaxMeshedSideNum)
 					{
 						break;
@@ -1186,7 +1176,7 @@ void FParametricMesher::MeshSurfaceByFront(TArray<FCostToFace>& QuadTrimmedSurfa
 				{
 					for (; Index < CandidateFacesForMesh.Num(); ++Index)
 					{
-						TSharedPtr<FTopologicalFace> CandidateSurface = CandidateFacesForMesh[Index];
+						FTopologicalFace* CandidateSurface = CandidateFacesForMesh[Index];
 						if (CandidateMeshedSideRatio < CandidateSurface->MeshedSideRatio())
 						{
 							CandidateMeshedSideRatio = CandidateSurface->MeshedSideRatio();
@@ -1197,25 +1187,24 @@ void FParametricMesher::MeshSurfaceByFront(TArray<FCostToFace>& QuadTrimmedSurfa
 
 				if (IndexOfBestCandidate >= 0)
 				{
-					ensureCADKernel(CandidateFacesForMesh[IndexOfBestCandidate].IsValid());
-					MeshFace(CandidateFacesForMesh[IndexOfBestCandidate].ToSharedRef());
+					MeshFace(*CandidateFacesForMesh[IndexOfBestCandidate]);
 					continue;
 				}
 			}
 
 			for (int32 Index = 0; Index < SecondChoiceOfCandidateFacesForMesh.Num(); ++Index)
 			{
-				TSharedPtr<FTopologicalFace> CandidateSurface = SecondChoiceOfCandidateFacesForMesh[Index];
+				FTopologicalFace* CandidateSurface = SecondChoiceOfCandidateFacesForMesh[Index];
 				if (CandidateMeshedSideRatio < CandidateSurface->MeshedSideRatio())
 				{
 					CandidateMeshedSideRatio = CandidateSurface->MeshedSideRatio();
 					IndexOfBestCandidate = Index;
 				}
 			}
+
 			if (IndexOfBestCandidate >= 0)
 			{
-				ensureCADKernel(SecondChoiceOfCandidateFacesForMesh[IndexOfBestCandidate].IsValid());
-				MeshFace(SecondChoiceOfCandidateFacesForMesh[IndexOfBestCandidate].ToSharedRef());
+				MeshFace(*SecondChoiceOfCandidateFacesForMesh[IndexOfBestCandidate]);
 			}
 		}
 	};
@@ -1223,8 +1212,8 @@ void FParametricMesher::MeshSurfaceByFront(TArray<FCostToFace>& QuadTrimmedSurfa
 	// the front is initialized with quad surface
 	for (const FCostToFace& Quad : QuadTrimmedSurfaceSet)
 	{
-		TSharedPtr<FTopologicalFace> Surface = Quad.Face;
-		MeshFacesByFront(Surface.ToSharedRef());
+		FTopologicalFace* Face = Quad.Face;
+		MeshFacesByFront(*Face);
 	}
 
 	// the the other surface
@@ -1232,7 +1221,7 @@ void FParametricMesher::MeshSurfaceByFront(TArray<FCostToFace>& QuadTrimmedSurfa
 	{
 		if (!Face->IsMeshed())
 		{
-			MeshFacesByFront(Face.ToSharedRef());
+			MeshFacesByFront(*Face);
 		}
 	}
 
diff --git a/Engine/Source/Runtime/Datasmith/CADKernel/Private/CADKernel/Mesh/Structure/Grid.cpp b/Engine/Source/Runtime/Datasmith/CADKernel/Private/CADKernel/Mesh/Structure/Grid.cpp
--- a/Engine/Source/Runtime/Datasmith/CADKernel/Private/CADKernel/Mesh/Structure/Grid.cpp
+++ b/Engine/Source/Runtime/Datasmith/CADKernel/Private/CADKernel/Mesh/Structure/Grid.cpp
@@ -15,18 +15,18 @@
 //#define DEBUG_GETPREFERREDUVCOORDINATESFROMNEIGHBOURS
 namespace CADKernel
 {
-FGrid::FGrid(TSharedRef<FTopologicalFace>& InFace, TSharedRef<FModelMesh>& InMeshModel)
+FGrid::FGrid(FTopologicalFace& InFace, TSharedRef<FModelMesh>& InMeshModel)
 	: Face(InFace)
-	, FaceTolerance(InFace->GetIsoTolerances())
-	, Tolerance3D(InFace->GetCarrierSurface()->Get3DTolerance())
+	, FaceTolerance(InFace.GetIsoTolerances())
+	, Tolerance3D(InFace.GetCarrierSurface()->Get3DTolerance())
 	, MinimumElementSize(Tolerance3D * 2.)
 	, MeshModel(InMeshModel)
 	, ThinZoneFinder(*this)
-	, CuttingCoordinates(Face->GetCuttingPointCoordinates())
+	, CuttingCoordinates(InFace.GetCuttingPointCoordinates())
 {
 #ifdef DEBUG_ONLY_SURFACE_TO_DEBUG
-	bDisplay = (Face->GetId() == FaceToDebug);
-	Open3DDebugSession(bDisplay, FString::Printf(TEXT("Grid %d"), Face->GetId()));
+	bDisplay = (Face.GetId() == FaceToDebug);
+	Open3DDebugSession(bDisplay, FString::Printf(TEXT("Grid %d"), Face.GetId()));
 #endif
 }
 
@@ -108,21 +108,21 @@ void FGrid::DefineCuttingParameters(EIso Iso, FCuttingGrid& Neighbors)
 {
 	FTimePoint StartTime = FChrono::Now();
 
-	const FSurfacicBoundary& Boundary = Face->GetBoundary();
+	const FSurfacicBoundary& Boundary = Face.GetBoundary();
 
-#ifdef DEBUG_GETPREFERREDUVCOORDINATESFROMNEIGHBOURS
+#ifdef DEBUG_GET_PREFERRED_UVCOORDINATES_FROM_NEIGHBOURS
 	TArray<double> CuttingPointTmp;
 #endif
 	if (Neighbors[Iso].Num())
 	{
-#ifdef DEBUG_GETPREFERREDUVCOORDINATESFROMNEIGHBOURS
+#ifdef DEBUG_GET_PREFERRED_UVCOORDINATES_FROM_NEIGHBOURS
 		TArray<FCuttingPoint> Extremities;
 		Extremities.Reserve(2);
 		Extremities.Emplace(Boundary[Iso].Min, ECoordinateType::VertexCoordinate, -1, 0.001);
 		Extremities.Emplace(Boundary[Iso].Max, ECoordinateType::VertexCoordinate, -1, 0.001);
-		FMesherTools::ComputeFinalCuttingPointsWithImposedCuttingPoints(Face->GetCrossingPointCoordinates(Iso), Face->GetCrossingPointDeltaMaxs(Iso), Extremities, CuttingPointTmp);
+		FMesherTools::ComputeFinalCuttingPointsWithImposedCuttingPoints(Face.GetCrossingPointCoordinates(Iso), Face.GetCrossingPointDeltaMaxs(Iso), Extremities, CuttingPointTmp);
 #endif
-		FMesherTools::ComputeFinalCuttingPointsWithPreferredCuttingPoints(Face->GetCrossingPointCoordinates(Iso), Face->GetCrossingPointDeltaMaxs(Iso), Neighbors[Iso], Boundary[Iso], Face->GetCuttingCoordinatesAlongIso(Iso));
+		FMesherTools::ComputeFinalCuttingPointsWithPreferredCuttingPoints(Face.GetCrossingPointCoordinates(Iso), Face.GetCrossingPointDeltaMaxs(Iso), Neighbors[Iso], Boundary[Iso], Face.GetCuttingCoordinatesAlongIso(Iso));
 	}
 	else
 	{
@@ -130,10 +130,10 @@ void FGrid::DefineCuttingParameters(EIso Iso, FCuttingGrid& Neighbors)
 		Extremities.Reserve(2);
 		Extremities.Emplace(Boundary[Iso].Min, ECoordinateType::VertexCoordinate, -1, 0.001);
 		Extremities.Emplace(Boundary[Iso].Max, ECoordinateType::VertexCoordinate, -1, 0.001);
-		FMesherTools::ComputeFinalCuttingPointsWithImposedCuttingPoints(Face->GetCrossingPointCoordinates(Iso), Face->GetCrossingPointDeltaMaxs(Iso), Extremities, Face->GetCuttingCoordinatesAlongIso(Iso));
+		FMesherTools::ComputeFinalCuttingPointsWithImposedCuttingPoints(Face.GetCrossingPointCoordinates(Iso), Face.GetCrossingPointDeltaMaxs(Iso), Extremities, Face.GetCuttingCoordinatesAlongIso(Iso));
 	}
 
-#ifdef DEBUG_GETPREFERREDUVCOORDINATESFROMNEIGHBOURS
+#ifdef DEBUG_GET_PREFERRED_UVCOORDINATES_FROM_NEIGHBOURS
 	if (bDisplay)
 	{
 		EIso OtherIso = Other(Iso);
@@ -168,7 +168,7 @@ void FGrid::DefineCuttingParameters(EIso Iso, FCuttingGrid& Neighbors)
 		}
 		{
 			F3DDebugSession _(FString::Printf(TEXT("%s From Neighbours"), IsoNames[Iso]));
-			for (double CuttingU : Face->GetCuttingCoordinatesAlongIso(Iso))
+			for (double CuttingU : Face.GetCuttingCoordinatesAlongIso(Iso))
 			{
 				if (Iso == EIso::IsoU)
 				{
@@ -188,23 +188,23 @@ void FGrid::DefineCuttingParameters(EIso Iso, FCuttingGrid& Neighbors)
 	Chronos.DefineCuttingParametersDuration = FChrono::Elapse(StartTime);
 }
 
-//#define DEBUG_GETPREFERREDUVCOORDINATESFROMNEIGHBOURS
+//#define DEBUG_GET_PREFERRED_UVCOORDINATES_FROM_NEIGHBOURS
 void FGrid::GetPreferredUVCuttingParametersFromLoops(FCuttingGrid& CuttingParametersFromLoops)
 {
-#ifdef DEBUG_GETPREFERREDUVCOORDINATESFROMNEIGHBOURS
+#ifdef DEBUG_GET_PREFERRED_UVCOORDINATES_FROM_NEIGHBOURS
 	F3DDebugSession _(bDisplay, TEXT("GetPreferredUVCoordinatesFromNeighbours"));
 
 	if (bDisplay)
 	{
 		{
 			F3DDebugSession _(TEXT("Surface 2D"));
-			CADKernel::Display2D(*Face->GetCarrierSurface());
+			CADKernel::Display2D(*Face.GetCarrierSurface());
 		}
 	}
 #endif
 
 	int32 nbPoints = 0;
-	for (const TSharedPtr<FTopologicalLoop>& Loop : Face->GetLoops())
+	for (const TSharedPtr<FTopologicalLoop>& Loop : Face.GetLoops())
 	{
 		for (const FOrientedEdge& Edge : Loop->GetEdges())
 		{
@@ -215,13 +215,13 @@ void FGrid::GetPreferredUVCuttingParametersFromLoops(FCuttingGrid& CuttingParame
 	CuttingParametersFromLoops[EIso::IsoU].Reserve(nbPoints);
 	CuttingParametersFromLoops[EIso::IsoV].Reserve(nbPoints);
 
-	for (const TSharedPtr<FTopologicalLoop>& Loop : Face->GetLoops())
+	for (const TSharedPtr<FTopologicalLoop>& Loop : Face.GetLoops())
 	{
 		for (const FOrientedEdge& OrientedEdge : Loop->GetEdges())
 		{
 			const TSharedPtr<FTopologicalEdge>& Edge = OrientedEdge.Entity;
 
-#ifdef DEBUG_GETPREFERREDUVCOORDINATESFROMNEIGHBOURS
+#ifdef DEBUG_GET_PREFERRED_UVCOORDINATES_FROM_NEIGHBOURS
 			if (bDisplay)
 			{
 				F3DDebugSession _(FString::Printf(TEXT("Edge %d"), Edge->GetId()));
@@ -257,14 +257,14 @@ void FGrid::GetPreferredUVCuttingParametersFromLoops(FCuttingGrid& CuttingParame
 			TArray<FPoint2D> EdgePoints2D;
 			Edge->Approximate2DPoints(ProjectedPointCoords, EdgePoints2D);
 
-#ifdef DEBUG_GETPREFERREDUVCOORDINATESFROMNEIGHBOURS
+#ifdef DEBUG_GET_PREFERRED_UVCOORDINATES_FROM_NEIGHBOURS
 			F3DDebugSession _(bDisplay, ("Nodes"));
 #endif
 			for (int32 Index = 0; Index < EdgePoints2D.Num(); ++Index)
 			{
 				CuttingParametersFromLoops[EIso::IsoU].Emplace(EdgePoints2D[Index].U, ECoordinateType::OtherCoordinate);
 				CuttingParametersFromLoops[EIso::IsoV].Emplace(EdgePoints2D[Index].V, ECoordinateType::OtherCoordinate);
-#ifdef DEBUG_GETPREFERREDUVCOORDINATESFROMNEIGHBOURS
+#ifdef DEBUG_GET_PREFERRED_UVCOORDINATES_FROM_NEIGHBOURS
 				if (bDisplay)
 				{
 					DisplayPoint(EdgePoints2D[Index]);
@@ -339,11 +339,11 @@ bool FGrid::GeneratePointCloud()
 	{
 		for (int32 IPointU = 0; IPointU < CuttingCount[EIso::IsoU]; ++IPointU, ++Index)
 		{
-			Points2D[EGridSpace::Default2D][Index].Set(Face->GetCuttingCoordinatesAlongIso(EIso::IsoU)[IPointU], Face->GetCuttingCoordinatesAlongIso(EIso::IsoV)[IPointV]);
+			Points2D[EGridSpace::Default2D][Index].Set(Face.GetCuttingCoordinatesAlongIso(EIso::IsoU)[IPointU], Face.GetCuttingCoordinatesAlongIso(EIso::IsoV)[IPointV]);
 		}
 	}
 
-	Face->EvaluateGrid(*this);
+	Face.EvaluateGrid(*this);
 
 	ComputeMaxElementSize();
 
@@ -933,12 +933,12 @@ void SlightlyDisplacedPolyline(TArray<FPoint2D>& D2Points, const FSurfacicBounda
 bool FGrid::GetMeshOfLoops()
 {
 	int32 ThinZoneNum = 0;
-	if (Face->HasThinZone())
+	if (Face.HasThinZone())
 	{
 		ThinZoneNum = ThinZoneFinder.GetThinZones().Num();
 	}
 
-	int32 LoopCount = Face->GetLoops().Num();
+	int32 LoopCount = Face.GetLoops().Num();
 	FaceLoops2D[EGridSpace::Default2D].Reserve(LoopCount + ThinZoneNum);
 
 	FaceLoops3D.Reserve(LoopCount);
@@ -949,13 +949,13 @@ bool FGrid::GetMeshOfLoops()
 	F3DDebugSession _(bDisplay, ("GetLoopMesh"));
 #endif
 
-	for (const TSharedPtr<FTopologicalLoop>& Loop : Face->GetLoops())
+	for (const TSharedPtr<FTopologicalLoop>& Loop : Face.GetLoops())
 	{
 		int32 LoopNodeCount = 0;
 
 		for (const FOrientedEdge& Edge : Loop->GetEdges())
 		{
-			LoopNodeCount += Edge.Entity->GetLinkActiveEdge()->GetCuttingPoints().Num() + 2;
+			LoopNodeCount += Edge.Entity->GetLinkActiveEdge()->GetMesh()->GetNodeCount() + 2;
 		}
 
 		TArray<FPoint2D>& Loop2D = FaceLoops2D[EGridSpace::Default2D].Emplace_GetRef();
@@ -1170,7 +1170,7 @@ bool FGrid::GetMeshOfLoops()
 	}
 
 	// Fit boundaries to Surface bounds.
-	const FSurfacicBoundary& Bounds = Face->GetBoundary();
+	const FSurfacicBoundary& Bounds = Face.GetBoundary();
 	for (TArray<FPoint2D>& Loop : FaceLoops2D[EGridSpace::Default2D])
 	{
 		for (FPoint2D& Point : Loop)
@@ -1292,7 +1292,7 @@ void FGrid::ScaleGrid()
 		Grid[EIso::IsoV].Add(MiddleV);
 
 		FSurfacicSampling MiddlePoints;
-		Face->EvaluatePointGrid(Grid, MiddlePoints);
+		Face.EvaluatePointGrid(Grid, MiddlePoints);
 
 		int32 StartIndexUp = GetMiddleIndex(CuttingCoordinates[EIso::IsoV], MiddleV);
 		int32 StartIndexDown = StartIndexUp - 1;
@@ -1330,7 +1330,7 @@ void FGrid::ScaleGrid()
 		Grid[EIso::IsoV] = CuttingCoordinates[EIso::IsoV];
 
 		FSurfacicSampling MiddlePoints;
-		Face->EvaluatePointGrid(Grid, MiddlePoints);
+		Face.EvaluatePointGrid(Grid, MiddlePoints);
 
 		int32 StartIndexUp = GetMiddleIndex(CuttingCoordinates[EIso::IsoU], MiddleU);
 		int32 StartIndexDown = StartIndexUp - 1;
diff --git a/Engine/Source/Runtime/Datasmith/CADKernel/Private/CADKernel/Mesh/Structure/Grid.h b/Engine/Source/Runtime/Datasmith/CADKernel/Private/CADKernel/Mesh/Structure/Grid.h
--- a/Engine/Source/Runtime/Datasmith/CADKernel/Private/CADKernel/Mesh/Structure/Grid.h
+++ b/Engine/Source/Runtime/Datasmith/CADKernel/Private/CADKernel/Mesh/Structure/Grid.h
@@ -71,7 +71,7 @@ class FGrid : public FHaveStates
 {
 protected:
 
-	const TSharedRef<FTopologicalFace> Face;
+	FTopologicalFace& Face;
 	const FSurfacicTolerance FaceTolerance;
 	const double Tolerance3D;
 	const double MinimumElementSize;
@@ -157,7 +157,7 @@ public:
 	FGridChronos Chronos;
 
 public:
-	FGrid(TSharedRef<FTopologicalFace>& InFace, TSharedRef<FModelMesh>& InShellMesh);
+	FGrid(FTopologicalFace& InFace, TSharedRef<FModelMesh>& InShellMesh);
 
 #ifndef CADKERNEL_DEV
 	virtual ~FGrid() = default;
@@ -458,7 +458,12 @@ public:
 		return Normals;
 	}
 
-	TSharedRef<FTopologicalFace> GetFace() const
+	const FTopologicalFace& GetFace() const
+	{
+		return Face;
+	}
+
+	FTopologicalFace& GetFace()
 	{
 		return Face;
 	}
diff --git a/Engine/Source/Runtime/Datasmith/CADKernel/Private/CADKernel/Mesh/Structure/LoopCleaner.cpp b/Engine/Source/Runtime/Datasmith/CADKernel/Private/CADKernel/Mesh/Structure/LoopCleaner.cpp
--- a/Engine/Source/Runtime/Datasmith/CADKernel/Private/CADKernel/Mesh/Structure/LoopCleaner.cpp
+++ b/Engine/Source/Runtime/Datasmith/CADKernel/Private/CADKernel/Mesh/Structure/LoopCleaner.cpp
@@ -87,7 +87,7 @@ bool FLoopCleaner::CleanLoops()
 
 		if (!RemoveSelfIntersectionsOfLoop())
 		{
-			FMessage::Printf(Log, TEXT("Loop intersections of the surface %d cannot be fixed. The mesh of this surface is canceled.\n"), Grid.GetFace()->GetId());
+			FMessage::Printf(Log, TEXT("Loop intersections of the surface %d cannot be fixed. The mesh of this surface is canceled.\n"), Grid.GetFace().GetId());
 			return false;
 		}
 
@@ -98,7 +98,7 @@ bool FLoopCleaner::CleanLoops()
 
 		if (!RemoveLoopPicks())
 		{
-			FMessage::Printf(Log, TEXT("Loop intersections of the surface %d cannot be fixed. The mesh of this surface is canceled.\n"), Grid.GetFace()->GetId());
+			FMessage::Printf(Log, TEXT("Loop intersections of the surface %d cannot be fixed. The mesh of this surface is canceled.\n"), Grid.GetFace().GetId());
 			return false;
 		}
 
diff --git a/Engine/Source/Runtime/Datasmith/CADKernel/Private/CADKernel/Mesh/Structure/ModelMesh.cpp b/Engine/Source/Runtime/Datasmith/CADKernel/Private/CADKernel/Mesh/Structure/ModelMesh.cpp
--- a/Engine/Source/Runtime/Datasmith/CADKernel/Private/CADKernel/Mesh/Structure/ModelMesh.cpp
+++ b/Engine/Source/Runtime/Datasmith/CADKernel/Private/CADKernel/Mesh/Structure/ModelMesh.cpp
@@ -101,7 +101,7 @@ const int32 FModelMesh::GetIndexOfSurfaceFromId(const int32 Ident) const
 
 void FModelMesh::GetNodeCoordinates(TArray<FPoint>& NodeCoordinates) const
 {
-	NodeCoordinates.Reserve(LastIdUsed);
+	NodeCoordinates.Reserve(LastIdUsed+1);
 
 	for (const TArray<FPoint>* PointArray : GlobalPointCloud)
 	{
diff --git a/Engine/Source/Runtime/Datasmith/CADKernel/Private/CADKernel/Mesh/Structure/ThinZone2D.cpp b/Engine/Source/Runtime/Datasmith/CADKernel/Private/CADKernel/Mesh/Structure/ThinZone2D.cpp
--- a/Engine/Source/Runtime/Datasmith/CADKernel/Private/CADKernel/Mesh/Structure/ThinZone2D.cpp
+++ b/Engine/Source/Runtime/Datasmith/CADKernel/Private/CADKernel/Mesh/Structure/ThinZone2D.cpp
@@ -612,7 +612,7 @@ void FThinZone2DFinder::SearchThinZones()
 
 	if (ThinZones.Num() > 0)
 	{
-		Grid.GetFace()->SetHasThinZone();
+		Grid.GetFace().SetHasThinZone();
 	}
 
 #ifdef DISPLAY_THIN_ZONES
@@ -684,7 +684,7 @@ void FThinZone2DFinder::BuildLoopSegments()
 	double Length = 0;
 	double SegmentLength = Tolerance / 5.;
 
-	const TArray<TSharedPtr<FTopologicalLoop>>& Loops = Grid.GetFace()->GetLoops();
+	const TArray<TSharedPtr<FTopologicalLoop>>& Loops = Grid.GetFace().GetLoops();
 
 	FirstLoopLength = -1.;
 	for (const TSharedPtr<FTopologicalLoop>& Loop : Loops)
diff --git a/Engine/Source/Runtime/Datasmith/CADKernel/Private/CADKernel/Mesh/Structure/ThinZone2D.h b/Engine/Source/Runtime/Datasmith/CADKernel/Private/CADKernel/Mesh/Structure/ThinZone2D.h
--- a/Engine/Source/Runtime/Datasmith/CADKernel/Private/CADKernel/Mesh/Structure/ThinZone2D.h
+++ b/Engine/Source/Runtime/Datasmith/CADKernel/Private/CADKernel/Mesh/Structure/ThinZone2D.h
@@ -253,7 +253,7 @@ namespace CADKernel
 		FThinZoneChronos Chronos;
 
 	protected:
-		const FGrid& Grid;
+		FGrid& Grid;
 
 		double Tolerance;
 		double SquareTolerance;
@@ -270,7 +270,7 @@ namespace CADKernel
 
 	public:
 
-		FThinZone2DFinder(const FGrid& InGrid)
+		FThinZone2DFinder(FGrid& InGrid)
 			: Grid(InGrid)
 			, SegmentFatory()
 		{
diff --git a/Engine/Source/Runtime/Datasmith/CADKernel/Private/CADKernel/Topo/Body.cpp b/Engine/Source/Runtime/Datasmith/CADKernel/Private/CADKernel/Topo/Body.cpp
--- a/Engine/Source/Runtime/Datasmith/CADKernel/Private/CADKernel/Topo/Body.cpp
+++ b/Engine/Source/Runtime/Datasmith/CADKernel/Private/CADKernel/Topo/Body.cpp
@@ -11,7 +11,7 @@ namespace CADKernel
 void FBody::AddShell(TSharedRef<FShell> Shell)
 {
 	Shells.Add(Shell);
-	Shell->HostedBy = StaticCastSharedRef<FBody>(AsShared());
+	Shell->HostedBy = this;
 }
 
 void FBody::RemoveEmptyShell()
@@ -24,6 +24,10 @@ void FBody::RemoveEmptyShell()
 		{
 			NewShells.Emplace(Shell);
 		}
+		else
+		{
+			Shell->Delete();
+		}
 	}
 	Swap(NewShells, Shells);
 }
diff --git a/Engine/Source/Runtime/Datasmith/CADKernel/Private/CADKernel/Topo/Joiner.cpp b/Engine/Source/Runtime/Datasmith/CADKernel/Private/CADKernel/Topo/Joiner.cpp
--- a/Engine/Source/Runtime/Datasmith/CADKernel/Private/CADKernel/Topo/Joiner.cpp
+++ b/Engine/Source/Runtime/Datasmith/CADKernel/Private/CADKernel/Topo/Joiner.cpp
@@ -86,30 +86,28 @@ void FJoiner::InitFaces()
 	}
 }
 
-
 void FJoiner::RemoveFacesFromShell()
 {
 	// remove faces from their shells
-	TSet<TWeakPtr<FShell>> ShellSet;
+	TSet<FShell*> ShellSet;
 	for (const TSharedPtr<FTopologicalFace>& Face : Faces)
 	{
-		if (Face->GetHost().IsValid())
+		if (Face->GetHost() != nullptr)
 		{
 			ShellSet.Add(Face->GetHost());
 			Face->ResetHost();
 		}	
 	}
 
-	for (TWeakPtr<FShell> WeakShell : ShellSet)
+	for (FShell* Shell : ShellSet)
 	{
-		TSharedPtr<FShell> Shell = WeakShell.Pin();
 		bool bIsOutter = Shell->IsOutter();
 
 		TArray<FOrientedFace> ShellFace;
 		ShellFace.Reserve(Shell->FaceCount());
 		for (const FOrientedFace& Face : Shell->GetFaces())
 		{
-			if (!Face.Entity->GetHost().IsValid())
+			if (!Face.Entity->GetHost())
 			{
 				if (bIsOutter != (Face.Direction == EOrientation::Front))
 				{
@@ -139,7 +137,7 @@ void FJoiner::JoinFaces()
 
 	//RemoveIsolatedEdges();
 
-	TArray<TSharedPtr<FTopologicalVertex>> BorderVertices;
+	TArray<FTopologicalVertex*> BorderVertices;
 	GetBorderVertices(BorderVertices);
 	MergeCoincidentVertices(BorderVertices);
 	//RemoveIsolatedEdges();
@@ -158,12 +156,9 @@ void FJoiner::JoinFaces()
 
 	FDuration InitKernelIODuration = FChrono::Elapse(StartJoinTime);
 	FChrono::PrintClockElapse(EVerboseLevel::Log, TEXT(""), TEXT("Join"), InitKernelIODuration);
-
-	// clean the model at the end
-	SplitIntoConnectedShell();
 }
 
-void FJoiner::GetVertices(TArray<TSharedPtr<FTopologicalVertex>>& Vertices)
+void FJoiner::GetVertices(TArray<FTopologicalVertex*>& Vertices)
 {
 	Vertices.Empty(10 * Faces.Num());
 
@@ -177,30 +172,30 @@ void FJoiner::GetVertices(TArray<TSharedPtr<FTopologicalVertex>>& Vertices)
 				if (!Edge->GetStartVertex()->GetLinkActiveEntity()->HasMarker1())
 				{
 					Edge->GetStartVertex()->GetLinkActiveEntity()->SetMarker1();
-					Vertices.Add(Edge->GetStartVertex()->GetLinkActiveEntity());
+					Vertices.Add(&*Edge->GetStartVertex()->GetLinkActiveEntity());
 				}
 				if (!Edge->GetEndVertex()->GetLinkActiveEntity()->HasMarker1())
 				{
 					Edge->GetEndVertex()->GetLinkActiveEntity()->SetMarker1();
-					Vertices.Add(Edge->GetEndVertex()->GetLinkActiveEntity());
+					Vertices.Add(&*Edge->GetEndVertex()->GetLinkActiveEntity());
 				}
 			}
 		}
 	}
-	for (const TSharedPtr<FTopologicalVertex>& Vertex : Vertices)
+	for (FTopologicalVertex* Vertex : Vertices)
 	{
 		Vertex->ResetMarker1();
 	}
 }
 
-void FJoiner::GetBorderVertices(TArray<TSharedPtr<FTopologicalVertex>>& BorderVertices)
+void FJoiner::GetBorderVertices(TArray<FTopologicalVertex*>& BorderVertices)
 {
-	TArray<TSharedPtr<FTopologicalVertex>> Vertices;
+	TArray<FTopologicalVertex*> Vertices;
 	GetVertices(Vertices);
 
 	BorderVertices.Empty(Vertices.Num());
 
-	for (const TSharedPtr<FTopologicalVertex>& Vertex : Vertices)
+	for (FTopologicalVertex* Vertex : Vertices)
 	{
 		if (Vertex->IsBorderVertex())
 		{
@@ -209,7 +204,7 @@ void FJoiner::GetBorderVertices(TArray<TSharedPtr<FTopologicalVertex>>& BorderVe
 	}
 }
 
-void FJoiner::MergeCoincidentVertices(TArray<TSharedPtr<FTopologicalVertex>>& VerticesToMerge)
+void FJoiner::MergeCoincidentVertices(TArray<FTopologicalVertex*>& VerticesToMerge)
 {
 	FTimePoint StartTime = FChrono::Now();
 
@@ -224,7 +219,7 @@ void FJoiner::MergeCoincidentVertices(TArray<TSharedPtr<FTopologicalVertex>>& Ve
 	TArray<int32> SortedVertexIndices;
 	SortedVertexIndices.Reserve(VertexNum);
 
-	for (TSharedPtr<FTopologicalVertex> Vertex : VerticesToMerge)
+	for (FTopologicalVertex* Vertex : VerticesToMerge)
 	{
 		VerticesWeight.Add(Vertex->GetCoordinates().X + Vertex->GetCoordinates().Y + Vertex->GetCoordinates().Z);
 	}
@@ -237,7 +232,7 @@ void FJoiner::MergeCoincidentVertices(TArray<TSharedPtr<FTopologicalVertex>>& Ve
 
 	for (int32 IndexI = 0; IndexI < VertexNum; ++IndexI)
 	{
-		TSharedPtr<FTopologicalVertex> Vertex = VerticesToMerge[SortedVertexIndices[IndexI]];
+		FTopologicalVertex* Vertex = VerticesToMerge[SortedVertexIndices[IndexI]];
 		if (Vertex->HasMarker1())
 		{
 			continue;
@@ -252,7 +247,7 @@ void FJoiner::MergeCoincidentVertices(TArray<TSharedPtr<FTopologicalVertex>>& Ve
 
 		for (int32 IndexJ = IndexI + 1; IndexJ < VertexNum; ++IndexJ)
 		{
-			TSharedPtr<FTopologicalVertex> OtherVertex = VerticesToMerge[SortedVertexIndices[IndexJ]];
+			FTopologicalVertex* OtherVertex = VerticesToMerge[SortedVertexIndices[IndexJ]];
 			if (OtherVertex->HasMarker1())
 			{
 				continue;
@@ -274,26 +269,26 @@ void FJoiner::MergeCoincidentVertices(TArray<TSharedPtr<FTopologicalVertex>>& Ve
 		}
 	}
 
-	for (const TSharedPtr<FTopologicalVertex>& Vertex : VerticesToMerge)
+	for (FTopologicalVertex* Vertex : VerticesToMerge)
 	{
 		Vertex->ResetMarker1();
 	}
 
-	TArray<TSharedPtr<FTopologicalVertex>> ActiveVertices;
+	TArray<FTopologicalVertex*> ActiveVertices;
 	ActiveVertices.Reserve(VertexNum);
 
-	for (const TSharedPtr<FTopologicalVertex>& Vertex : VerticesToMerge)
+	for (FTopologicalVertex* Vertex : VerticesToMerge)
 	{
-		TSharedPtr<FTopologicalVertex> ActiveVertex = Vertex->GetLinkActiveEntity();
-		if (ActiveVertex->HasMarker1())
+		FTopologicalVertex& ActiveVertex = *Vertex->GetLinkActiveEntity();
+		if (ActiveVertex.HasMarker1())
 		{
 			continue;
 		}
-		ActiveVertex->SetMarker1();
-		ActiveVertices.Add(ActiveVertex);
+		ActiveVertex.SetMarker1();
+		ActiveVertices.Add(&ActiveVertex);
 	}
 
-	for (const TSharedPtr<FTopologicalVertex>& Vertex : ActiveVertices)
+	for (FTopologicalVertex* Vertex : ActiveVertices)
 	{
 		Vertex->ResetMarker1();
 	}
@@ -305,7 +300,7 @@ void FJoiner::MergeCoincidentVertices(TArray<TSharedPtr<FTopologicalVertex>>& Ve
 
 }
 
-void FJoiner::MergeBorderVerticesWithCoincidentOtherVertices(TArray<TSharedPtr<FTopologicalVertex>>& Vertices)
+void FJoiner::MergeBorderVerticesWithCoincidentOtherVertices(TArray<FTopologicalVertex*>& Vertices)
 {
 	int32 VertexNum = (int32)Vertices.Num();
 
@@ -314,7 +309,7 @@ void FJoiner::MergeBorderVerticesWithCoincidentOtherVertices(TArray<TSharedPtr<F
 	TArray<int32> SortedVertexIndices;
 	SortedVertexIndices.Reserve(VertexNum);
 
-	for (const TSharedPtr<FTopologicalVertex>& Vertex : Vertices)
+	for (FTopologicalVertex* Vertex : Vertices)
 	{
 		VerticesWeight.Add(Vertex->GetCoordinates().X + Vertex->GetCoordinates().Y + Vertex->GetCoordinates().Z);
 	}
@@ -328,7 +323,7 @@ void FJoiner::MergeBorderVerticesWithCoincidentOtherVertices(TArray<TSharedPtr<F
 	int32 StartIndexJ = 1;
 	for (int32 IndexI = 0; IndexI < VertexNum; ++IndexI)
 	{
-		TSharedRef<FTopologicalVertex> Vertex = Vertices[SortedVertexIndices[IndexI]].ToSharedRef();
+		FTopologicalVertex* Vertex = Vertices[SortedVertexIndices[IndexI]];
 		if (!Vertex->IsBorderVertex())
 		{
 			continue;
@@ -374,21 +369,21 @@ void FJoiner::MergeBorderVerticesWithCoincidentOtherVertices(TArray<TSharedPtr<F
 		}
 	}
 
-	TArray<TSharedPtr<FTopologicalVertex>> NewVertices;
+	TArray<FTopologicalVertex*> NewVertices;
 	NewVertices.Empty(VertexNum);
 
-	for (const TSharedPtr<FTopologicalVertex>& Vertex : Vertices)
+	for (FTopologicalVertex* Vertex : Vertices)
 	{
-		const TSharedPtr<FTopologicalVertex>& ActiveVertex = Vertex->GetLinkActiveEntity();
-		if (ActiveVertex->HasMarker1())
+		FTopologicalVertex& ActiveVertex = *Vertex->GetLinkActiveEntity();
+		if (ActiveVertex.HasMarker1())
 		{
 			continue;
 		}
-		ActiveVertex->SetMarker1();
-		NewVertices.Add(ActiveVertex);
+		ActiveVertex.SetMarker1();
+		NewVertices.Add(&ActiveVertex);
 	}
 
-	for (TSharedPtr<FTopologicalVertex> Vertex : NewVertices)
+	for (FTopologicalVertex* Vertex : NewVertices)
 	{
 		Vertex->ResetMarker1();
 	}
@@ -396,15 +391,12 @@ void FJoiner::MergeBorderVerticesWithCoincidentOtherVertices(TArray<TSharedPtr<F
 	Swap(NewVertices, Vertices);
 }
 
-void FJoiner::MergeCoincidentEdges(TArray<TSharedPtr<FTopologicalVertex>>& VerticesToProcess)
+void FJoiner::MergeCoincidentEdges(TArray<FTopologicalVertex*>& VerticesToProcess)
 {
 	FTimePoint StartTime = FChrono::Now();
 
-	for (TSharedPtr<FTopologicalVertex> VertexPtr : VerticesToProcess)
+	for (FTopologicalVertex* Vertex : VerticesToProcess)
 	{
-		ensureCADKernel(VertexPtr.IsValid());
-		TSharedRef<FTopologicalVertex> Vertex = VertexPtr.ToSharedRef();
-
 		TArray<FTopologicalEdge*> ConnectedEdges;
 		Vertex->GetConnectedEdges(ConnectedEdges);
 		int32 ConnectedEdgeCount = ConnectedEdges.Num();
@@ -421,7 +413,7 @@ void FJoiner::MergeCoincidentEdges(TArray<TSharedPtr<FTopologicalVertex>>& Verti
 				continue;
 			}
 			bool bFirstEdgeBorder = Edge->IsBorder();
-			TSharedRef<FTopologicalVertex> EndVertex = Edge->GetOtherVertex(Vertex)->GetLinkActiveEntity();
+			FTopologicalVertex& EndVertex = *Edge->GetOtherVertex(*Vertex)->GetLinkActiveEntity();
 
 			for (int32 EdgeJ = EdgeI + 1; EdgeJ < ConnectedEdgeCount; ++EdgeJ)
 			{
@@ -438,9 +430,9 @@ void FJoiner::MergeCoincidentEdges(TArray<TSharedPtr<FTopologicalVertex>>& Verti
 					continue;
 				}
 
-				TSharedRef<FTopologicalVertex> OtherEdgeEndVertex = SecondEdge->GetOtherVertex(Vertex)->GetLinkActiveEntity();
+				FTopologicalVertex& OtherEdgeEndVertex = *SecondEdge->GetOtherVertex(*Vertex)->GetLinkActiveEntity();
 
-				if (OtherEdgeEndVertex == EndVertex)
+				if (&OtherEdgeEndVertex == &EndVertex)
 				{
 					FPoint StartTangentEdge = Edge->GetTangentAt(*Vertex);
 					FPoint StartTangentOtherEdge = SecondEdge->GetTangentAt(*Vertex);
@@ -507,24 +499,19 @@ TSharedPtr<FTopologicalVertex> FJoiner::SplitAndLink(FTopologicalVertex& StartVe
 	return NewVertex;
 }
 
-int32 CountSplit = 0;
-
-void FJoiner::StitchParallelEdges(TArray<TSharedPtr<FTopologicalVertex>>& VerticesToProcess)
+void FJoiner::StitchParallelEdges(TArray<FTopologicalVertex*>& VerticesToProcess)
 {
 	FTimePoint StartTime = FChrono::Now();
 
 	for (int32 VertexI = 0; VertexI < VerticesToProcess.Num(); ++VertexI)
 	{
-		TSharedPtr<FTopologicalVertex> VertexPtr = VerticesToProcess[VertexI];
-		ensureCADKernel(VertexPtr.IsValid());
+		FTopologicalVertex* Vertex = VerticesToProcess[VertexI];
 
-		if (!VertexPtr->IsBorderVertex())
+		if (!Vertex->IsBorderVertex())
 		{
 			continue;
 		}
 
-		TSharedRef<FTopologicalVertex> Vertex = VertexPtr.ToSharedRef();
-
 		TArray<FTopologicalEdge*> ConnectedEdges;
 		Vertex->GetConnectedEdges(ConnectedEdges);
 		int32 ConnectedEdgeCount = ConnectedEdges.Num();
@@ -577,9 +564,9 @@ void FJoiner::StitchParallelEdges(TArray<TSharedPtr<FTopologicalVertex>>& Vertic
 					continue;
 				}
 
-				TSharedRef<FTopologicalVertex> EndVertex = Edge->GetOtherVertex(Vertex)->GetLinkActiveEntity();
-				TSharedRef<FTopologicalVertex> OtherEdgeEndVertex = SecondEdge->GetOtherVertex(Vertex)->GetLinkActiveEntity();
-				if (EndVertex == OtherEdgeEndVertex)
+				FTopologicalVertex& EndVertex = *Edge->GetOtherVertex(*Vertex)->GetLinkActiveEntity();
+				FTopologicalVertex& OtherEdgeEndVertex = *SecondEdge->GetOtherVertex(*Vertex)->GetLinkActiveEntity();
+				if (&EndVertex == &OtherEdgeEndVertex)
 				{
 					Edge->Link(*SecondEdge, JoiningTolerance);
 				}
@@ -587,23 +574,20 @@ void FJoiner::StitchParallelEdges(TArray<TSharedPtr<FTopologicalVertex>>& Vertic
 				{
 					double EdgeLength = Edge->Length();
 					double OtherEdgeLength = SecondEdge->Length();
+
+					TSharedPtr<FTopologicalVertex> NewVertex;
 					if (EdgeLength < OtherEdgeLength)
 					{
-						CountSplit++;
-						TSharedPtr<FTopologicalVertex> NewVertex = SplitAndLink(*Vertex, *Edge, *SecondEdge);
-						if (NewVertex.IsValid())
-						{
-							VerticesToProcess.Add(NewVertex);
-						}
+						NewVertex = SplitAndLink(*Vertex, *Edge, *SecondEdge);
 					}
 					else
 					{
-						CountSplit++;
-						TSharedPtr<FTopologicalVertex> NewVertex = SplitAndLink(*Vertex, *SecondEdge, *Edge);
-						if (NewVertex.IsValid())
-						{
-							VerticesToProcess.Add(NewVertex);
-						}
+						NewVertex = SplitAndLink(*Vertex, *SecondEdge, *Edge);
+					}
+
+					if (NewVertex.IsValid())
+					{
+						VerticesToProcess.Add(NewVertex.Get());
 					}
 				}
 			}
@@ -618,8 +602,15 @@ void FJoiner::MergeUnconnectedAdjacentEdges()
 {
 	FTimePoint StartTime = FChrono::Now();
 
-	for (TSharedPtr<FTopologicalFace> Face : Faces)
+	for (TSharedPtr<FTopologicalFace> FacePtr : Faces)
 	{
+		if (!FacePtr.IsValid())
+		{
+			continue;
+		}
+
+		FTopologicalFace* Face = FacePtr.Get();
+
 		TArray<TArray<FOrientedEdge>> ArrayOfCandidates;
 		ArrayOfCandidates.Reserve(10);
 
@@ -775,10 +766,10 @@ void FJoiner::RemoveIsolatedEdges()
 
 	TArray<FTopologicalEdge*> IsolatedEdges;
 
-	TArray<TSharedPtr<FTopologicalVertex>> Vertices;
+	TArray<FTopologicalVertex*> Vertices;
 	GetVertices(Vertices);
 
-	for (const TSharedPtr<FTopologicalVertex>& Vertex : Vertices)
+	for (const FTopologicalVertex* Vertex : Vertices)
 	{
 		for (const FTopologicalVertex* TwinVertex : Vertex->GetTwinsEntities())
 		{
@@ -799,6 +790,63 @@ void FJoiner::RemoveIsolatedEdges()
 	FMessage::Printf(EVerboseLevel::Log, TEXT("\n\nIsolatedEdges count %d\n\n\n"), IsolatedEdges.Num());
 }
 
+void FJoiner::UnlinkNonManifoldVertex()
+{
+	TArray<FTopologicalVertex*> Vertices;
+	GetVertices(Vertices);
+
+	TMap<FShell*, int32> ShellToVertexCount;
+	for (FTopologicalVertex* Vertex : Vertices)
+	{
+		ShellToVertexCount.Empty(Vertex->GetTwinsEntityCount());
+		for (const FTopologicalVertex* TwinVertex : Vertex->GetTwinsEntities())
+		{
+			ensureCADKernel(!TwinVertex->GetDirectConnectedEdges().IsEmpty());
+			FTopologicalEdge* Edge = TwinVertex->GetDirectConnectedEdges()[0];
+			FShell* Shell = Edge->GetLoop()->GetFace()->GetHost();
+			if (Shell != nullptr)
+			{
+				ShellToVertexCount.FindOrAdd(Shell)++;
+			}
+		}
+
+		if (ShellToVertexCount.Num() > 1)
+		{
+			TMap<FShell*, TArray<FTopologicalVertex*>> ShellToVertices;
+			ShellToVertices.Reserve(ShellToVertexCount.Num());
+			for (TPair<FShell*, int32> Pair : ShellToVertexCount)
+			{
+				ShellToVertices.FindOrAdd(Pair.Key).Reserve(Pair.Value);
+			}
+
+			for (FTopologicalVertex* TwinVertex : Vertex->GetTwinsEntities())
+			{
+				ensureCADKernel(!TwinVertex->GetDirectConnectedEdges().IsEmpty());
+				FTopologicalEdge* Edge = TwinVertex->GetDirectConnectedEdges()[0];
+				FShell* Shell = Edge->GetLoop()->GetFace()->GetHost();
+				if (Shell != nullptr)
+				{
+					ShellToVertices.FindOrAdd(Shell).Add(TwinVertex);
+				}
+			}
+
+			Vertex->UnlinkTwinEntities();
+			for (TPair<FShell*, TArray<FTopologicalVertex*>> ShellVertices : ShellToVertices)
+			{
+				if (ShellVertices.Value.Num() > 1)
+				{
+					TArray<FTopologicalVertex*>& Twins = ShellVertices.Value;
+					FTopologicalVertex* FirstVertex = Twins[0];
+					for (FTopologicalVertex* TwinVertex : Twins)
+					{
+						FirstVertex->Link(*TwinVertex);
+					}
+				}
+			}
+		}
+	}
+}
+
 // =========================================================================================================================================================================================================
 // =========================================================================================================================================================================================================
 // =========================================================================================================================================================================================================
@@ -851,10 +899,10 @@ void FJoiner::SplitIntoConnectedShell()
 
 	int32 ProcessFaceCount = 0;
 
-	TArray<TSharedPtr<FTopologicalFace>> Front;
-	TFunction<void(const TSharedPtr<FTopologicalFace>&, FFaceSubset&)> GetNeighboringFaces = [&](const TSharedPtr<FTopologicalFace>& Face, FFaceSubset& Shell)
+	TArray<FTopologicalFace*> Front;
+	TFunction<void(const FTopologicalFace&, FFaceSubset&)> GetNeighboringFaces = [&](const FTopologicalFace& Face, FFaceSubset& Shell)
 	{
-		for (const TSharedPtr<FTopologicalLoop>& Loop : Face->GetLoops())
+		for (const TSharedPtr<FTopologicalLoop>& Loop : Face.GetLoops())
 		{
 			for (const FOrientedEdge& OrientedEdge : Loop->GetEdges())
 			{
@@ -887,16 +935,12 @@ void FJoiner::SplitIntoConnectedShell()
 					}
 					NextEdge->SetMarker1();
 
-					TSharedPtr<FTopologicalFace> NextFace = NextEdge->GetFace();
-					if (!NextFace.IsValid())
+					FTopologicalFace* NextFace = NextEdge->GetFace();
+					if ((NextFace == nullptr) || NextFace->HasMarker1())
 					{
 						continue;
 					}
 
-					if (NextFace->HasMarker1())
-					{
-						continue;
-					}
 					NextFace->SetMarker1();
 					Front.Add(NextFace);
 				}
@@ -908,9 +952,14 @@ void FJoiner::SplitIntoConnectedShell()
 	{
 		while (Front.Num())
 		{
-			TSharedPtr<FTopologicalFace> Face = Front.Pop();
+			FTopologicalFace* Face = Front.Pop();
+			if (Face == nullptr)
+			{
+				continue;
+			}
+
 			Shell.Faces.Add(Face);
-			GetNeighboringFaces(Face, Shell);
+			GetNeighboringFaces(*Face, Shell);
 		}
 	};
 
@@ -926,7 +975,7 @@ void FJoiner::SplitIntoConnectedShell()
 		Front.Empty(TopologicalFaceCount);
 
 		Face->SetMarker1();
-		Front.Add(Face);
+		Front.Add(Face.Get());
 		SpreadFront(Shell);
 		ProcessFaceCount += Shell.Faces.Num();
 
@@ -952,64 +1001,26 @@ void FJoiner::SplitIntoConnectedShell()
 	// for each FaceSubset, find the main shell
 	for (FFaceSubset& FaceSubset : SubShells)
 	{
-		TMap<TWeakPtr<FShell>, int32> ShellToFaceCount;
+		TMap<FBody*, int32> BodyToFaceCount;
+		TMap<FShell*, int32> ShellToFaceCount;
 		TMap<uint32, int32> ColorToFaceCount;
 		TMap<FString, int32> NameToFaceCount;
 
-		for (const TSharedPtr<FTopologicalFace>& Face : FaceSubset.Faces)
+		for (FTopologicalFace* Face : FaceSubset.Faces)
 		{
-			TWeakPtr<FShell> Shell = Face->GetHost();
-			int32* Count = ShellToFaceCount.Find(Shell);
-			if (Count)
-			{
-				(*Count)++;
-			}
-			else
-			{
-				ShellToFaceCount.Add(Shell, 1);
-			}
+			FShell* Shell = Face->GetHost();
+			FBody* Body = Shell->GetHost();
 
+			ShellToFaceCount.FindOrAdd(Shell)++;
+			BodyToFaceCount.FindOrAdd(Body)++;
 			ColorToFaceCount.FindOrAdd(Face->GetColorId())++;
 			NameToFaceCount.FindOrAdd(Face->GetName())++;
 		}
 
-		int32 MaxInstance = 0;
-		TWeakPtr<FShell> MainShell;
-		for (TPair<TWeakPtr<FShell>, int32>& Pair : ShellToFaceCount)
-		{
-			if (Pair.Value > MaxInstance)
-			{
-				MaxInstance = Pair.Value;
-				MainShell = Pair.Key;
-			}
-		}
-
-		if (MainShell.IsValid() && MainShell.Pin()->FaceCount() / 2 + 1 < MaxInstance)
-		{
-			FaceSubset.MainShell = MainShell;
-		}
-
-		MaxInstance = FaceSubset.Faces.Num()/3;
-		FString MainName;
-		for (TPair<FString, int32>& Pair : NameToFaceCount)
-		{
-			if (Pair.Value > MaxInstance)
-			{
-				MaxInstance = Pair.Value;
-				MainName = Pair.Key;
-			}
-		}
-
-		MaxInstance = 0;
-		for (TPair<uint32, int32>& Pair : ColorToFaceCount)
-		{
-			if (Pair.Value > MaxInstance)
-			{
-				MaxInstance = Pair.Value;
-				FaceSubset.MainColor = Pair.Key;
-			}
-		}
-
+		FaceSubset.SetMainShell(ShellToFaceCount);
+		FaceSubset.SetMainBody(BodyToFaceCount);
+		FaceSubset.SetMainName(NameToFaceCount);
+		FaceSubset.SetMainColor(ColorToFaceCount);
 	}
 
 	if (Shells.Num())
@@ -1024,20 +1035,26 @@ void FJoiner::SplitIntoConnectedShell()
 	// for each FaceSubset, process the Shell
 	for (FFaceSubset FaceSubset : SubShells)
 	{
-		if (FaceSubset.MainShell.IsValid())
+		if (FaceSubset.MainShell != nullptr)
 		{
-			TSharedPtr<FShell> Shell = FaceSubset.MainShell.Pin();
+			FShell* Shell = FaceSubset.MainShell;
 			Shell->Empty(FaceSubset.Faces.Num());
 			Shell->Add(FaceSubset.Faces);
 		}
 		else
 		{
-			TSharedRef<FBody> Body = FEntity::MakeShared<FBody>();
-			Session.GetModel()->Add(Body);
+			FBody* Body = FaceSubset.MainBody;
+			if (Body == nullptr)
+			{
+				TSharedRef<FBody> SharedBody = FEntity::MakeShared<FBody>();
+				Session.GetModel()->Add(SharedBody);
+				Body = &SharedBody.Get();
+				Body->SetName(FaceSubset.MainName);
+				Body->SetColorId(FaceSubset.MainColor);
+			}
+
 			TSharedRef<FShell> Shell = FEntity::MakeShared<FShell>();
 			Body->AddShell(Shell);
-			Body->SetName(FaceSubset.MainName);
-			Body->SetColorId(FaceSubset.MainColor);
 
 			Shell->Add(FaceSubset.Faces);
 			Shell->SetName(FaceSubset.MainName);
@@ -1045,6 +1062,8 @@ void FJoiner::SplitIntoConnectedShell()
 		}
 	}
 
+	UnlinkNonManifoldVertex();
+
 	Session.GetModel()->RemoveEmptyBodies();
 }
 
diff --git a/Engine/Source/Runtime/Datasmith/CADKernel/Private/CADKernel/Topo/Model.cpp b/Engine/Source/Runtime/Datasmith/CADKernel/Private/CADKernel/Topo/Model.cpp
--- a/Engine/Source/Runtime/Datasmith/CADKernel/Private/CADKernel/Topo/Model.cpp
+++ b/Engine/Source/Runtime/Datasmith/CADKernel/Private/CADKernel/Topo/Model.cpp
@@ -86,6 +86,10 @@ void FModel::RemoveEmptyBodies()
 			NewBodies.Emplace(Body);
 			NbBody++;
 		}
+		else
+		{
+			Body->Delete();
+		}
 	}
 	Swap(NewBodies, Bodies);
 	FMessage::Printf(Log, TEXT("After RemoveEmptyBodies, Body count %d shell count %d \n"), NbBody, NbShell);
diff --git a/Engine/Source/Runtime/Datasmith/CADKernel/Private/CADKernel/Topo/Shell.cpp b/Engine/Source/Runtime/Datasmith/CADKernel/Private/CADKernel/Topo/Shell.cpp
--- a/Engine/Source/Runtime/Datasmith/CADKernel/Private/CADKernel/Topo/Shell.cpp
+++ b/Engine/Source/Runtime/Datasmith/CADKernel/Private/CADKernel/Topo/Shell.cpp
@@ -62,16 +62,15 @@ void FShell::Empty(int32 NewSize)
 	TopologicalFaces.Empty(NewSize);
 }
 
-void FShell::Add(TArray<TSharedPtr<FTopologicalFace>> Faces)
+void FShell::Add(TArray<FTopologicalFace*> Faces)
 {
-	TSharedPtr<FShell> Shell = StaticCastSharedRef<FShell>(AsShared());
-
 	TopologicalFaces.Reserve(TopologicalFaces.Num() + Faces.Num());
 
-	for(TSharedPtr<FTopologicalFace>& Face : Faces)
+	for(FTopologicalFace* Face : Faces)
 	{
-		TopologicalFaces.Emplace(Face, Face->IsBackOriented() ? EOrientation::Back : EOrientation::Front);
-		Face->SetHost(Shell);
+		TSharedPtr<FTopologicalFace> FacePtr = StaticCastSharedRef<FTopologicalFace>(Face->AsShared());
+		TopologicalFaces.Emplace(FacePtr, Face->IsBackOriented() ? EOrientation::Back : EOrientation::Front);
+		Face->SetHost(this);
 	}
 }
 
@@ -81,14 +80,14 @@ void FShell::Add(TSharedRef<FTopologicalFace> InTopologicalFace, EOrientation Or
 	TSharedPtr<FTopologicalFace> Face = InTopologicalFace;
 	TopologicalFaces.Emplace(Face, Orientation);
 
-	Face->SetHost(StaticCastSharedRef<FShell>(AsShared()));
+	Face->SetHost(this);
 }
 
 #ifdef CADKERNEL_DEV
 FInfoEntity& FShell::GetInfo(FInfoEntity& Info) const
 {
 	return FEntity::GetInfo(Info)
-		.Add(TEXT("Hosted by"), (TWeakPtr<FEntity>&) HostedBy)
+		.Add(TEXT("Hosted by"), (FEntity*) HostedBy)
 		.Add(TEXT("TopologicalFaces"), (TArray<TOrientedEntity<FEntity>>&) TopologicalFaces)
 		.Add(*this);
 }
@@ -165,10 +164,10 @@ void FShell::CheckTopology(TArray<FFaceSubset>& SubShells)
 
 	int32 ProcessFaceCount = 0;
 
-	TArray<TSharedPtr<FTopologicalFace>> Front;
-	TFunction<void(const TSharedPtr<FTopologicalFace>&, FFaceSubset&)> GetNeighboringFaces = [&](const TSharedPtr<FTopologicalFace>& Face, FFaceSubset& Shell)
+	TArray<FTopologicalFace*> Front;
+	TFunction<void(const FTopologicalFace&, FFaceSubset&)> GetNeighboringFaces = [&](const FTopologicalFace& Face, FFaceSubset& Shell)
 	{
-		for (const TSharedPtr<FTopologicalLoop>& Loop : Face->GetLoops())
+		for (const TSharedPtr<FTopologicalLoop>& Loop : Face.GetLoops())
 		{
 			for (const FOrientedEdge& OrientedEdge : Loop->GetEdges())
 			{
@@ -201,13 +200,8 @@ void FShell::CheckTopology(TArray<FFaceSubset>& SubShells)
 					}
 					NextEdge->SetMarker1();
 
-					TSharedPtr<FTopologicalFace> NextFace = NextEdge->GetFace();
-					if (!NextFace.IsValid())
-					{
-						continue;
-					}
-
-					if (NextFace->HasMarker1())
+					FTopologicalFace* NextFace = NextEdge->GetFace();
+					if ((NextFace == nullptr) || NextFace->HasMarker1())
 					{
 						continue;
 					}
@@ -222,13 +216,13 @@ void FShell::CheckTopology(TArray<FFaceSubset>& SubShells)
 	{
 		while (Front.Num())
 		{
-			TSharedPtr<FTopologicalFace> Face = Front.Pop();
+			FTopologicalFace* Face = Front.Pop();
 			Shell.Faces.Add(Face);
-			GetNeighboringFaces(Face, Shell);
+			GetNeighboringFaces(*Face, Shell);
 		}
 	};
 
-	for (const FOrientedFace& OrientedFace : GetFaces())
+	for (FOrientedFace& OrientedFace : GetFaces())
 	{
 		if (OrientedFace.Entity->HasMarker1())
 		{
@@ -239,7 +233,7 @@ void FShell::CheckTopology(TArray<FFaceSubset>& SubShells)
 		Shell.Faces.Reserve(TopologicalFaceCount - ProcessFaceCount);
 		Front.Empty(TopologicalFaceCount);
 
-		const TSharedPtr<FTopologicalFace>& Face = OrientedFace.Entity;
+		FTopologicalFace* Face = OrientedFace.Entity.Get();
 
 		Front.Empty(TopologicalFaceCount);
 		Face->SetMarker1();
diff --git a/Engine/Source/Runtime/Datasmith/CADKernel/Private/CADKernel/Topo/TopologicalEdge.cpp b/Engine/Source/Runtime/Datasmith/CADKernel/Private/CADKernel/Topo/TopologicalEdge.cpp
--- a/Engine/Source/Runtime/Datasmith/CADKernel/Private/CADKernel/Topo/TopologicalEdge.cpp
+++ b/Engine/Source/Runtime/Datasmith/CADKernel/Private/CADKernel/Topo/TopologicalEdge.cpp
@@ -264,13 +264,13 @@ void FTopologicalEdge::Delete()
 	SetDeleted();
 }
 
-TSharedPtr<FTopologicalFace> FTopologicalEdge::GetFace() const
+FTopologicalFace* FTopologicalEdge::GetFace() const
 {
 	if(Loop != nullptr)
 	{
 		return Loop->GetFace();
 	}
-	return TSharedPtr<FTopologicalFace>();
+	return nullptr;
 }
 
 void FTopologicalEdge::ComputeCrossingPointCoordinates()
diff --git a/Engine/Source/Runtime/Datasmith/CADKernel/Private/CADKernel/Topo/TopologicalFace.cpp b/Engine/Source/Runtime/Datasmith/CADKernel/Private/CADKernel/Topo/TopologicalFace.cpp
--- a/Engine/Source/Runtime/Datasmith/CADKernel/Private/CADKernel/Topo/TopologicalFace.cpp
+++ b/Engine/Source/Runtime/Datasmith/CADKernel/Private/CADKernel/Topo/TopologicalFace.cpp
@@ -10,6 +10,7 @@
 #include "CADKernel/Geo/Surfaces/Surface.h"
 #include "CADKernel/Mesh/Structure/FaceMesh.h"
 #include "CADKernel/Mesh/Structure/Grid.h"
+#include "CADKernel/Topo/Shell.h"
 #include "CADKernel/Topo/TopologicalEdge.h"
 #include "CADKernel/Topo/TopologyReport.h"
 
@@ -120,8 +121,7 @@ void FTopologicalFace::AddLoops(const TArray<TSharedPtr<FTopologicalLoop>>& InLo
 
 void FTopologicalFace::AddLoop(const TSharedPtr<FTopologicalLoop>& InLoop)
 {
-	TSharedRef<FTopologicalFace> Face = StaticCastSharedRef<FTopologicalFace>(AsShared());
-	InLoop->SetSurface(Face);
+	InLoop->SetSurface(this);
 	if (Loops.Num() > 0)
 	{
 		InLoop->SetAsInnerBoundary();
@@ -283,7 +283,7 @@ void FTopologicalFace::SpawnIdent(FDatabase& Database)
 FInfoEntity& FTopologicalFace::GetInfo(FInfoEntity& Info) const
 {
 	return FTopologicalEntity::GetInfo(Info)
-		.Add(TEXT("Hosted by"), (TWeakPtr<FEntity>&) HostedBy)
+		.Add(TEXT("Hosted by"), HostedBy)
 		.Add(TEXT("Carrier Surface"), CarrierSurface)
 		.Add(TEXT("Boundary"), (FSurfacicBoundary&) Boundary)
 		.Add(TEXT("Loops"), Loops)
@@ -440,7 +440,7 @@ void FTopologicalFace::DefineSurfaceType()
 					}
 				}
 
-				TSharedPtr<FTopologicalFace> Neighbor;
+				FTopologicalFace* Neighbor = nullptr;
 				{
 					for(FTopologicalEdge* NeighborEdge : Edge->GetTwinsEntities() )
 					{
@@ -452,7 +452,7 @@ void FTopologicalFace::DefineSurfaceType()
 					}
 				}
 
-				ensure(Neighbor.IsValid());
+				ensure(Neighbor != nullptr);
 
 				// it's not a quad surface
 				if (Neighbor->SurfaceCorners.Num() == 0)
@@ -516,4 +516,99 @@ void FTopologicalFace::DefineSurfaceType()
 	}
 }
 
+
+void FFaceSubset::SetMainShell(TMap<FShell*, int32>& ShellToFaceCount)
+{
+	if (ShellToFaceCount.Num() == 0)
+	{
+		return;
+	}
+
+	int32 MaxFaceCount = 0;
+	FShell* CandidateShell = nullptr;
+
+	for (TPair<FShell*, int32>& Pair : ShellToFaceCount)
+	{
+		if (Pair.Value > MaxFaceCount)
+		{
+			MaxFaceCount = Pair.Value;
+			CandidateShell = Pair.Key;
+		}
+	}
+
+	if ((CandidateShell != nullptr) && ((CandidateShell->FaceCount() / 2 + 1) < MaxFaceCount))
+	{
+		MainShell = CandidateShell;
+	}
+}
+
+void FFaceSubset::SetMainBody(TMap<FBody*, int32>& BodyToFaceCount)
+{
+	if (BodyToFaceCount.Num() == 0)
+	{
+		return;
+	}
+
+	FBody* CandidateBody = nullptr;
+	int32 MaxFaceCount = 0;
+	for (TPair<FBody*, int32>& Pair : BodyToFaceCount)
+	{
+		if (Pair.Value > MaxFaceCount)
+		{
+			MaxFaceCount = Pair.Value;
+			CandidateBody = Pair.Key;
+		}
+	}
+
+	// if Faces come mainly from MainBody
+	if ((CandidateBody != nullptr) && ((Faces.Num() / 2) <= MaxFaceCount))
+	{
+		MainBody = CandidateBody;
+	}
+}
+
+void FFaceSubset::SetMainName(TMap<FString, int32>& NameToFaceCount)
+{
+	if (NameToFaceCount.Num() == 0)
+	{
+		return;
+	}
+
+	int32 MaxInstance = Faces.Num() / 3;
+	for (TPair<FString, int32>& Pair : NameToFaceCount)
+	{
+		if (Pair.Value > MaxInstance)
+		{
+			MaxInstance = Pair.Value;
+			MainName = Pair.Key;
+		}
+	}
+}
+
+void FFaceSubset::SetMainColor(TMap<uint32, int32>& ColorToFaceCount)
+{
+	int32 MaxInstance = 0;
+	for (TPair<uint32, int32>& Pair : ColorToFaceCount)
+	{
+		if (Pair.Value > MaxInstance)
+		{
+			MaxInstance = Pair.Value;
+			MainColor = Pair.Key;
+		}
+	}
+}
+
+void FFaceSubset::SetMainMaterial(TMap<uint32, int32>& MaterialToFaceCount)
+{
+	int32 MaxInstance = 0;
+	for (TPair<uint32, int32>& Pair : MaterialToFaceCount)
+	{
+		if (Pair.Value > MaxInstance)
+		{
+			MaxInstance = Pair.Value;
+			MainColor = Pair.Key;
+		}
+	}
+}
+
 }
\ No newline at end of file
diff --git a/Engine/Source/Runtime/Datasmith/CADKernel/Private/CADKernel/Topo/TopologicalLoop.cpp b/Engine/Source/Runtime/Datasmith/CADKernel/Private/CADKernel/Topo/TopologicalLoop.cpp
--- a/Engine/Source/Runtime/Datasmith/CADKernel/Private/CADKernel/Topo/TopologicalLoop.cpp
+++ b/Engine/Source/Runtime/Datasmith/CADKernel/Private/CADKernel/Topo/TopologicalLoop.cpp
@@ -58,7 +58,8 @@ TSharedPtr<FTopologicalLoop> FTopologicalLoop::Make(const TArray<TSharedPtr<FTop
 }
 
 FTopologicalLoop::FTopologicalLoop(const TArray<TSharedPtr<FTopologicalEdge>>& InEdges, const TArray<EOrientation>& InEdgeDirections)
-	: bExternalLoop(true)
+	: Face(nullptr)
+	, bExternalLoop(true)
 {
 	Edges.Reserve(InEdges.Num());
 	for (int32 Index = 0; Index < InEdges.Num(); ++Index)
@@ -151,7 +152,7 @@ bool FTopologicalLoop::Get2DSamplingWithoutDegeneratedEdges(TArray<FPoint2D>& Lo
 
 	double LoopMeanLength = LoopLength / EdgeCount;
 	double MinEdgeLength = LoopMeanLength * 0.01;
-	MinEdgeLength = FMath::Max(MinEdgeLength, Face.Pin()->GetCarrierSurface()->Get3DTolerance());
+	MinEdgeLength = FMath::Max(MinEdgeLength, Face->GetCarrierSurface()->Get3DTolerance());
 
 	LoopSampling.Empty(PointCount);
 
@@ -462,7 +463,7 @@ bool FTopologicalLoop::Orient()
 #endif
 
 		bSucceed = false;
-		FMessage::Printf(Log, TEXT("WARNING: Loop Orientation of surface %d is doubtful\n"), Face.Pin()->GetId());
+		FMessage::Printf(Log, TEXT("WARNING: Loop Orientation of surface %d is doubtful\n"), Face->GetId());
 	}
 
 	if ((WrongOrientationCount > GoodOrientationCount) == bExternalLoop)
diff --git a/Engine/Source/Runtime/Datasmith/CADKernel/Public/CADKernel/Core/System.h b/Engine/Source/Runtime/Datasmith/CADKernel/Public/CADKernel/Core/System.h
--- a/Engine/Source/Runtime/Datasmith/CADKernel/Public/CADKernel/Core/System.h
+++ b/Engine/Source/Runtime/Datasmith/CADKernel/Public/CADKernel/Core/System.h
@@ -55,8 +55,8 @@ namespace CADKernel
 
 		void CloseLogFiles();
 
-		FString GetToolkitVersion() const;
-		FString GetCompilationDate() const;
+		static FString GetToolkitVersion();
+		static FString GetCompilationDate();
 
 		FVisu* GetVisu() const
 		{
diff --git a/Engine/Source/Runtime/Datasmith/CADKernel/Public/CADKernel/Geo/Surfaces/Surface.h b/Engine/Source/Runtime/Datasmith/CADKernel/Public/CADKernel/Geo/Surfaces/Surface.h
--- a/Engine/Source/Runtime/Datasmith/CADKernel/Public/CADKernel/Geo/Surfaces/Surface.h
+++ b/Engine/Source/Runtime/Datasmith/CADKernel/Public/CADKernel/Geo/Surfaces/Surface.h
@@ -226,9 +226,9 @@ namespace CADKernel
 
 		void PresampleIsoCircle(const FSurfacicBoundary& InBoundaries, FCoordinateGrid& Coordinates, EIso Iso)
 		{
-			constexpr const double ThreeOverPi = 3 / PI;
+			constexpr const double SixOverPi = 6 / PI;
 			double Delta = InBoundaries.Length(Iso);
-			int32 SampleCount = /*FMath::Max(3, */ (int32)(Delta * ThreeOverPi + 1)/*)*/;
+			int32 SampleCount = /*FMath::Max(3, */ (int32)(Delta * SixOverPi + 1)/*)*/;
 			Delta /= SampleCount;
 
 			Coordinates[Iso].Empty(SampleCount + 1);
diff --git a/Engine/Source/Runtime/Datasmith/CADKernel/Public/CADKernel/Mesh/Meshers/ParametricMesher.h b/Engine/Source/Runtime/Datasmith/CADKernel/Public/CADKernel/Mesh/Meshers/ParametricMesher.h
--- a/Engine/Source/Runtime/Datasmith/CADKernel/Public/CADKernel/Mesh/Meshers/ParametricMesher.h
+++ b/Engine/Source/Runtime/Datasmith/CADKernel/Public/CADKernel/Mesh/Meshers/ParametricMesher.h
@@ -10,131 +10,130 @@
 
 namespace CADKernel
 {
-	class FCriterion;
-	class FGrid;
-	class FThinZoneSide;
-	class FTopologicalEntity;
-	class FTopologicalLoop;
+class FCriterion;
+class FGrid;
+class FThinZoneSide;
+class FTopologicalEntity;
+class FTopologicalLoop;
 
-	struct FCostToFace
-	{
-		double Cost;
-		TSharedRef<FTopologicalFace> Face;
-
-		FCostToFace(double NewCost, TSharedRef<FTopologicalFace> NewFace)
-			: Cost(NewCost)
-			, Face(NewFace)
-		{
-		}
-
-	};
+struct FCostToFace
+{
+	double Cost;
+	FTopologicalFace* Face;
 
-	class CADKERNEL_API FMesherParameters : public FParameters
+	FCostToFace(double NewCost, FTopologicalFace* NewFace)
+		: Cost(NewCost)
+		, Face(NewFace)
 	{
-	public:
-		FParameter InconsistencyAngle;
+	}
+};
 
-		FMesherParameters();
+class CADKERNEL_API FMesherParameters : public FParameters
+{
+public:
+	FParameter InconsistencyAngle;
 
-		void SetInconsistencyAngle(const double Value)
-		{
-			InconsistencyAngle = Value;
-		}
-	};
+	FMesherParameters();
 
-	class CADKERNEL_API FParametricMesher
+	void SetInconsistencyAngle(const double Value)
 	{
-	protected:
+		InconsistencyAngle = Value;
+	}
+};
 
-		/**
-		 * Limit of flatness of quad face
-		 */
-		const double ConstMinCurvature = 0.001;
+class CADKERNEL_API FParametricMesher
+{
+protected:
 
-		TSharedRef<FModelMesh> MeshModel;
-		TSharedRef<FMesherParameters> Parameters;
+	/**
+	 * Limit of flatness of quad face
+	 */
+	const double ConstMinCurvature = 0.001;
 
-		TArray<TSharedPtr<FTopologicalFace>> Faces;
-		TArray<TSharedPtr<FTopologicalEdge>> Edges;
-		TArray<TSharedPtr<FTopologicalVertex>> Vertices;
+	TSharedRef<FModelMesh> MeshModel;
+	TSharedRef<FMesherParameters> Parameters;
 
-		FMesherReport MesherReport;
+	TArray<TSharedPtr<FTopologicalFace>> Faces;
+	TArray<TSharedPtr<FTopologicalEdge>> Edges;
+	TArray<TSharedPtr<FTopologicalVertex>> Vertices;
 
-		bool bDisplay = false;
+	FMesherReport MesherReport;
 
-	public:
+	bool bDisplay = false;
 
-		FParametricMesher(TSharedRef<FModelMesh> MeshModel);
+public:
 
-		const TSharedRef<FModelMesh>& GetMeshModel() const
-		{
-			return MeshModel;
-		}
+	FParametricMesher(TSharedRef<FModelMesh> MeshModel);
 
-		TSharedRef<FModelMesh>& GetMeshModel()
-		{
-			return MeshModel;
-		}
+	const TSharedRef<FModelMesh>& GetMeshModel() const
+	{
+		return MeshModel;
+	}
 
-		void InitParameters(const FString& ParametersString);
+	TSharedRef<FModelMesh>& GetMeshModel()
+	{
+		return MeshModel;
+	}
+
+	void InitParameters(const FString& ParametersString);
 
-		const TSharedRef<FMesherParameters>& GetParameters() const
-		{
-			return Parameters;
-		}
+	const TSharedRef<FMesherParameters>& GetParameters() const
+	{
+		return Parameters;
+	}
 
-		void MeshEntities(TArray<TSharedPtr<FEntity>>& InEntities);
+	void MeshEntities(TArray<TSharedPtr<FEntity>>& InEntities);
 
-		template<typename EntityType>
-		void MeshEntity(TSharedRef<EntityType>& Entity)
-		{
-			MeshEntity((TSharedRef<FTopologicalEntity>&) Entity);
-		}
+	template<typename EntityType>
+	void MeshEntity(TSharedRef<EntityType>& Entity)
+	{
+		MeshEntity((TSharedRef<FTopologicalEntity>&) Entity);
+	}
 
-		void MeshEntity(TSharedRef<FModel>& InModel)
-		{
-			MeshEntity((TSharedRef<FTopologicalEntity>&) InModel);
-		}
+	void MeshEntity(TSharedRef<FModel>& InModel)
+	{
+		MeshEntity((TSharedRef<FTopologicalEntity>&) InModel);
+	}
 
-		void MeshEntity(TSharedRef<FTopologicalEntity>& InEntity)
-		{
-			TArray<TSharedPtr<FEntity>> Entities;
-			Entities.Add(InEntity);
-			MeshEntities(Entities);
-		}
+	void MeshEntity(TSharedRef<FTopologicalEntity>& InEntity)
+	{
+		TArray<TSharedPtr<FEntity>> Entities;
+		Entities.Add(InEntity);
+		MeshEntities(Entities);
+	}
 
-		void Mesh(TSharedRef<FTopologicalFace> Face);
-		void Mesh(FTopologicalEdge& InEdge, FTopologicalFace& CarrierFace);
-		void Mesh(TSharedRef<FTopologicalVertex> Vertex);
+	void Mesh(FTopologicalFace& Face);
+	void Mesh(FTopologicalEdge& InEdge, const FTopologicalFace& CarrierFace);
+	void Mesh(FTopologicalVertex& Vertex);
 
-		void MeshFaceLoops(FGrid& Grid);
+	void MeshFaceLoops(FGrid& Grid);
 
-		void MeshThinZoneEdges(FGrid&);
-		void MeshThinZoneSide(const FThinZoneSide& Side);
-		void GetThinZoneBoundary(const FThinZoneSide& Side);
+	void MeshThinZoneEdges(FGrid&);
+	void MeshThinZoneSide(const FThinZoneSide& Side);
+	void GetThinZoneBoundary(const FThinZoneSide& Side);
 
-		void GenerateCloud(FGrid& Grid);
+	void GenerateCloud(FGrid& Grid);
 
-		void PrintReport();
+	void PrintReport();
 
-	protected:
+protected:
 
-		void MeshEntities();
+	void MeshEntities();
 
 
-		void IsolateQuadFace(TArray<FCostToFace>& QuadSurfaces, TArray<TSharedPtr<FTopologicalFace>>& OtherSurfaces) const;
+	void IsolateQuadFace(TArray<FCostToFace>& QuadSurfaces, TArray<FTopologicalFace*>& OtherSurfaces) const;
 
-		void LinkQuadSurfaceForMesh(TArray<FCostToFace>& QuadTrimmedSurfaceSet, TArray<TArray<TSharedPtr<FTopologicalFace>>>& OutStrips);
-		void MeshSurfaceByFront(TArray<FCostToFace>& QuadTrimmedSurfaceSet);
+	void LinkQuadSurfaceForMesh(TArray<FCostToFace>& QuadTrimmedSurfaceSet, TArray<TArray<FTopologicalFace*>>& OutStrips);
+	void MeshSurfaceByFront(TArray<FCostToFace>& QuadTrimmedSurfaceSet);
 
-		void ApplyEdgeCriteria(FTopologicalEdge& Edge);
-		void ApplyFaceCriteria(TSharedRef<FTopologicalFace> Face);
+	void ApplyEdgeCriteria(FTopologicalEdge& Edge);
+	void ApplyFaceCriteria(TSharedRef<FTopologicalFace> Face);
 
-		/**
-		 * Generate Edge Elements on active edge from Edge cutting points
-		 */
-		void GenerateEdgeElements(FTopologicalEdge& Edge);
-	};
+	/**
+	 * Generate Edge Elements on active edge from Edge cutting points
+	 */
+	void GenerateEdgeElements(FTopologicalEdge& Edge);
+};
 
 } // namespace CADKernel
 
diff --git a/Engine/Source/Runtime/Datasmith/CADKernel/Public/CADKernel/Topo/Joiner.h b/Engine/Source/Runtime/Datasmith/CADKernel/Public/CADKernel/Topo/Joiner.h
--- a/Engine/Source/Runtime/Datasmith/CADKernel/Public/CADKernel/Topo/Joiner.h
+++ b/Engine/Source/Runtime/Datasmith/CADKernel/Public/CADKernel/Topo/Joiner.h
@@ -5,120 +5,126 @@
 
 namespace CADKernel
 {
-	class FModel;
-	class FSession;
-	class FShell;
-	class FTopologicalEdge;
-	class FTopologicalFace;
-	class FTopologicalVertex;
-
-	class CADKERNEL_API FJoiner
-	{
-		friend class FThinZone;
-		friend class FThinZoneFinder;
-
-	protected:
-
-		FSession& Session;
-
-		TArray<TSharedPtr<FShell>> Shells;
-		TArray<TSharedPtr<FTopologicalFace>> Faces;
-
-		double JoiningTolerance;
-		double JoiningToleranceSquare;
-
-	public:
-
-		FJoiner(FSession& InSession, const TSharedRef<FModel>& Shells, double Tolerance);
-		FJoiner(FSession& InSession, const TArray<TSharedPtr<FShell>>& Shells, double Tolerance);
-		FJoiner(FSession& InSession, const TArray<TSharedPtr<FTopologicalFace>>& Surfaces, double Tolerance);
-
-		void JoinFaces();
-		//void JoinFaces(bool bProcessOnlyBorderEdges, bool bProcessOnlyNonManifoldEdges);
-
-		/**
-		 * Check topology of each body
-		 */
-		void CheckTopology();
-
-		//void MergeInto(TSharedPtr<FBody> Body, TArray<TSharedPtr<FTopologicalEntity>>& InEntities);
-		//void SortByShell(TSharedPtr<FBody> Body, TArray<TSharedPtr<FBody>>& OutNewBody);
-		//void Join(TArray<TSharedPtr<FBody>> Bodies, double Tolerance);
-
-		/**
-		 * Split into connected shell and put each shell into the appropriate body
-		 */
-		void SplitIntoConnectedShell();
-
-	private:
-
-		/**
-		 * Call by constructor.
-		 * For each shell, add their faces into Faces array, complete the metadata and set the states for the joining process
-		 */
-		void InitFaces();
-
-		void EmptyShells();
-
-		void RemoveFacesFromShell();
-
-		/**
-		 * Return an array of active vertices.
-		 */
-		void GetVertices(TArray<TSharedPtr<FTopologicalVertex>>& Vertices);
-
-		/**
-		 * Return an array of active border vertices.
-		 */
-		void GetBorderVertices(TArray<TSharedPtr<FTopologicalVertex>>& BorderVertices);
-
-		/**
-		 * Merge Border Vertices with other vertices.
-		 * @param Vertices: the initial array of active vertices, this array is updated at the end of the process
-		 */
-		void MergeCoincidentVertices(TArray<TSharedPtr<FTopologicalVertex>>& VerticesToMerge);
-
-		/**
-		 * Merge Border Vertices with other vertices.
-		 * @param Vertices: the initial array of active vertices, this array is updated at the end of the process
-		 */
-		void MergeBorderVerticesWithCoincidentOtherVertices(TArray<TSharedPtr<FTopologicalVertex>>& Vertices);
-
-		//void FixCollapsedEdges();
-		void CheckSelfConnectedEdge();
-		void RemoveIsolatedEdges(); // Usefull ???
-
-		/**
-		 * First step, trivial edge merge i.e. couple of edges with same extremity vertex
-		 */
-		void MergeCoincidentEdges(TArray<TSharedPtr<FTopologicalVertex>>& Vertices);
-
-		/**
-		 * Second step, parallel edges but with different length. the longest must be split
-		 */
-		void StitchParallelEdges(TArray<TSharedPtr<FTopologicalVertex>>& Vertices);
-
-		/**
-		 * Split the edge to split except if one side is too small (JoiningTolerance). In this case (too small), the EdgeToSplit is not split but linked to EdgeToLink. 
-		 * @return the created vertex or TSharedPtr<FTopologicalVertex>()
-		 */
-		TSharedPtr<FTopologicalVertex> SplitAndLink(FTopologicalVertex& StartVertex, FTopologicalEdge& EdgeToLink, FTopologicalEdge& EdgeToSplit);
-
-		/**
-		 * For each loop of each surface, check if successive edges are unconnected and if their common vertices are connected only to them. 
-		 * These edges are merged into one edge. 
-		 * E.g. :
-		 * Face A has 3 successive unconnected edges. If these 3 edges are merged to give only one edge, the new edge could be linked to its parallel edge of Face B 
-		 * 
-		 *              \                          Face A                                   |    Face C
-		 *               \                                                                  |
-		 *    Face E     CV ------------------- UV ------------------ UV ----------------- CV --------------------
-		 *               CV -------------------------------------------------------------- CV --------------------
-		 *               /                         Face B                                   |    Face D
-		 *              /                                                                   |
-		 */
-		void MergeUnconnectedAdjacentEdges();
-
-	};
+
+class FModel;
+class FSession;
+class FShell;
+class FTopologicalEdge;
+class FTopologicalFace;
+class FTopologicalVertex;
+
+class CADKERNEL_API FJoiner
+{
+	friend class FThinZone;
+	friend class FThinZoneFinder;
+
+protected:
+
+	FSession& Session;
+
+	TArray<TSharedPtr<FShell>> Shells;
+	TArray<TSharedPtr<FTopologicalFace>> Faces;
+
+	double JoiningTolerance;
+	double JoiningToleranceSquare;
+
+public:
+
+	FJoiner(FSession& InSession, const TSharedRef<FModel>& Shells, double Tolerance);
+	FJoiner(FSession& InSession, const TArray<TSharedPtr<FShell>>& Shells, double Tolerance);
+	FJoiner(FSession& InSession, const TArray<TSharedPtr<FTopologicalFace>>& Surfaces, double Tolerance);
+
+	void JoinFaces();
+	//void JoinFaces(bool bProcessOnlyBorderEdges, bool bProcessOnlyNonManifoldEdges);
+
+	/**
+	 * Check topology of each body
+	 */
+	void CheckTopology();
+
+	//void MergeInto(TSharedPtr<FBody> Body, TArray<TSharedPtr<FTopologicalEntity>>& InEntities);
+	//void SortByShell(TSharedPtr<FBody> Body, TArray<TSharedPtr<FBody>>& OutNewBody);
+	//void Join(TArray<TSharedPtr<FBody>> Bodies, double Tolerance);
+
+	/**
+	 * Split into connected shell and put each shell into the appropriate body
+	 */
+	void SplitIntoConnectedShell();
+
+	/**
+	 * Unlink Non-Manifold Vertex i.e. Vertex belong to tow or more shell
+	 */
+	void UnlinkNonManifoldVertex();
+
+private:
+
+	/**
+	 * Call by constructor.
+	 * For each shell, add their faces into Faces array, complete the metadata and set the states for the joining process
+	 */
+	void InitFaces();
+
+	void EmptyShells();
+
+	void RemoveFacesFromShell();
+
+	/**
+	 * Return an array of active vertices.
+	 */
+	void GetVertices(TArray<FTopologicalVertex*>& Vertices);
+
+	/**
+	 * Return an array of active border vertices.
+	 */
+	void GetBorderVertices(TArray<FTopologicalVertex*>& BorderVertices);
+
+	/**
+	 * Merge Border Vertices with other vertices.
+	 * @param Vertices: the initial array of active vertices, this array is updated at the end of the process
+	 */
+	void MergeCoincidentVertices(TArray<FTopologicalVertex*>& VerticesToMerge);
+
+	/**
+	 * Merge Border Vertices with other vertices.
+	 * @param Vertices: the initial array of active vertices, this array is updated at the end of the process
+	 */
+	void MergeBorderVerticesWithCoincidentOtherVertices(TArray<FTopologicalVertex*>& Vertices);
+
+	//void FixCollapsedEdges();
+	void CheckSelfConnectedEdge();
+	void RemoveIsolatedEdges(); // Useful ???
+
+	/**
+	 * First step, trivial edge merge i.e. couple of edges with same extremity vertex
+	 */
+	void MergeCoincidentEdges(TArray<FTopologicalVertex*>& Vertices);
+
+	/**
+	 * Second step, parallel edges but with different length. the longest must be split
+	 */
+	void StitchParallelEdges(TArray<FTopologicalVertex*>& Vertices);
+
+	/**
+	 * Split the edge to split except if one side is too small (JoiningTolerance). In this case (too small), the EdgeToSplit is not split but linked to EdgeToLink.
+	 * @return the created vertex or TSharedPtr<FTopologicalVertex>()
+	 */
+	TSharedPtr<FTopologicalVertex> SplitAndLink(FTopologicalVertex& StartVertex, FTopologicalEdge& EdgeToLink, FTopologicalEdge& EdgeToSplit);
+
+	/**
+	 * For each loop of each surface, check if successive edges are unconnected and if their common vertices are connected only to them.
+	 * These edges are merged into one edge.
+	 * E.g. :
+	 * Face A has 3 successive unconnected edges. If these 3 edges are merged to give only one edge, the new edge could be linked to its parallel edge of Face B
+	 *
+	 *              \                          Face A                                   |    Face C
+	 *               \                                                                  |
+	 *    Face E     CV ------------------- UV ------------------ UV ----------------- CV --------------------
+	 *               CV -------------------------------------------------------------- CV --------------------
+	 *               /                         Face B                                   |    Face D
+	 *              /                                                                   |
+	 */
+	void MergeUnconnectedAdjacentEdges();
+
+};
 
 } // namespace CADKernel
diff --git a/Engine/Source/Runtime/Datasmith/CADKernel/Public/CADKernel/Topo/Linkable.h b/Engine/Source/Runtime/Datasmith/CADKernel/Public/CADKernel/Topo/Linkable.h
--- a/Engine/Source/Runtime/Datasmith/CADKernel/Public/CADKernel/Topo/Linkable.h
+++ b/Engine/Source/Runtime/Datasmith/CADKernel/Public/CADKernel/Topo/Linkable.h
@@ -114,6 +114,14 @@ public:
 		TopologicalLink.Reset();
 	}
 
+	/**
+	 * Unlink all twin entities
+	 */
+	void UnlinkTwinEntities()
+	{
+		TopologicalLink->UnlinkTwinEntities();
+	}
+
 	const bool IsThinZone() const
 	{
 		return ((States & EHaveStates::ThinZone) == EHaveStates::ThinZone);
diff --git a/Engine/Source/Runtime/Datasmith/CADKernel/Public/CADKernel/Topo/Shell.h b/Engine/Source/Runtime/Datasmith/CADKernel/Public/CADKernel/Topo/Shell.h
--- a/Engine/Source/Runtime/Datasmith/CADKernel/Public/CADKernel/Topo/Shell.h
+++ b/Engine/Source/Runtime/Datasmith/CADKernel/Public/CADKernel/Topo/Shell.h
@@ -10,165 +10,170 @@
 
 namespace CADKernel
 {
-	class FTopologicalFace;
-	class FBody;
+class FTopologicalFace;
+class FBody;
 
-	struct FFaceSubset;
+struct FFaceSubset;
 
-	class CADKERNEL_API FOrientedFace : public TOrientedEntity<FTopologicalFace>
+class CADKERNEL_API FOrientedFace : public TOrientedEntity<FTopologicalFace>
+{
+public:
+	FOrientedFace(TSharedPtr<FTopologicalFace>& InEntity, EOrientation InOrientation)
+		: TOrientedEntity(InEntity, InOrientation)
 	{
-	public:
-		FOrientedFace(TSharedPtr<FTopologicalFace>& InEntity, EOrientation InOrientation)
-			: TOrientedEntity(InEntity, InOrientation)
-		{
-		}
+	}
 
-		FOrientedFace(const FOrientedFace& OrientiredEntity)
-			: TOrientedEntity(OrientiredEntity)
-		{
-		}
-
-		FOrientedFace()
-			: TOrientedEntity()
-		{
-		}
-	};
+	FOrientedFace(const FOrientedFace& OrientiredEntity)
+		: TOrientedEntity(OrientiredEntity)
+	{
+	}
 
-	class CADKERNEL_API FShell : public FTopologicalEntity, public FMetadataDictionary
+	FOrientedFace()
+		: TOrientedEntity()
 	{
-		friend FEntity;
-		friend class FBody;
+	}
+};
+
+class CADKERNEL_API FShell : public FTopologicalEntity, public FMetadataDictionary
+{
+	friend FEntity;
+	friend class FBody;
 
-	private:
-		TArray<FOrientedFace> TopologicalFaces;
-		TWeakPtr<FBody> HostedBy;
+private:
+	TArray<FOrientedFace> TopologicalFaces;
+	FBody* HostedBy;
 
-		FShell() = default;
+	FShell() = default;
 
-		FShell(const TArray<FOrientedFace> InTopologicalFaces, bool bIsInnerShell = false)
-			: FTopologicalEntity()
-			, TopologicalFaces(InTopologicalFaces)
+	FShell(const TArray<FOrientedFace> InTopologicalFaces, bool bIsInnerShell = false)
+		: FTopologicalEntity()
+		, TopologicalFaces(InTopologicalFaces)
+	{
+		if (bIsInnerShell)
 		{
-			if (bIsInnerShell)
-			{
-				SetInner();
-			}
+			SetInner();
 		}
+	}
 
-		FShell(const TArray<TSharedPtr<FTopologicalFace>>& InTopologicalFaces, bool bIsInnerShell = true);
+	FShell(const TArray<TSharedPtr<FTopologicalFace>>& InTopologicalFaces, bool bIsInnerShell = true);
 
-		FShell(const TArray<TSharedPtr<FTopologicalFace>>& InTopologicalFaces, const TArray<EOrientation>& InOrientations, bool bIsInnerShell = true);
+	FShell(const TArray<TSharedPtr<FTopologicalFace>>& InTopologicalFaces, const TArray<EOrientation>& InOrientations, bool bIsInnerShell = true);
 
 
-	public:
+public:
 
-		virtual void Serialize(FCADKernelArchive& Ar) override
-		{
-			FTopologicalEntity::Serialize(Ar);
-			SerializeIdents(Ar, (TArray<TOrientedEntity<FEntity>>&) TopologicalFaces);
-			SerializeIdent(Ar, HostedBy);
-			FMetadataDictionary::Serialize(Ar);
-		}
+	virtual void Serialize(FCADKernelArchive& Ar) override
+	{
+		FTopologicalEntity::Serialize(Ar);
+		SerializeIdents(Ar, (TArray<TOrientedEntity<FEntity>>&) TopologicalFaces);
+		SerializeIdent(Ar, &HostedBy);
+		FMetadataDictionary::Serialize(Ar);
+	}
 
-		virtual void SpawnIdent(FDatabase& Database) override
+	virtual void SpawnIdent(FDatabase& Database) override
+	{
+		if (!FEntity::SetId(Database))
 		{
-			if (!FEntity::SetId(Database))
-			{
-				return;
-			}
-
-			SpawnIdentOnEntities((TArray<TOrientedEntity<FEntity>>&) TopologicalFaces, Database);
+			return;
 		}
 
-		virtual void ResetMarkersRecursively() override
-		{
-			ResetMarkers();
-			ResetMarkersRecursivelyOnEntities((TArray<TOrientedEntity<FEntity>>&) TopologicalFaces);
-		}
+		SpawnIdentOnEntities((TArray<TOrientedEntity<FEntity>>&) TopologicalFaces, Database);
+	}
 
-		TWeakPtr<FBody>& GetHost()
-		{
-			return HostedBy;
-		}
+	virtual void ResetMarkersRecursively() override
+	{
+		ResetMarkers();
+		ResetMarkersRecursivelyOnEntities((TArray<TOrientedEntity<FEntity>>&) TopologicalFaces);
+	}
 
-		void SetHost(TSharedPtr<FBody> Body)
-		{
-			HostedBy = Body;
-		}
+	FBody* GetHost()
+	{
+		return HostedBy;
+	}
+
+	void SetHost(FBody* Body)
+	{
+		HostedBy = Body;
+	}
 
-		void CompleteMetadata()
+	void CompleteMetadata()
+	{
+		if (HostedBy != nullptr)
 		{
-			if (HostedBy.IsValid())
-			{
-				CompleteDictionary((const FMetadataDictionary&) * HostedBy.Pin());
-			}
+			CompleteDictionary((const FMetadataDictionary&)*HostedBy);
 		}
+	}
 
-		void Empty(int32 NewSize = 0);
+	void Empty(int32 NewSize = 0);
 
-		void Add(TSharedRef<FTopologicalFace> InTopologicalFace, EOrientation InOrientation);
-		void Add(TArray<TSharedPtr<FTopologicalFace>> Faces);
+	void Add(TSharedRef<FTopologicalFace> InTopologicalFace, EOrientation InOrientation);
+	void Add(TArray<FTopologicalFace*> Faces);
 
 #ifdef CADKERNEL_DEV
-		virtual FInfoEntity& GetInfo(FInfoEntity&) const override;
+	virtual FInfoEntity& GetInfo(FInfoEntity&) const override;
 #endif
 
-		virtual EEntity GetEntityType() const override
-		{
-			return EEntity::Shell;
-		}
+	virtual EEntity GetEntityType() const override
+	{
+		return EEntity::Shell;
+	}
 
-		TSharedPtr<FEntityGeom> ApplyMatrix(const FMatrixH& InMatrix) const;
+	TSharedPtr<FEntityGeom> ApplyMatrix(const FMatrixH& InMatrix) const;
 
-		virtual int32 FaceCount() const override
-		{
-			return TopologicalFaces.Num();
-		}
+	virtual int32 FaceCount() const override
+	{
+		return TopologicalFaces.Num();
+	}
 
-		void ReplaceFaces(TArray<FOrientedFace>& NewFaces)
-		{
-			Swap(TopologicalFaces, NewFaces);
-			NewFaces.Empty();
-		}
+	void ReplaceFaces(TArray<FOrientedFace>& NewFaces)
+	{
+		Swap(TopologicalFaces, NewFaces);
+		NewFaces.Empty();
+	}
 
-		const TArray<FOrientedFace>& GetFaces() const
-		{
-			return TopologicalFaces;
-		}
+	const TArray<FOrientedFace>& GetFaces() const
+	{
+		return TopologicalFaces;
+	}
 
-		virtual void GetFaces(TArray<TSharedPtr<FTopologicalFace>>& OutFaces) override;
+	TArray<FOrientedFace>& GetFaces()
+	{
+		return TopologicalFaces;
+	}
 
-		virtual void Merge(TSharedPtr<FShell>& Shell);
+	virtual void GetFaces(TArray<TSharedPtr<FTopologicalFace>>& OutFaces) override;
 
-		virtual void SpreadBodyOrientation() override;
+	virtual void Merge(TSharedPtr<FShell>& Shell);
 
-		void CheckTopology(TArray<FFaceSubset>& SubShells);
-		virtual void FillTopologyReport(FTopologyReport& Report) const override;
+	virtual void SpreadBodyOrientation() override;
 
-		/**
-		 * @return true if the shell has at least one border edge 
-		 */
-		bool IsOpenShell();
+	void CheckTopology(TArray<FFaceSubset>& SubShells);
+	virtual void FillTopologyReport(FTopologyReport& Report) const override;
 
-		bool IsInner() const
-		{
-			return ((States & EHaveStates::IsInner) == EHaveStates::IsInner);
-		}
+	/**
+	 * @return true if the shell has at least one border edge
+	 */
+	bool IsOpenShell();
 
-		bool IsOutter() const
-		{
-			return ((States & EHaveStates::IsInner) != EHaveStates::IsInner);
-		}
+	bool IsInner() const
+	{
+		return ((States & EHaveStates::IsInner) == EHaveStates::IsInner);
+	}
 
-		void SetInner()
-		{
-			States |= EHaveStates::IsInner;
-		}
+	bool IsOutter() const
+	{
+		return ((States & EHaveStates::IsInner) != EHaveStates::IsInner);
+	}
 
-		void SetOutter()
-		{
-			States &= ~EHaveStates::IsInner;
-		}
-	};
+	void SetInner()
+	{
+		States |= EHaveStates::IsInner;
+	}
+
+	void SetOutter()
+	{
+		States &= ~EHaveStates::IsInner;
+	}
+};
 }
 
diff --git a/Engine/Source/Runtime/Datasmith/CADKernel/Public/CADKernel/Topo/TopologicalEdge.h b/Engine/Source/Runtime/Datasmith/CADKernel/Public/CADKernel/Topo/TopologicalEdge.h
--- a/Engine/Source/Runtime/Datasmith/CADKernel/Public/CADKernel/Topo/TopologicalEdge.h
+++ b/Engine/Source/Runtime/Datasmith/CADKernel/Public/CADKernel/Topo/TopologicalEdge.h
@@ -11,875 +11,876 @@
 
 namespace CADKernel
 {
-	typedef TTopologicalLink<FTopologicalEdge> FEdgeLink;
 
-	class FOrientedEdge;
+typedef TTopologicalLink<FTopologicalEdge> FEdgeLink;
 
-	enum ECoordinateType : uint8
-	{
-		VertexCoordinate,
-		ImposedCoordinate, // imposted cutting point by thin zone meshing process
-		IsoUCoordinate,
-		IsoVCoordinate,
-		IsoUVCoordinate,
-		OtherCoordinate
-	};
-
-	template<typename FCuttingPointType>
-	void GetCuttingPointCoordinates(const TArray<FCuttingPointType>& CuttingPoints, TArray<double>& CuttingPointCoordinates);
-
-	struct FEdge2DProperties;
-	struct FCuttingPoint;
-	struct FImposedCuttingPoint;
-
-	class FModelMesh;
-	class FRestrictionCurve;
-	class FSurface;
-	class FTopologicalLoop;
-	class FTopologicalVertex;
-
-	class CADKERNEL_API FTopologicalEdge : public TLinkable<FTopologicalEdge, FEdgeLink>
-	{
-		friend class FEntity;
-		friend class FTopologicalLoop;
-		friend class FTopologicalFace;
-
-	protected:
+class FOrientedEdge;
 
-		TSharedPtr<FTopologicalVertex> StartVertex;
-		TSharedPtr<FTopologicalVertex> EndVertex;
+enum ECoordinateType : uint8
+{
+	VertexCoordinate,
+	ImposedCoordinate, // imposted cutting point by thin zone meshing process
+	IsoUCoordinate,
+	IsoVCoordinate,
+	IsoUVCoordinate,
+	OtherCoordinate
+};
+
+template<typename FCuttingPointType>
+void GetCuttingPointCoordinates(const TArray<FCuttingPointType>& CuttingPoints, TArray<double>& CuttingPointCoordinates);
+
+struct FEdge2DProperties;
+struct FCuttingPoint;
+struct FImposedCuttingPoint;
+
+class FModelMesh;
+class FRestrictionCurve;
+class FSurface;
+class FTopologicalLoop;
+class FTopologicalVertex;
+
+class CADKERNEL_API FTopologicalEdge : public TLinkable<FTopologicalEdge, FEdgeLink>
+{
+	friend class FEntity;
+	friend class FTopologicalLoop;
+	friend class FTopologicalFace;
 
-		/**
-		 * The edge is oriented in the curve orientation i.e. StartCoordinate < EndCoordinate
-		 */
-		FLinearBoundary Boundary;
+protected:
 
-		TSharedPtr<FRestrictionCurve> Curve;
- 		mutable double Length3D = -1.;
+	TSharedPtr<FTopologicalVertex> StartVertex;
+	TSharedPtr<FTopologicalVertex> EndVertex;
 
-		FTopologicalLoop* Loop = nullptr;
+	/**
+	 * The edge is oriented in the curve orientation i.e. StartCoordinate < EndCoordinate
+	 */
+	FLinearBoundary Boundary;
 
-		TSharedPtr<FEdgeMesh> Mesh;
+	TSharedPtr<FRestrictionCurve> Curve;
+	mutable double Length3D = -1.;
 
-		/**
-		 * Final U coordinates of the edge's mesh nodes
-		 */
-		TArray<FCuttingPoint> CuttingPointUs;
+	FTopologicalLoop* Loop = nullptr;
 
-		/**
-		 * U coordinates of the edge's mesh nodes for thin zone purpose
-		 */
-		TArray<FImposedCuttingPoint> ImposedCuttingPointUs;
-		TArray<FLinearBoundary> ThinZoneBounds;
+	TSharedPtr<FEdgeMesh> Mesh;
 
-		/**
-		 * Temporary discretization of the edge used to compute the mesh of the edge
-		 */
-		TArray<double> CrossingPointUs;
+	/**
+	 * Final U coordinates of the edge's mesh nodes
+	 */
+	TArray<FCuttingPoint> CuttingPointUs;
 
-		/**
-		 * Min delta U at the crossing points to respect meshing criteria
-		 */
-		TArray<double> CrossingPointDeltaUMins;
+	/**
+	 * U coordinates of the edge's mesh nodes for thin zone purpose
+	 */
+	TArray<FImposedCuttingPoint> ImposedCuttingPointUs;
+	TArray<FLinearBoundary> ThinZoneBounds;
 
-		/**
-		 * Max delta U at the crossing points to respect meshing criteria
-		 */
-		TArray<double> CrossingPointDeltaUMaxs;
+	/**
+	 * Temporary discretization of the edge used to compute the mesh of the edge
+	 */
+	TArray<double> CrossingPointUs;
 
-	private:
+	/**
+	 * Min delta U at the crossing points to respect meshing criteria
+	 */
+	TArray<double> CrossingPointDeltaUMins;
 
-		/**
-		 * MANDATORY
-		 * Constructors of FEdge cannot be used directly, use FEdge::Make to create an FEdge object.
-		 */
-		FTopologicalEdge(const TSharedRef<FRestrictionCurve>& InCurve, const TSharedRef<FTopologicalVertex>& InVertex1, const TSharedRef<FTopologicalVertex>& InVertex2, const FLinearBoundary& InBoundary);
-		FTopologicalEdge(const TSharedRef<FRestrictionCurve>& InCurve, const TSharedRef<FTopologicalVertex>& InVertex1, const TSharedRef<FTopologicalVertex>& InVertex2);
-		FTopologicalEdge(const TSharedRef<FRestrictionCurve>& InCurve, const FLinearBoundary& InBoundary);
-		FTopologicalEdge(const TSharedRef<FSurface>& InSurface, const FPoint2D& InCoordinateVertex1, const TSharedRef<FTopologicalVertex>& InVertex1, const FPoint2D& InCoordinateVertex2, const TSharedRef<FTopologicalVertex>& InVertex2);
+	/**
+	 * Max delta U at the crossing points to respect meshing criteria
+	 */
+	TArray<double> CrossingPointDeltaUMaxs;
 
-		FTopologicalEdge() = default;
+private:
 
-		void SetLoop(FTopologicalLoop& NewBoundary)
-		{
-			Loop = &NewBoundary;
-		}
+	/**
+	 * MANDATORY
+	 * Constructors of FEdge cannot be used directly, use FEdge::Make to create an FEdge object.
+	 */
+	FTopologicalEdge(const TSharedRef<FRestrictionCurve>& InCurve, const TSharedRef<FTopologicalVertex>& InVertex1, const TSharedRef<FTopologicalVertex>& InVertex2, const FLinearBoundary& InBoundary);
+	FTopologicalEdge(const TSharedRef<FRestrictionCurve>& InCurve, const TSharedRef<FTopologicalVertex>& InVertex1, const TSharedRef<FTopologicalVertex>& InVertex2);
+	FTopologicalEdge(const TSharedRef<FRestrictionCurve>& InCurve, const FLinearBoundary& InBoundary);
+	FTopologicalEdge(const TSharedRef<FSurface>& InSurface, const FPoint2D& InCoordinateVertex1, const TSharedRef<FTopologicalVertex>& InVertex1, const FPoint2D& InCoordinateVertex2, const TSharedRef<FTopologicalVertex>& InVertex2);
 
-		void RemoveLoop()
-		{
-			Loop = nullptr;
-		}
+	FTopologicalEdge() = default;
 
-		/**
-		 * Used by FEdge::Make
-		 * Sort vertex coordinate to ensure that start coordinate < end coordinate
-		 * Swap vertices if needed i.e. to ensure that 3D point at start (and end) coordinate is near to start (and end) vertex 3d coordinates
-		 * @return false if one vertex is too far to the associated curve 3d point
-		 */
-		bool CheckVertices();
+	void SetLoop(FTopologicalLoop& NewBoundary)
+	{
+		Loop = &NewBoundary;
+	}
 
-	public:
+	void RemoveLoop()
+	{
+		Loop = nullptr;
+	}
 
-		static TSharedPtr<FTopologicalEdge> Make(const TSharedRef<FRestrictionCurve>& InCurve, const TSharedRef<FTopologicalVertex>& InVertex1, const TSharedRef<FTopologicalVertex>& InVertex2, const FLinearBoundary& InBoundary);
-		static TSharedPtr<FTopologicalEdge> Make(const TSharedRef<FRestrictionCurve>& InCurve, const TSharedRef<FTopologicalVertex>& InVertex1, const TSharedRef<FTopologicalVertex>& InVertex2);
-		static TSharedPtr<FTopologicalEdge> Make(const TSharedRef<FRestrictionCurve>& InCurve, const FLinearBoundary& InBoundary);
-		static TSharedPtr<FTopologicalEdge> Make(const TSharedRef<FRestrictionCurve>& InCurve);
+	/**
+	 * Used by FEdge::Make
+	 * Sort vertex coordinate to ensure that start coordinate < end coordinate
+	 * Swap vertices if needed i.e. to ensure that 3D point at start (and end) coordinate is near to start (and end) vertex 3d coordinates
+	 * @return false if one vertex is too far to the associated curve 3d point
+	 */
+	bool CheckVertices();
 
-		/**
-		 * Build an edge to connect two vertices carried by a 2d segment
-		 */
-		static TSharedPtr<FTopologicalEdge> Make(const TSharedRef<FSurface>& InSurface, const FPoint2D& InCoordinateVertex1, const TSharedRef<FTopologicalVertex>& InVertex1, const FPoint2D& InCoordinateVertex2, const TSharedRef<FTopologicalVertex>& InVertex2);
+public:
 
-		/**
-		 * To check the build edge before returning it or return TSharedPtr<FTopologicalEdge>() if the edge is not valid
-		 */
-		static TSharedPtr<FTopologicalEdge> ReturnIfValid(TSharedRef<FTopologicalEdge>& InEdge, bool bCheckVertices);
+	static TSharedPtr<FTopologicalEdge> Make(const TSharedRef<FRestrictionCurve>& InCurve, const TSharedRef<FTopologicalVertex>& InVertex1, const TSharedRef<FTopologicalVertex>& InVertex2, const FLinearBoundary& InBoundary);
+	static TSharedPtr<FTopologicalEdge> Make(const TSharedRef<FRestrictionCurve>& InCurve, const TSharedRef<FTopologicalVertex>& InVertex1, const TSharedRef<FTopologicalVertex>& InVertex2);
+	static TSharedPtr<FTopologicalEdge> Make(const TSharedRef<FRestrictionCurve>& InCurve, const FLinearBoundary& InBoundary);
+	static TSharedPtr<FTopologicalEdge> Make(const TSharedRef<FRestrictionCurve>& InCurve);
 
-		virtual ~FTopologicalEdge() = default;
+	/**
+	 * Build an edge to connect two vertices carried by a 2d segment
+	 */
+	static TSharedPtr<FTopologicalEdge> Make(const TSharedRef<FSurface>& InSurface, const FPoint2D& InCoordinateVertex1, const TSharedRef<FTopologicalVertex>& InVertex1, const FPoint2D& InCoordinateVertex2, const TSharedRef<FTopologicalVertex>& InVertex2);
 
-		virtual void Delete() override;
+	/**
+	 * To check the build edge before returning it or return TSharedPtr<FTopologicalEdge>() if the edge is not valid
+	 */
+	static TSharedPtr<FTopologicalEdge> ReturnIfValid(TSharedRef<FTopologicalEdge>& InEdge, bool bCheckVertices);
 
-		virtual void Serialize(FCADKernelArchive& Ar) override
-		{
-			TLinkable<FTopologicalEdge, FEdgeLink>::Serialize(Ar);
-			SerializeIdent(Ar, StartVertex);
-			SerializeIdent(Ar, EndVertex);
-			SerializeIdent(Ar, Curve);
-			Ar << Boundary;
-			SerializeIdent(Ar, &Loop);
-			Ar << Length3D;
-		}
+	virtual ~FTopologicalEdge() = default;
 
-		virtual void SpawnIdent(FDatabase& Database) override;
+	virtual void Delete() override;
 
-		virtual void ResetMarkersRecursively() override
-		{
-			ResetMarkers();
-			TLinkable<FTopologicalEdge, FEdgeLink>::ResetMarkers();
-			StartVertex->ResetMarkersRecursively();
-			EndVertex->ResetMarkersRecursively();
+	virtual void Serialize(FCADKernelArchive& Ar) override
+	{
+		TLinkable<FTopologicalEdge, FEdgeLink>::Serialize(Ar);
+		SerializeIdent(Ar, StartVertex);
+		SerializeIdent(Ar, EndVertex);
+		SerializeIdent(Ar, Curve);
+		Ar << Boundary;
+		SerializeIdent(Ar, &Loop);
+		Ar << Length3D;
+	}
+
+	virtual void SpawnIdent(FDatabase& Database) override;
+
+	virtual void ResetMarkersRecursively() override
+	{
+		ResetMarkers();
+		TLinkable<FTopologicalEdge, FEdgeLink>::ResetMarkers();
+		StartVertex->ResetMarkersRecursively();
+		EndVertex->ResetMarkersRecursively();
 
-			Curve->ResetMarkersRecursively();;
-		}
+		Curve->ResetMarkersRecursively();;
+	}
 
 #ifdef CADKERNEL_DEV
-		virtual FInfoEntity& GetInfo(FInfoEntity&) const override;
+	virtual FInfoEntity& GetInfo(FInfoEntity&) const override;
 #endif
 
-		double GetTolerance3D();
+	double GetTolerance3D();
 
-		virtual EEntity GetEntityType() const override
-		{
-			return EEntity::TopologicalEdge;
-		}
+	virtual EEntity GetEntityType() const override
+	{
+		return EEntity::TopologicalEdge;
+	}
 
-		virtual TSharedPtr<FEntityGeom> ApplyMatrix(const FMatrixH& InMatrix) const override;
+	virtual TSharedPtr<FEntityGeom> ApplyMatrix(const FMatrixH& InMatrix) const override;
 
-		// ======   Topological Function   ======
+	// ======   Topological Function   ======
 
-		void LinkVertex();
+	void LinkVertex();
 
-		/**
-		 * Checks if the carrier curve is degenerated i.e. the 2d length of the curve is nearly zero
-		 * If the 3d length is nearly zero, the edge is flag as degenerated
-		 */
-		bool CheckIfDegenerated();
+	/**
+	 * Checks if the carrier curve is degenerated i.e. the 2d length of the curve is nearly zero
+	 * If the 3d length is nearly zero, the edge is flag as degenerated
+	 */
+	bool CheckIfDegenerated();
 
-		/**
-		 * Build face links with its neigbour (link edges) should be done after the loop is finalize.
-		 * In some case, edges can be delete, split, extend... Link edges when the loop is well done and clean avoid problems
-		 */
-		void Link(FTopologicalEdge& OtherEdge, double SquareJoiningTolerance);
+	/**
+	 * Build face links with its neigbour (link edges) should be done after the loop is finalize.
+	 * In some case, edges can be delete, split, extend... Link edges when the loop is well done and clean avoid problems
+	 */
+	void Link(FTopologicalEdge& OtherEdge, double SquareJoiningTolerance);
 
-		TSharedRef<const FTopologicalEdge> GetLinkActiveEdge() const
-		{
-			return StaticCastSharedRef<const FTopologicalEdge>(GetLinkActiveEntity());
-		}
+	TSharedRef<const FTopologicalEdge> GetLinkActiveEdge() const
+	{
+		return StaticCastSharedRef<const FTopologicalEdge>(GetLinkActiveEntity());
+	}
+
+	TSharedRef<FTopologicalEdge> GetLinkActiveEdge()
+	{
+		return StaticCastSharedRef<FTopologicalEdge>(GetLinkActiveEntity());
+	}
 
-		TSharedRef<FTopologicalEdge> GetLinkActiveEdge()
+	FTopologicalEdge* GetFirstTwinEdge() const
+	{
+		if (!TopologicalLink)
 		{
-			return StaticCastSharedRef<FTopologicalEdge>(GetLinkActiveEntity());
+			return nullptr;
 		}
 
-		FTopologicalEdge* GetFirstTwinEdge() const
+		if (TopologicalLink->GetTwinsEntities().Num() < 2)
 		{
-			if (!TopologicalLink)
-			{
-				return nullptr;
-			}
+			return nullptr;
+		}
 
-			if (TopologicalLink->GetTwinsEntities().Num() < 2)
-			{
-				return nullptr;
-			}
+		FTopologicalEdge* FirstTwinEdge = (TopologicalLink->GetTwinsEntities()[0] == this) ? TopologicalLink->GetTwinsEntities()[1] : TopologicalLink->GetTwinsEntities()[0];
+		return FirstTwinEdge;
+	}
 
-			FTopologicalEdge* FirstTwinEdge = (TopologicalLink->GetTwinsEntities()[0] == this) ? TopologicalLink->GetTwinsEntities()[1] : TopologicalLink->GetTwinsEntities()[0];
-			return FirstTwinEdge;
-		}
+	/**
+	 * @return true if the twin edge is in the same direction as this
+	 */
+	bool IsSameDirection(const FTopologicalEdge& Edge) const;
 
-		/**
-		 * @return true if the twin edge is in the same direction as this
-		 */
-		bool IsSameDirection(const FTopologicalEdge& Edge) const;
+	/**
+	 * @return true if the edge is self connected at its extremities
+	 */
+	bool IsClosed() const
+	{
+		return StartVertex->GetLinkActiveEntity() == EndVertex->GetLinkActiveEntity();
+	}
 
-		/**
-		 * @return true if the edge is self connected at its extremities
-		 */
-		bool IsClosed() const
-		{
-			return StartVertex->GetLinkActiveEntity() == EndVertex->GetLinkActiveEntity();
-		}
+	/**
+	 * This function is used to :
+	 *   - sew exactly two edges of the same loop
+	 *   - prepare to disconnect two edges connected by the same vertex
+	 *
+	 * OldVertex = bStartVertex ? StartVertex : EndVertex
+	 * OldVertex can be an extremity of other edges
+	 * OldVertex can be link to other vertices (TopologicalLink can be valid or not)
+	 *
+	 * NewVertex can be an extremity of other edges
+	 * NewVertex can be linked to other vertices (TopologicalLink can be valid or not)
+	 * NewVertex can be not linked to OldVertex
+	 *
+	 * After the process:
+	 * OldVertex is no more connected to this edge but stay connected to its other linked edges
+	 * If OldVertex is isolated (not connected to any edge), OldVertex is deleted
+	 * otherwise OldVertex is linked to NewVertex
+	 */
+	void ReplaceEdgeVertex(bool bIsStartVertex, TSharedRef<FTopologicalVertex>& NewVertex);
 
-		/**
-		 * This function is used to :
-		 *   - sew exactly two edges of the same loop
-		 *   - prepare to disconnect two edges connected by the same vertex
-		 * 
-		 * OldVertex = bStartVertex ? StartVertex : EndVertex
-		 * OldVertex can be an extremity of other edges
-		 * OldVertex can be link to other vertices (TopologicalLink can be valid or not)
-		 * 
-		 * NewVertex can be an extremity of other edges
-		 * NewVertex can be linked to other vertices (TopologicalLink can be valid or not)
-		 * NewVertex can be not linked to OldVertex
-		 * 
-		 * After the process:
-		 * OldVertex is no more connected to this edge but stay connected to its other linked edges
-		 * If OldVertex is isolated (not connected to any edge), OldVertex is deleted
-		 * otherwise OldVertex is linked to NewVertex
-		 */
-		void ReplaceEdgeVertex(bool bIsStartVertex, TSharedRef<FTopologicalVertex>& NewVertex);
-
-		/**
-		 * @return the containing boundary
-		 */
-		const FTopologicalLoop* GetLoop() const
-		{
-			return Loop;
-		}
+	/**
+	 * @return the containing boundary
+	 */
+	const FTopologicalLoop* GetLoop() const
+	{
+		return Loop;
+	}
 
-		/**
-		 * @return the containing boundary
-		 */
-		FTopologicalLoop* GetLoop()
-		{
-			return Loop;
-		}
+	/**
+	 * @return the containing boundary
+	 */
+	FTopologicalLoop* GetLoop()
+	{
+		return Loop;
+	}
 
-		/**
-		 * @return the carrier topological face
-		 */
-		TSharedPtr<FTopologicalFace> GetFace() const;
+	/**
+	 * @return the carrier topological face
+	 */
+	FTopologicalFace* GetFace() const;
 
-		// ======   Vertex Functions (Get, Set, ...)   ======
+	// ======   Vertex Functions (Get, Set, ...)   ======
 
-		const TSharedRef<FTopologicalVertex> GetStartVertex(EOrientation Forward) const
-		{
-			return (Forward == EOrientation::Front ? StartVertex.ToSharedRef() : EndVertex.ToSharedRef());
-		}
+	const TSharedRef<FTopologicalVertex> GetStartVertex(EOrientation Forward) const
+	{
+		return (Forward == EOrientation::Front ? StartVertex.ToSharedRef() : EndVertex.ToSharedRef());
+	}
 
-		const TSharedRef<FTopologicalVertex> GetEndVertex(EOrientation Forward) const
-		{
-			return (Forward == EOrientation::Front ? EndVertex.ToSharedRef() : StartVertex.ToSharedRef());
-		}
+	const TSharedRef<FTopologicalVertex> GetEndVertex(EOrientation Forward) const
+	{
+		return (Forward == EOrientation::Front ? EndVertex.ToSharedRef() : StartVertex.ToSharedRef());
+	}
 
-		const TSharedRef<FTopologicalVertex> GetStartVertex(bool Forward) const
-		{
-			return (Forward ? StartVertex.ToSharedRef() : EndVertex.ToSharedRef());
-		}
+	const TSharedRef<FTopologicalVertex> GetStartVertex(bool Forward) const
+	{
+		return (Forward ? StartVertex.ToSharedRef() : EndVertex.ToSharedRef());
+	}
 
-		const TSharedRef<FTopologicalVertex> GetEndVertex(bool Forward) const
-		{
-			return (Forward ? EndVertex.ToSharedRef() : StartVertex.ToSharedRef());
-		}
+	const TSharedRef<FTopologicalVertex> GetEndVertex(bool Forward) const
+	{
+		return (Forward ? EndVertex.ToSharedRef() : StartVertex.ToSharedRef());
+	}
 
-		const TSharedRef<FTopologicalVertex> GetStartVertex() const
-		{
-			return StartVertex.ToSharedRef();
-		}
+	const TSharedRef<FTopologicalVertex> GetStartVertex() const
+	{
+		return StartVertex.ToSharedRef();
+	}
 
-		const TSharedRef< FTopologicalVertex> GetEndVertex() const
-		{
-			return EndVertex.ToSharedRef();
-		}
+	const TSharedRef< FTopologicalVertex> GetEndVertex() const
+	{
+		return EndVertex.ToSharedRef();
+	}
 
-		TSharedRef<FTopologicalVertex> GetStartVertex()
-		{
-			return StartVertex.ToSharedRef();
-		}
+	TSharedRef<FTopologicalVertex> GetStartVertex()
+	{
+		return StartVertex.ToSharedRef();
+	}
 
-		TSharedRef<FTopologicalVertex> GetEndVertex()
-		{
-			return EndVertex.ToSharedRef();
-		}
+	TSharedRef<FTopologicalVertex> GetEndVertex()
+	{
+		return EndVertex.ToSharedRef();
+	}
 
-		TSharedPtr<FTopologicalVertex> GetOtherVertex(const TSharedRef<FTopologicalVertex>& Vertex)
-		{
-			return (Vertex->GetLink() == StartVertex->GetLink() ? EndVertex : (Vertex->GetLink() == EndVertex->GetLink() ? StartVertex : TSharedPtr<FTopologicalVertex>()));
-		}
+	TSharedPtr<FTopologicalVertex> GetOtherVertex(const TSharedRef<FTopologicalVertex>& Vertex)
+	{
+		return (Vertex->GetLink() == StartVertex->GetLink() ? EndVertex : (Vertex->GetLink() == EndVertex->GetLink() ? StartVertex : TSharedPtr<FTopologicalVertex>()));
+	}
 
-		FTopologicalVertex* GetOtherVertex(FTopologicalVertex& Vertex)
-		{
-			return (Vertex.GetLink() == StartVertex->GetLink() ? EndVertex.Get() : (Vertex.GetLink() == EndVertex->GetLink() ? StartVertex.Get() : nullptr));
-		}
+	FTopologicalVertex* GetOtherVertex(FTopologicalVertex& Vertex)
+	{
+		return (Vertex.GetLink() == StartVertex->GetLink() ? EndVertex.Get() : (Vertex.GetLink() == EndVertex->GetLink() ? StartVertex.Get() : nullptr));
+	}
 
-		const FTopologicalVertex* GetOtherVertex(const FTopologicalVertex& Vertex) const
-		{
-			return (Vertex.GetLink() == StartVertex->GetLink() ? EndVertex.Get() : (Vertex.GetLink() == EndVertex->GetLink() ? StartVertex.Get() : nullptr));
-		}
+	const FTopologicalVertex* GetOtherVertex(const FTopologicalVertex& Vertex) const
+	{
+		return (Vertex.GetLink() == StartVertex->GetLink() ? EndVertex.Get() : (Vertex.GetLink() == EndVertex->GetLink() ? StartVertex.Get() : nullptr));
+	}
 
-		const TSharedPtr<FTopologicalVertex> GetOtherVertex(const TSharedRef<FTopologicalVertex>& Vertex) const
-		{
-			return (Vertex->GetLink() == StartVertex->GetLink() ? EndVertex : (Vertex->GetLink() == EndVertex->GetLink() ? StartVertex : TSharedPtr<FTopologicalVertex>()));
-		}
+	const TSharedPtr<FTopologicalVertex> GetOtherVertex(const TSharedRef<FTopologicalVertex>& Vertex) const
+	{
+		return (Vertex->GetLink() == StartVertex->GetLink() ? EndVertex : (Vertex->GetLink() == EndVertex->GetLink() ? StartVertex : TSharedPtr<FTopologicalVertex>()));
+	}
 
-		void SetStartVertex(const double NewCoordinate);
-		void SetEndVertex(const double NewCoordinate);
+	void SetStartVertex(const double NewCoordinate);
+	void SetEndVertex(const double NewCoordinate);
 
-		void SetStartVertex(const double NewCoordinate, const FPoint& NewPoint3D);
-		void SetEndVertex(const double NewCoordinate, const FPoint& NewPoint3D);
+	void SetStartVertex(const double NewCoordinate, const FPoint& NewPoint3D);
+	void SetEndVertex(const double NewCoordinate, const FPoint& NewPoint3D);
 
-		// ======   Boundary Function   ======
+	// ======   Boundary Function   ======
 
-		const FLinearBoundary& GetBoundary() const
-		{
-			return Boundary;
-		}
+	const FLinearBoundary& GetBoundary() const
+	{
+		return Boundary;
+	}
 
-		double GetStartCurvilinearCoordinates() const
-		{
-			return Boundary.GetMin();
-		}
+	double GetStartCurvilinearCoordinates() const
+	{
+		return Boundary.GetMin();
+	}
 
-		double GetEndCurvilinearCoordinates() const
-		{
-			return Boundary.GetMax();
-		}
+	double GetEndCurvilinearCoordinates() const
+	{
+		return Boundary.GetMax();
+	}
 
-		/**
-		 * @return the 3d coordinate of the start vertex (the barycenter of the twin vertices)
-		 */
-		FPoint GetStartBarycenter()
-		{
-			return StartVertex->GetBarycenter();
-		}
+	/**
+	 * @return the 3d coordinate of the start vertex (the barycenter of the twin vertices)
+	 */
+	FPoint GetStartBarycenter()
+	{
+		return StartVertex->GetBarycenter();
+	}
 
-		/**
-		 * @return the 3d coordinate of the end vertex (the barycenter of the twin vertices)
-		 */
-		FPoint GetEndBarycenter()
-		{
-			return EndVertex->GetBarycenter();
-		}
+	/**
+	 * @return the 3d coordinate of the end vertex (the barycenter of the twin vertices)
+	 */
+	FPoint GetEndBarycenter()
+	{
+		return EndVertex->GetBarycenter();
+	}
 
-		/**
-		 * @return the 3d coordinate of the start vertex (prefer GetStartBarycenter)
-		 */
-		FPoint GetStartCoordinate()
-		{
-			return StartVertex->GetCoordinates();
-		}
+	/**
+	 * @return the 3d coordinate of the start vertex (prefer GetStartBarycenter)
+	 */
+	FPoint GetStartCoordinate()
+	{
+		return StartVertex->GetCoordinates();
+	}
 
-		/**
-		 * @return the 3d coordinate of the end vertex (prefer GetEndBarycenter)
-		 */
-		FPoint GetEndCoordinate()
-		{
-			return EndVertex->GetCoordinates();
-		}
+	/**
+	 * @return the 3d coordinate of the end vertex (prefer GetEndBarycenter)
+	 */
+	FPoint GetEndCoordinate()
+	{
+		return EndVertex->GetCoordinates();
+	}
 
 
 
-		void GetTangentsAtExtremities(FPoint& StartTangent, FPoint& EndTangent, bool bForward) const;
+	void GetTangentsAtExtremities(FPoint& StartTangent, FPoint& EndTangent, bool bForward) const;
 
-		// ======   Meshing Function   ======
+	// ======   Meshing Function   ======
 
-		const TSharedRef<FEdgeMesh> GetMesh() const
+	const TSharedRef<FEdgeMesh> GetMesh() const
+	{
+		if (GetLinkActiveEntity() != AsShared())
 		{
-			if (GetLinkActiveEntity() != AsShared())
-			{
-				return GetLinkActiveEdge()->GetMesh();
-			}
-			ensureCADKernel(Mesh.IsValid());
-			return Mesh.ToSharedRef();
+			return GetLinkActiveEdge()->GetMesh();
 		}
+		ensureCADKernel(Mesh.IsValid());
+		return Mesh.ToSharedRef();
+	}
 
-		TSharedRef<FEdgeMesh> GetOrCreateMesh(const TSharedRef<FModelMesh>& MeshModel);
+	TSharedRef<FEdgeMesh> GetOrCreateMesh(const TSharedRef<FModelMesh>& MeshModel);
 
-		/**
-		 * Generate a sampling of the curve.
-		 * This sampling is used by apply meshing criteria function to defined the optimal mesh of the edge.
-		 * This sampling is saved in CrossingPointUs TArray.
-		 */
-		void ComputeCrossingPointCoordinates();
-		
-		int32 EvaluateCuttingPointNum();
+	/**
+	 * Generate a sampling of the curve.
+	 * This sampling is used by apply meshing criteria function to defined the optimal mesh of the edge.
+	 * This sampling is saved in CrossingPointUs TArray.
+	 */
+	void ComputeCrossingPointCoordinates();
 
-		void InitDeltaUs()
-		{
-			int32 Size = CrossingPointUs.Num();
-			ensureCADKernel(Size >= 2);
-			CrossingPointUs.SetNum(Size);
-			CrossingPointDeltaUMins.Init(SMALL_NUMBER, Size - 1);
-			CrossingPointDeltaUMaxs.Init(2.0 * (GetEndCurvilinearCoordinates() - GetStartCurvilinearCoordinates()), Size - 1);
-		}
+	int32 EvaluateCuttingPointNum();
 
-		void ChooseFinalDeltaUs();
+	void InitDeltaUs()
+	{
+		int32 Size = CrossingPointUs.Num();
+		ensureCADKernel(Size >= 2);
+		CrossingPointUs.SetNum(Size);
+		CrossingPointDeltaUMins.Init(SMALL_NUMBER, Size - 1);
+		CrossingPointDeltaUMaxs.Init(2.0 * (GetEndCurvilinearCoordinates() - GetStartCurvilinearCoordinates()), Size - 1);
+	}
 
-		const TArray<double>& GetCrossingPointUs() const
-		{
-			return CrossingPointUs;
-		}
+	void ChooseFinalDeltaUs();
 
-		TArray<double>& GetCrossingPointUs()
-		{
-			return CrossingPointUs;
-		}
+	const TArray<double>& GetCrossingPointUs() const
+	{
+		return CrossingPointUs;
+	}
 
-		TArray<double>& GetDeltaUMins()
-		{
-			return CrossingPointDeltaUMins;
-		}
+	TArray<double>& GetCrossingPointUs()
+	{
+		return CrossingPointUs;
+	}
 
-		TArray<double>& GetDeltaUMaxs()
-		{
-			return CrossingPointDeltaUMaxs;
-		}
+	TArray<double>& GetDeltaUMins()
+	{
+		return CrossingPointDeltaUMins;
+	}
 
-		const TArray<double>& GetDeltaUMaxs() const
-		{
-			return CrossingPointDeltaUMaxs;
-		}
+	TArray<double>& GetDeltaUMaxs()
+	{
+		return CrossingPointDeltaUMaxs;
+	}
 
-		TArray<FCuttingPoint>& GetCuttingPoints()
-		{
-			return CuttingPointUs;
-		};
+	const TArray<double>& GetDeltaUMaxs() const
+	{
+		return CrossingPointDeltaUMaxs;
+	}
 
-		const TArray<FCuttingPoint>& GetCuttingPoints() const
-		{
-			return CuttingPointUs;
-		}
+	TArray<FCuttingPoint>& GetCuttingPoints()
+	{
+		return CuttingPointUs;
+	};
 
-		// For thin zone purpose
+	const TArray<FCuttingPoint>& GetCuttingPoints() const
+	{
+		return CuttingPointUs;
+	}
 
-		const TArray<FImposedCuttingPoint>& GetImposedCuttingPoints()
-		{
-			return ImposedCuttingPointUs;
-		}
+	// For thin zone purpose
 
-		void AddThinZone(const FLinearBoundary& InThinZoneBounds)
-		{
-			ThinZoneBounds.Add(InThinZoneBounds);
-		}
+	const TArray<FImposedCuttingPoint>& GetImposedCuttingPoints()
+	{
+		return ImposedCuttingPointUs;
+	}
 
-		const TArray<FLinearBoundary>& GetThinZoneBounds() const
-		{
-			return ThinZoneBounds;
-		}
+	void AddThinZone(const FLinearBoundary& InThinZoneBounds)
+	{
+		ThinZoneBounds.Add(InThinZoneBounds);
+	}
 
-		void AddImposedCuttingPointU(const double ImposedCuttingPointU, int32 OppositeNodeIndex);
+	const TArray<FLinearBoundary>& GetThinZoneBounds() const
+	{
+		return ThinZoneBounds;
+	}
 
+	void AddImposedCuttingPointU(const double ImposedCuttingPointU, int32 OppositeNodeIndex);
 
-		// ======   Curve Functions   ======
 
-		TSharedRef<FRestrictionCurve> GetCurve() const
-		{
-			return Curve.ToSharedRef();
-		}
+	// ======   Curve Functions   ======
 
-		TSharedRef<FRestrictionCurve> GetCurve()
-		{
-			return Curve.ToSharedRef();
-		}
+	TSharedRef<FRestrictionCurve> GetCurve() const
+	{
+		return Curve.ToSharedRef();
+	}
 
-		double Length() const;
+	TSharedRef<FRestrictionCurve> GetCurve()
+	{
+		return Curve.ToSharedRef();
+	}
 
-		/**
-		 * Samples the curve with segments of a desired length
-		 */
-		void Sample(const double DesiredSegmentLength, TArray<double>& OutCoordinates) const;
+	double Length() const;
 
-		/**
-		 * Exact evaluation of point on the 3D curve
-		 * According to derivativeOrder Gradient of the point (DerivativeOrder = 1) and Laplacian (DerivativeOrder = 1) can also be return
-		 */
-		void EvaluatePoint(double InCoordinate, int32 Derivative, FCurvePoint& Point) const
-		{
-			Curve->EvaluatePoint(InCoordinate, Point, Derivative);
-		}
+	/**
+	 * Samples the curve with segments of a desired length
+	 */
+	void Sample(const double DesiredSegmentLength, TArray<double>& OutCoordinates) const;
 
-		/**
-		 * Exact evaluation of points on the 3D curve
-		 * According to derivativeOrder Gradient of the point (DerivativeOrder = 1) and Laplacian (DerivativeOrder = 1) can also be return
-		 */
-		void EvaluatePoints(const TArray<double>& InCoordinates, int32 DerivativeOrder, TArray<FCurvePoint>& OutPoints) const
-		{
-			Curve->EvaluatePoints(InCoordinates, OutPoints, DerivativeOrder);
-		}
+	/**
+	 * Exact evaluation of point on the 3D curve
+	 * According to derivativeOrder Gradient of the point (DerivativeOrder = 1) and Laplacian (DerivativeOrder = 1) can also be return
+	 */
+	void EvaluatePoint(double InCoordinate, int32 Derivative, FCurvePoint& Point) const
+	{
+		Curve->EvaluatePoint(InCoordinate, Point, Derivative);
+	}
 
-		/**
-		 * Approximation of 3D points compute with carrier surface 3D polyline
-		 */
-		void ApproximatePoints(const TArray<double>& InCoordinates, TArray<FPoint>& OutPoints) const
-		{
-			Curve->Approximate3DPoints(InCoordinates, OutPoints);
-		}
+	/**
+	 * Exact evaluation of points on the 3D curve
+	 * According to derivativeOrder Gradient of the point (DerivativeOrder = 1) and Laplacian (DerivativeOrder = 1) can also be return
+	 */
+	void EvaluatePoints(const TArray<double>& InCoordinates, int32 DerivativeOrder, TArray<FCurvePoint>& OutPoints) const
+	{
+		Curve->EvaluatePoints(InCoordinates, OutPoints, DerivativeOrder);
+	}
 
-		/**
-		 * Approximation of 2D point defined by its coordinate compute with carrier surface 2D polyline
-		 */
-		FPoint Approximate2DPoint(const double InCoordinate) const
-		{
-			return Curve->Approximate2DPoint(InCoordinate);
-		}
+	/**
+	 * Approximation of 3D points compute with carrier surface 3D polyline
+	 */
+	void ApproximatePoints(const TArray<double>& InCoordinates, TArray<FPoint>& OutPoints) const
+	{
+		Curve->Approximate3DPoints(InCoordinates, OutPoints);
+	}
 
-		/**
-		 * Approximation of 2D points defined by its coordinates compute with carrier surface 2D polyline
-		 */
-		void Approximate2DPoints(const TArray<double>& InCoordinates, TArray<FPoint2D>& OutPoints) const	
-		{
-			Curve->Approximate2DPoints(InCoordinates, OutPoints);
-		}
+	/**
+	 * Approximation of 2D point defined by its coordinate compute with carrier surface 2D polyline
+	 */
+	FPoint Approximate2DPoint(const double InCoordinate) const
+	{
+		return Curve->Approximate2DPoint(InCoordinate);
+	}
 
-		/**
-		 * Approximation of surfacic polyline (points 2d, 3d, normals, tangents) defined by its coordinates compute with carrier surface polyline
-		 */
-		void ApproximatePolyline(FSurfacicPolyline& Polyline) const
-		{
-			Curve->ApproximatePolyline(Polyline);
-		}
+	/**
+	 * Approximation of 2D points defined by its coordinates compute with carrier surface 2D polyline
+	 */
+	void Approximate2DPoints(const TArray<double>& InCoordinates, TArray<FPoint2D>& OutPoints) const
+	{
+		Curve->Approximate2DPoints(InCoordinates, OutPoints);
+	}
 
-		FPoint GetTangentAt(const double InCoordinate) const
-		{
-			return Curve->GetTangentAt(InCoordinate);
-		}
+	/**
+	 * Approximation of surfacic polyline (points 2d, 3d, normals, tangents) defined by its coordinates compute with carrier surface polyline
+	 */
+	void ApproximatePolyline(FSurfacicPolyline& Polyline) const
+	{
+		Curve->ApproximatePolyline(Polyline);
+	}
 
-		FPoint2D GetTangent2DAt(const double InCoordinate) const
-		{
-			return Curve->GetTangent2DAt(InCoordinate);
-		}
+	FPoint GetTangentAt(const double InCoordinate) const
+	{
+		return Curve->GetTangentAt(InCoordinate);
+	}
 
-		/**
-		 * Return the tangent at the input vertex
-		 */
-		FPoint GetTangentAt(const FTopologicalVertex& InVertex);
-		FPoint2D GetTangent2DAt(const FTopologicalVertex& InVertex);
-
-		/**
-		 * Project Point (2D or 3D) on the polyline (2D or 3D) and return the coordinate of the projected point
-		 */
-		template<class PointType>
-		double ProjectPoint(const PointType& InPointToProject, PointType& OutProjectedPoint) const
-		{
-			return Curve->GetCoordinateOfProjectedPoint(Boundary, InPointToProject, OutProjectedPoint);
-		}
+	FPoint2D GetTangent2DAt(const double InCoordinate) const
+	{
+		return Curve->GetTangent2DAt(InCoordinate);
+	}
 
-		/**
-		 * Project a set of points on the 3D polyline and return the coordinate of the projected point
-		 */
-		void ProjectPoints(const TArray<FPoint>& InPointsToProject, TArray<double>& OutProjectedPointCoords, TArray<FPoint>& OutProjectedPoints) const
-		{
-			Curve->ProjectPoints(Boundary, InPointsToProject, OutProjectedPointCoords, OutProjectedPoints);
-		}
+	/**
+	 * Return the tangent at the input vertex
+	 */
+	FPoint GetTangentAt(const FTopologicalVertex& InVertex);
+	FPoint2D GetTangent2DAt(const FTopologicalVertex& InVertex);
 
-		/**
-		 * Project a set of points of a twin edge on the 3D polyline and return the coordinate of the projected point
-		 */
-		void ProjectTwinEdgePoints(const TArray<FPoint>& InPointsToProject, bool bSameOrientation, TArray<double>& OutProjectedPointCoords) const
-		{
-			Curve->ProjectTwinCurvePoints(InPointsToProject, bSameOrientation, OutProjectedPointCoords);
-		}
+	/**
+	 * Project Point (2D or 3D) on the polyline (2D or 3D) and return the coordinate of the projected point
+	 */
+	template<class PointType>
+	double ProjectPoint(const PointType& InPointToProject, PointType& OutProjectedPoint) const
+	{
+		return Curve->GetCoordinateOfProjectedPoint(Boundary, InPointToProject, OutProjectedPoint);
+	}
 
-		/**
-		 * Compute 2D points of the edge coincident the points of the twin edge defined by their coordinates
-		 */
-		void ProjectTwinEdgePointsOn2DCurve(const TSharedRef<FTopologicalEdge>& InTwinEdge, const TArray<double>& InTwinEdgePointCoords, TArray<FPoint2D>& OutPoints2D);
+	/**
+	 * Project a set of points on the 3D polyline and return the coordinate of the projected point
+	 */
+	void ProjectPoints(const TArray<FPoint>& InPointsToProject, TArray<double>& OutProjectedPointCoords, TArray<FPoint>& OutProjectedPoints) const
+	{
+		Curve->ProjectPoints(Boundary, InPointsToProject, OutProjectedPointCoords, OutProjectedPoints);
+	}
 
-		void ComputeIntersectionsWithIsos(const TArray<double>& InIsoCoordinates, const EIso InTypeIso, const FSurfacicTolerance& ToleranceIso, TArray<double>& OutIntersection) const
-		{
-			Curve->ComputeIntersectionsWithIsos(Boundary, InIsoCoordinates, InTypeIso, ToleranceIso, OutIntersection);
-		}
+	/**
+	 * Project a set of points of a twin edge on the 3D polyline and return the coordinate of the projected point
+	 */
+	void ProjectTwinEdgePoints(const TArray<FPoint>& InPointsToProject, bool bSameOrientation, TArray<double>& OutProjectedPointCoords) const
+	{
+		Curve->ProjectTwinCurvePoints(InPointsToProject, bSameOrientation, OutProjectedPointCoords);
+	}
 
-		/**
-		 * Get the discretization points of the edge and add them to the outpoints TArray
-		 */
-		template<class PointType>
-		void GetDiscretization2DPoints(EOrientation Orientation, TArray<PointType>& OutPoints) const
-		{
-			Curve->GetDiscretizationPoints(Boundary, Orientation, OutPoints);
-		}
+	/**
+	 * Compute 2D points of the edge coincident the points of the twin edge defined by their coordinates
+	 */
+	void ProjectTwinEdgePointsOn2DCurve(const TSharedRef<FTopologicalEdge>& InTwinEdge, const TArray<double>& InTwinEdgePointCoords, TArray<FPoint2D>& OutPoints2D);
 
-		double TransformLocalCoordinateToActiveEdgeCoordinate(const double LocalCoordinate);
-		double TransformActiveEdgeCoordinateToLocalCoordinate(const double ActiveEdgeCoordinate);
+	void ComputeIntersectionsWithIsos(const TArray<double>& InIsoCoordinates, const EIso InTypeIso, const FSurfacicTolerance& ToleranceIso, TArray<double>& OutIntersection) const
+	{
+		Curve->ComputeIntersectionsWithIsos(Boundary, InIsoCoordinates, InTypeIso, ToleranceIso, OutIntersection);
+	}
 
-		void TransformActiveEdgeCoordinatesToLocalCoordinates(const TArray<double>& InActiveEdgeCoordinate, TArray<double>& OutLocalCoordinate);
-		void TransformLocalCoordinatesToActiveEdgeCoordinates(const TArray<double>& InLocalCoordinate, TArray<double>& OutActiveEdgeCoordinate);
+	/**
+	 * Get the discretization points of the edge and add them to the outpoints TArray
+	 */
+	template<class PointType>
+	void GetDiscretization2DPoints(EOrientation Orientation, TArray<PointType>& OutPoints) const
+	{
+		Curve->GetDiscretizationPoints(Boundary, Orientation, OutPoints);
+	}
 
-		/**
-		 * Compute the edge 2D properties i.e. the mean and standard deviation of the slop of the edge in the parametric space of the carrier surface
-		 */
-		void ComputeEdge2DProperties(FEdge2DProperties& SlopCharacteristics);
+	double TransformLocalCoordinateToActiveEdgeCoordinate(const double LocalCoordinate);
+	double TransformActiveEdgeCoordinateToLocalCoordinate(const double ActiveEdgeCoordinate);
 
-		void GetExtremities(FSurfacicCurveExtremities& Extremities) const
-		{
-			Curve->GetExtremities(Boundary, Extremities);
-		}
+	void TransformActiveEdgeCoordinatesToLocalCoordinates(const TArray<double>& InActiveEdgeCoordinate, TArray<double>& OutLocalCoordinate);
+	void TransformLocalCoordinatesToActiveEdgeCoordinates(const TArray<double>& InLocalCoordinate, TArray<double>& OutActiveEdgeCoordinate);
 
+	/**
+	 * Compute the edge 2D properties i.e. the mean and standard deviation of the slop of the edge in the parametric space of the carrier surface
+	 */
+	void ComputeEdge2DProperties(FEdge2DProperties& SlopCharacteristics);
 
-		// ======   Geometrical Functions   ======
+	void GetExtremities(FSurfacicCurveExtremities& Extremities) const
+	{
+		Curve->GetExtremities(Boundary, Extremities);
+	}
 
-		/**
-		 * Split the edge at input coordinate. The initial edge will be used for the first edge, the second edge will be return.
-		 * @param SplittingCoordinate
-		 * @param NewVertexCoordinate
-		 * @param bKeepStartVertexConnectivity: if true the new edge is connected to endVertex, otherwise the new edge is connected to start vertex
-		 * @param OutNewEdge, the second edge
-		 */
-		TSharedPtr<FTopologicalVertex> SplitAt(double SplittingCoordinate, const FPoint& NewVertexCoordinate, bool bKeepStartVertexConnectivity, TSharedPtr<FTopologicalEdge>& OutNewEdge);
 
-		/**
-		 * Extend the Edge to the NewVertex. 
-		 * NewVertex become the new extremity of the edge. The old vertex is unconnected of the edge.
-		 */
-		bool ExtendTo(bool bStartExtremity, const FPoint2D& NewExtremityCoordinate, TSharedRef<FTopologicalVertex>& NewVertex);
+	// ======   Geometrical Functions   ======
 
-		bool IsSharpEdge();
+	/**
+	 * Split the edge at input coordinate. The initial edge will be used for the first edge, the second edge will be return.
+	 * @param SplittingCoordinate
+	 * @param NewVertexCoordinate
+	 * @param bKeepStartVertexConnectivity: if true the new edge is connected to endVertex, otherwise the new edge is connected to start vertex
+	 * @param OutNewEdge, the second edge
+	 */
+	TSharedPtr<FTopologicalVertex> SplitAt(double SplittingCoordinate, const FPoint& NewVertexCoordinate, bool bKeepStartVertexConnectivity, TSharedPtr<FTopologicalEdge>& OutNewEdge);
 
+	/**
+	 * Extend the Edge to the NewVertex.
+	 * NewVertex become the new extremity of the edge. The old vertex is unconnected of the edge.
+	 */
+	bool ExtendTo(bool bStartExtremity, const FPoint2D& NewExtremityCoordinate, TSharedRef<FTopologicalVertex>& NewVertex);
 
-		// ======   State Functions   ======
+	bool IsSharpEdge();
 
-		/**
-		 * Important note: A Degenerated Edge is used to close 2D boundary in case of degenerated surface to ensureCADKernel a closed boundary
-		 * Specific process is done for the mesh of this kind of surface
-		 */
-		virtual bool IsDegenerated() const override
-		{
-			return FHaveStates::IsDegenerated();
-		}
 
-		bool IsThinPeak() const
-		{
-			return ((States & EHaveStates::ThinPeak) == EHaveStates::ThinPeak);
-		}
+	// ======   State Functions   ======
 
-		virtual void SetThinPeak() const
-		{
-			States |= EHaveStates::ThinPeak;
-		}
+	/**
+	 * Important note: A Degenerated Edge is used to close 2D boundary in case of degenerated surface to ensureCADKernel a closed boundary
+	 * Specific process is done for the mesh of this kind of surface
+	 */
+	virtual bool IsDegenerated() const override
+	{
+		return FHaveStates::IsDegenerated();
+	}
 
-		virtual void ResetThinPeak() const
-		{
-			States &= ~EHaveStates::ThinPeak;
-		}
+	bool IsThinPeak() const
+	{
+		return ((States & EHaveStates::ThinPeak) == EHaveStates::ThinPeak);
+	}
 
-		bool IsVirtuallyMeshed() const
-		{
-			return ((States & EHaveStates::IsVirtuallyMeshed) == EHaveStates::IsVirtuallyMeshed);
-		}
+	virtual void SetThinPeak() const
+	{
+		States |= EHaveStates::ThinPeak;
+	}
 
-		virtual void SetAsVirtuallyMeshed() const
-		{
-			States |= EHaveStates::IsVirtuallyMeshed;
-		}
+	virtual void ResetThinPeak() const
+	{
+		States &= ~EHaveStates::ThinPeak;
+	}
 
-		virtual void ResetVirtuallyMeshed() const
-		{
-			States &= ~EHaveStates::IsVirtuallyMeshed;
-		}
-		
-		/**
-		 * @return true if the edge is adjacent to only one surface (its carrier surface)
-		 */
-		bool IsBorder()
-		{
-			return GetTwinsEntityCount() == 1;
-		}
+	bool IsVirtuallyMeshed() const
+	{
+		return ((States & EHaveStates::IsVirtuallyMeshed) == EHaveStates::IsVirtuallyMeshed);
+	}
 
-		/**
-		 * @return true if the edge is adjacent to only two surfaces
-		 */
-		bool IsSurfacic()
-		{
-			return GetTwinsEntityCount() == 2;
-		}
+	virtual void SetAsVirtuallyMeshed() const
+	{
+		States |= EHaveStates::IsVirtuallyMeshed;
+	}
 
-		/**
-		 * Merge successive edges of a face in a single edge.
-		 * Edges must not be topological linked to other faces i.e. edges must be border edges
-		 * The merged edges are deleted
-		 * @return TSharedPtr<FTopologicalEdge>() if failed
-		 */
-		static TSharedPtr<FTopologicalEdge> CreateEdgeByMergingEdges(TArray<FOrientedEdge>& Edges, const TSharedRef<FTopologicalVertex> StartVertex, const TSharedRef<FTopologicalVertex> EndVertex);
-	};
+	virtual void ResetVirtuallyMeshed() const
+	{
+		States &= ~EHaveStates::IsVirtuallyMeshed;
+	}
 
-	struct CADKERNEL_API FEdge2DProperties
+	/**
+	 * @return true if the edge is adjacent to only one surface (its carrier surface)
+	 */
+	bool IsBorder()
 	{
-		double StandardDeviation = 0;
-		double MediumSlop = 0;
-		double Length3D = 0;
-		EIso IsoType = EIso::UndefinedIso;
-		bool bIsMesh = false;
-		double MeshedLength = 0;
+		return GetTwinsEntityCount() == 1;
+	}
 
-		void Add(double InSlop, double InLength)
-		{
-			double Temp = InSlop * InLength;
-			MediumSlop += Temp;
-			Temp *= InSlop;
-			StandardDeviation += Temp;
-			Length3D += InLength;
-		}
+	/**
+	 * @return true if the edge is adjacent to only two surfaces
+	 */
+	bool IsSurfacic()
+	{
+		return GetTwinsEntityCount() == 2;
+	}
 
-		// Finalize has been done on each Property
-		void Add2(FEdge2DProperties& Property)
-		{
-			StandardDeviation = (FMath::Square(StandardDeviation) + FMath::Square(MediumSlop)) * Length3D + (FMath::Square(Property.StandardDeviation) + FMath::Square(Property.MediumSlop)) * Property.Length3D;
-			MediumSlop = MediumSlop * Length3D + Property.MediumSlop * Property.Length3D;
-			Length3D += Property.Length3D;
+	/**
+	 * Merge successive edges of a face in a single edge.
+	 * Edges must not be topological linked to other faces i.e. edges must be border edges
+	 * The merged edges are deleted
+	 * @return TSharedPtr<FTopologicalEdge>() if failed
+	 */
+	static TSharedPtr<FTopologicalEdge> CreateEdgeByMergingEdges(TArray<FOrientedEdge>& Edges, const TSharedRef<FTopologicalVertex> StartVertex, const TSharedRef<FTopologicalVertex> EndVertex);
+};
 
-			MediumSlop /= Length3D;
-			StandardDeviation /= Length3D;
-			StandardDeviation -= FMath::Square(MediumSlop);
-			StandardDeviation = sqrt(StandardDeviation);
-		}
+struct CADKERNEL_API FEdge2DProperties
+{
+	double StandardDeviation = 0;
+	double MediumSlop = 0;
+	double Length3D = 0;
+	EIso IsoType = EIso::UndefinedIso;
+	bool bIsMesh = false;
+	double MeshedLength = 0;
+
+	void Add(double InSlop, double InLength)
+	{
+		double Temp = InSlop * InLength;
+		MediumSlop += Temp;
+		Temp *= InSlop;
+		StandardDeviation += Temp;
+		Length3D += InLength;
+	}
+
+	// Finalize has been done on each Property
+	void Add2(FEdge2DProperties& Property)
+	{
+		StandardDeviation = (FMath::Square(StandardDeviation) + FMath::Square(MediumSlop)) * Length3D + (FMath::Square(Property.StandardDeviation) + FMath::Square(Property.MediumSlop)) * Property.Length3D;
+		MediumSlop = MediumSlop * Length3D + Property.MediumSlop * Property.Length3D;
+		Length3D += Property.Length3D;
+
+		MediumSlop /= Length3D;
+		StandardDeviation /= Length3D;
+		StandardDeviation -= FMath::Square(MediumSlop);
+		StandardDeviation = sqrt(StandardDeviation);
+	}
+
+	// Finalize has not been done on each Property
+	void Add(FEdge2DProperties& Property)
+	{
+		StandardDeviation += Property.StandardDeviation;
+		MediumSlop += Property.MediumSlop;
+		Length3D += Property.Length3D;
+	}
 
-		// Finalize has not been done on each Property
-		void Add(FEdge2DProperties& Property)
+	void Finalize()
+	{
+		MediumSlop /= Length3D;
+		StandardDeviation /= Length3D;
+		StandardDeviation -= FMath::Square(MediumSlop);
+		if (StandardDeviation < 0)
 		{
-			StandardDeviation += Property.StandardDeviation;
-			MediumSlop += Property.MediumSlop;
-			Length3D += Property.Length3D;
+			StandardDeviation = 0;
 		}
-
-		void Finalize()
+		else
 		{
-			MediumSlop /= Length3D;
-			StandardDeviation /= Length3D;
-			StandardDeviation -= FMath::Square(MediumSlop);
-			if (StandardDeviation < 0)
-			{
-				StandardDeviation = 0;
-			}
-			else
-			{
-				StandardDeviation = sqrt(StandardDeviation);
-			}
-			IsoType = EIso::UndefinedIso;
+			StandardDeviation = sqrt(StandardDeviation);
+		}
+		IsoType = EIso::UndefinedIso;
 
-			if (MediumSlop < 0.2)
+		if (MediumSlop < 0.2)
+		{
+			if (StandardDeviation < 0.1)
 			{
-				if (StandardDeviation < 0.1)
-				{
-					IsoType = EIso::IsoU;
-				}
+				IsoType = EIso::IsoU;
 			}
-			else if (MediumSlop > 1.8)
+		}
+		else if (MediumSlop > 1.8)
+		{
+			if (StandardDeviation < 0.1)
 			{
-				if (StandardDeviation < 0.1)
-				{
-					IsoType = EIso::IsoV;
-				}
+				IsoType = EIso::IsoV;
 			}
 		}
+	}
 
-	};
+};
 
+/**
+ * Cutting point used for thin zone purpose
+ */
+struct CADKERNEL_API FImposedCuttingPoint
+{
 	/**
-	 * Cutting point used for thin zone purpose
+	 * coordinate of the edge's mesh nodes
 	 */
-	struct CADKERNEL_API FImposedCuttingPoint
-	{
-		/**
-		 * coordinate of the edge's mesh nodes
-		 */
-		double Coordinate;
-		int32 OppositNodeIndex = -1;
+	double Coordinate;
+	int32 OppositNodeIndex = -1;
 
-		FImposedCuttingPoint()
-		{
-		}
+	FImposedCuttingPoint()
+	{
+	}
 
-		FImposedCuttingPoint(double InCoordinate, int32 NodeIndex1)
-			: Coordinate(InCoordinate)
-			, OppositNodeIndex(NodeIndex1)
-		{
-		};
+	FImposedCuttingPoint(double InCoordinate, int32 NodeIndex1)
+		: Coordinate(InCoordinate)
+		, OppositNodeIndex(NodeIndex1)
+	{
 	};
+};
 
+/**
+ * Cutting point used for meshing purpose
+ */
+struct CADKERNEL_API FCuttingPoint
+{
 	/**
-	 * Cutting point used for meshing purpose
+	 * coordinate of the edge's mesh nodes
 	 */
-	struct CADKERNEL_API FCuttingPoint
-	{
-		/**
-		 * coordinate of the edge's mesh nodes
-		 */
-		double Coordinate;
-		ECoordinateType Type;
-		int32 OppositNodeIndex = -1;
-		int32 OppositNodeIndex2 = -1;
-		double IsoDeltaU;
-
-		FCuttingPoint()
-		{
-		}
-
-		FCuttingPoint(double InCoordinate, ECoordinateType InType)
-			: Coordinate(InCoordinate)
-			, Type(InType)
-			, OppositNodeIndex(-1)
-			, OppositNodeIndex2(-1)
-			, IsoDeltaU(HUGE_VAL)
-		{
-		}
+	double Coordinate;
+	ECoordinateType Type;
+	int32 OppositNodeIndex = -1;
+	int32 OppositNodeIndex2 = -1;
+	double IsoDeltaU;
 
-		FCuttingPoint(double InCoordinate, ECoordinateType InType, int32 NodeIndex1, int32 NodeIndex2)
-			: Coordinate(InCoordinate)
-			, Type(InType)
-			, OppositNodeIndex(NodeIndex1)
-			, OppositNodeIndex2(NodeIndex2)
-			, IsoDeltaU(HUGE_VAL)
-		{
-		}
+	FCuttingPoint()
+	{
+	}
+
+	FCuttingPoint(double InCoordinate, ECoordinateType InType)
+		: Coordinate(InCoordinate)
+		, Type(InType)
+		, OppositNodeIndex(-1)
+		, OppositNodeIndex2(-1)
+		, IsoDeltaU(HUGE_VAL)
+	{
+	}
+
+	FCuttingPoint(double InCoordinate, ECoordinateType InType, int32 NodeIndex1, int32 NodeIndex2)
+		: Coordinate(InCoordinate)
+		, Type(InType)
+		, OppositNodeIndex(NodeIndex1)
+		, OppositNodeIndex2(NodeIndex2)
+		, IsoDeltaU(HUGE_VAL)
+	{
+	}
+
+	FCuttingPoint(double InCoordinate, ECoordinateType InType, int32 InOppositeNodeId, double DeltaU)
+		: Coordinate(InCoordinate)
+		, Type(InType)
+		, OppositNodeIndex(InOppositeNodeId)
+		, OppositNodeIndex2(-1)
+		, IsoDeltaU(DeltaU)
+	{
+	}
+};
 
-		FCuttingPoint(double InCoordinate, ECoordinateType InType, int32 InOppositeNodeId, double DeltaU)
-			: Coordinate(InCoordinate)
-			, Type(InType)
-			, OppositNodeIndex(InOppositeNodeId)
-			, OppositNodeIndex2(-1)
-			, IsoDeltaU(DeltaU)
-		{
-		}
-	};
+struct CADKERNEL_API FCuttingGrid
+{
+	TArray<FCuttingPoint> Coordinates[2];
 
-	struct CADKERNEL_API FCuttingGrid
+	constexpr TArray<FCuttingPoint>& operator[](EIso Iso)
 	{
-		TArray<FCuttingPoint> Coordinates[2];
+		ensureCADKernel(Iso == 0 || Iso == 1);
+		return Coordinates[Iso];
+	}
 
-		constexpr TArray<FCuttingPoint>& operator[](EIso Iso)
-		{
-			ensureCADKernel(Iso == 0 || Iso == 1);
-			return Coordinates[Iso];
-		}
-
-		constexpr const TArray<FCuttingPoint>& operator[](EIso Iso) const
-		{
-			ensureCADKernel(Iso == 0 || Iso == 1);
-			return Coordinates[Iso];
-		}
-	};
+	constexpr const TArray<FCuttingPoint>& operator[](EIso Iso) const
+	{
+		ensureCADKernel(Iso == 0 || Iso == 1);
+		return Coordinates[Iso];
+	}
+};
 
 
-	template<typename FCuttingPointType>
-	void GetCuttingPointCoordinates(const TArray<FCuttingPointType>& CuttingPoints, TArray<double>& CuttingPointCoordinates)
+template<typename FCuttingPointType>
+void GetCuttingPointCoordinates(const TArray<FCuttingPointType>& CuttingPoints, TArray<double>& CuttingPointCoordinates)
+{
+	CuttingPointCoordinates.Empty(CuttingPoints.Num());
+	for (const FCuttingPointType& CuttingPoint : CuttingPoints)
 	{
-		CuttingPointCoordinates.Empty(CuttingPoints.Num());
-		for (const FCuttingPointType& CuttingPoint : CuttingPoints)
-		{
-			CuttingPointCoordinates.Add(CuttingPoint.Coordinate);
-		}
-	};
+		CuttingPointCoordinates.Add(CuttingPoint.Coordinate);
+	}
+};
 
 
 } // namespace CADKernel
diff --git a/Engine/Source/Runtime/Datasmith/CADKernel/Public/CADKernel/Topo/TopologicalFace.h b/Engine/Source/Runtime/Datasmith/CADKernel/Public/CADKernel/Topo/TopologicalFace.h
--- a/Engine/Source/Runtime/Datasmith/CADKernel/Public/CADKernel/Topo/TopologicalFace.h
+++ b/Engine/Source/Runtime/Datasmith/CADKernel/Public/CADKernel/Topo/TopologicalFace.h
@@ -11,548 +11,556 @@
 
 namespace CADKernel
 {
-	enum class EStatut : uint8
-	{
-		Interior = 0,
-		Exterior,
-		Border
-	};
-
-	enum class EQuadType : uint8
-	{
-		Unset = 0,
-		Quadrangular,
-		Triangular,
-		Other
-	};
 
-	class FGrid;
-	class FFaceMesh;
-	class FThinZone;
-	class FThinZoneFinder;
-	class FBezierSurface;
-	class FSegmentCurve;
+enum class EStatut : uint8
+{
+	Interior = 0,
+	Exterior,
+	Border
+};
 
-	class CADKERNEL_API FTopologicalFace : public FTopologicalEntity, public FMetadataDictionary
+enum class EQuadType : uint8
+{
+	Unset = 0,
+	Quadrangular,
+	Triangular,
+	Other
+};
+
+class FGrid;
+class FFaceMesh;
+class FThinZone;
+class FThinZoneFinder;
+class FBezierSurface;
+class FSegmentCurve;
+
+class CADKERNEL_API FTopologicalFace : public FTopologicalEntity, public FMetadataDictionary
+{
+	friend class FEntity;
+	//friend class FShell;
+
+protected:
+
+	TSharedPtr<FSurface> CarrierSurface;
+	TArray<TSharedPtr<FTopologicalLoop>> Loops;
+	FShell* HostedBy;
+
+	mutable TCache<FSurfacicBoundary> Boundary;
+
+	TSharedPtr<FFaceMesh> Mesh;
+
+	/**
+	 * Final U&V coordinates of the surface's mesh grid
+	 */
+	FCoordinateGrid MeshCuttingCoordinates;
+
+	/**
+	 * Temporary discretization of the surface used to compute the mesh of the edge
+	 */
+	FCoordinateGrid CrossingCoordinates;
+
+	/**
+	 * Min delta U at the crossing u coordinate to respect meshing criteria
+	 */
+	FCoordinateGrid CrossingPointDeltaMins;
+
+	/**
+	 * Max delta U at the crossing u coordinate to respect meshing criteria
+	 */
+	FCoordinateGrid CrossingPointDeltaMaxs;
+
+	/**
+	 * Build a non-trimmed trimmed surface
+	 * This constructor has to be completed with one of the three "AddBoundaries" methods to be finalized.
+	 */
+	FTopologicalFace(const TSharedPtr<FSurface>& InCarrierSurface)
+		: FTopologicalEntity()
+		, CarrierSurface(InCarrierSurface)
+		, Mesh(TSharedPtr<FFaceMesh>())
 	{
-		friend class FEntity;
-		//friend class FShell;
-
-	protected:
-
-		TSharedPtr<FSurface> CarrierSurface;
-		TArray<TSharedPtr<FTopologicalLoop>> Loops;
-		TWeakPtr<FShell> HostedBy;
-
-		mutable TCache<FSurfacicBoundary> Boundary;
-
-		TSharedPtr<FFaceMesh> Mesh;
+		ResetElementStatus();
+	}
 
-		/**
-		 * Final U&V coordinates of the surface's mesh grid
-		 */
-		FCoordinateGrid MeshCuttingCoordinates;
+	FTopologicalFace() = default;
 
-		/**
-		 * Temporary discretization of the surface used to compute the mesh of the edge
-		 */
-		FCoordinateGrid CrossingCoordinates;
+	/**
+	 * Compute the bounds of the topological surface according to the Loops
+	 */
+	void ComputeBoundary() const;
 
-		/**
-		 * Min delta U at the crossing u coordinate to respect meshing criteria
-		 */
-		FCoordinateGrid CrossingPointDeltaMins;
+public:
 
-		/**
-		 * Max delta U at the crossing u coordinate to respect meshing criteria
-		 */
-		FCoordinateGrid CrossingPointDeltaMaxs;
-
-		/**
-		 * Build a non-trimmed trimmed surface
-		 * This constructor has to be completed with one of the three "AddBoundaries" methods to be finalized.
-		 */
-		FTopologicalFace(const TSharedPtr<FSurface>& InCarrierSurface)
-			: FTopologicalEntity()
-			, CarrierSurface(InCarrierSurface)
-			, Mesh(TSharedPtr<FFaceMesh>())
-		{
-			ResetElementStatus();
-		}
-
-		FTopologicalFace() = default;
-
-		/**
-		 * Compute the bounds of the topological surface according to the Loops
-		 */
-		void ComputeBoundary() const;
-
-	public:
-
-		virtual void Serialize(FCADKernelArchive& Ar) override
-		{
-			FTopologicalEntity::Serialize(Ar);
-			SerializeIdent(Ar, CarrierSurface);
-			SerializeIdents(Ar, (TArray<TSharedPtr<FEntity>>&) Loops);
-			SerializeIdent(Ar, HostedBy);
-			FMetadataDictionary::Serialize(Ar);
-		}
+	virtual void Serialize(FCADKernelArchive& Ar) override
+	{
+		FTopologicalEntity::Serialize(Ar);
+		SerializeIdent(Ar, CarrierSurface);
+		SerializeIdents(Ar, (TArray<TSharedPtr<FEntity>>&) Loops);
+		SerializeIdent(Ar, &HostedBy);
+		FMetadataDictionary::Serialize(Ar);
+	}
 
-		virtual void SpawnIdent(FDatabase& Database) override;
+	virtual void SpawnIdent(FDatabase& Database) override;
 
-		virtual void ResetMarkersRecursively() override
-		{
-			ResetMarkers();
-			ResetMarkersRecursivelyOnEntities(Loops);
-			CarrierSurface->ResetMarkersRecursively();
-		}
+	virtual void ResetMarkersRecursively() override
+	{
+		ResetMarkers();
+		ResetMarkersRecursivelyOnEntities(Loops);
+		CarrierSurface->ResetMarkersRecursively();
+	}
 
 #ifdef CADKERNEL_DEV
-		virtual FInfoEntity& GetInfo(FInfoEntity&) const override;
+	virtual FInfoEntity& GetInfo(FInfoEntity&) const override;
 #endif
 
-		virtual EEntity GetEntityType() const override
-		{
-			return EEntity::TopologicalFace;
-		}
-
-		virtual TSharedPtr<FEntityGeom> ApplyMatrix(const FMatrixH& InMatrix) const override;
+	virtual EEntity GetEntityType() const override
+	{
+		return EEntity::TopologicalFace;
+	}
 
-		const FSurfacicTolerance& GetIsoTolerances() const
-		{
-			return CarrierSurface->GetIsoTolerances();
-		}
+	virtual TSharedPtr<FEntityGeom> ApplyMatrix(const FMatrixH& InMatrix) const override;
 
-		double GetIsoTolerance(EIso Iso) const
-		{
-			return CarrierSurface->GetIsoTolerance(Iso);
-		}
+	const FSurfacicTolerance& GetIsoTolerances() const
+	{
+		return CarrierSurface->GetIsoTolerances();
+	}
 
-		const FSurfacicBoundary& GetBoundary() const
-		{
-			if (!Boundary.IsValid())
-			{
-				ComputeBoundary();
-			}
-			return Boundary;
-		};
+	double GetIsoTolerance(EIso Iso) const
+	{
+		return CarrierSurface->GetIsoTolerance(Iso);
+	}
 
-		virtual int32 FaceCount() const override
+	const FSurfacicBoundary& GetBoundary() const
+	{
+		if (!Boundary.IsValid())
 		{
-			return 1;
+			ComputeBoundary();
 		}
+		return Boundary;
+	};
 
-		virtual void GetFaces(TArray<TSharedPtr<FTopologicalFace>>& OutFaces) override
-		{
-			if (!HasMarker1())
-			{
-				OutFaces.Emplace(StaticCastSharedRef<FTopologicalFace>(AsShared()));
-				SetMarker1();
-			}
-		}
+	virtual int32 FaceCount() const override
+	{
+		return 1;
+	}
 
-		void ResetHost()
+	virtual void GetFaces(TArray<TSharedPtr<FTopologicalFace>>& OutFaces) override
+	{
+		if (!HasMarker1())
 		{
-			HostedBy.Reset();
+			OutFaces.Emplace(StaticCastSharedRef<FTopologicalFace>(AsShared()));
+			SetMarker1();
 		}
+	}
 
-		TWeakPtr<FShell>& GetHost()
-		{
-			return HostedBy;
-		}
+	void ResetHost()
+	{
+		HostedBy = nullptr;
+	}
 
-		void SetHost(TSharedPtr<FShell> Shell)
-		{
-			HostedBy = Shell;
-		}
+	FShell* GetHost()
+	{
+		return HostedBy;
+	}
 
-		void CompleteMetadata()
+	void SetHost(FShell* Shell)
+	{
+		HostedBy = Shell;
+	}
+
+	void CompleteMetadata()
+	{
+		if (HostedBy != nullptr)
 		{
-			if (HostedBy.IsValid())
-			{
-				CompleteDictionary((const FMetadataDictionary&)*HostedBy.Pin());
-			}
+			CompleteDictionary((const FMetadataDictionary&)*HostedBy);
 		}
+	}
 
-		// ======   Loop Functions   ======
+	// ======   Loop Functions   ======
 
-		void RemoveLoop(const TSharedPtr<FTopologicalLoop>& Loop);
-		void AddLoop(const TSharedPtr<FTopologicalLoop>& Loop);
+	void RemoveLoop(const TSharedPtr<FTopologicalLoop>& Loop);
+	void AddLoop(const TSharedPtr<FTopologicalLoop>& Loop);
 
-		/**
-		 * Trimmed the face with an outer boundary (first boundary of the array) and inners boundaries
-		 */
-		void AddLoops(const TArray<TSharedPtr<FTopologicalLoop>>& Loops, int32& DoubtfulLoopOrientationCount);
+	/**
+	 * Trimmed the face with an outer boundary (first boundary of the array) and inners boundaries
+	 */
+	void AddLoops(const TArray<TSharedPtr<FTopologicalLoop>>& Loops, int32& DoubtfulLoopOrientationCount);
 
-		/**
-		 * Trimmed the face with curves i.e. Edges will be build from the cures to make boundaries.
-		 */
-		void AddLoops(const TArray<TSharedPtr<FCurve>>& Restrictions);
+	/**
+	 * Trimmed the face with curves i.e. Edges will be build from the cures to make boundaries.
+	 */
+	void AddLoops(const TArray<TSharedPtr<FCurve>>& Restrictions);
 
-		/**
-		 * Trimmed the face with its natural limit curves (Iso UMin,  ...). This function is called to trim untrimmed topological face.
-		 * This function should not be called if the topological face already has a loop.
-		 */
-		void ApplyNaturalLoops();
+	/**
+	 * Trimmed the face with its natural limit curves (Iso UMin,  ...). This function is called to trim untrimmed topological face.
+	 * This function should not be called if the topological face already has a loop.
+	 */
+	void ApplyNaturalLoops();
 
-		/**
-		 * Trimmed the face with the boundary limit curves (Iso UMin,  ...). This function is called to trim untrimmed topological face.
-		 * This function should not be called if the topological face already has a loop.
-		 */
-		void ApplyNaturalLoops(const FSurfacicBoundary& Boundaries);
+	/**
+	 * Trimmed the face with the boundary limit curves (Iso UMin,  ...). This function is called to trim untrimmed topological face.
+	 * This function should not be called if the topological face already has a loop.
+	 */
+	void ApplyNaturalLoops(const FSurfacicBoundary& Boundaries);
 
-		int32 LoopCount() const 
-		{
-			return Loops.Num();
-		}
+	int32 LoopCount() const
+	{
+		return Loops.Num();
+	}
 
-		const TArray<TSharedPtr<FTopologicalLoop>>& GetLoops() const
+	const TArray<TSharedPtr<FTopologicalLoop>>& GetLoops() const
+	{
+		return Loops;
+	}
+
+	/**
+	 * Get a sampling of each loop of the face
+	 * @param OutLoopSamplings an array of 2d points
+	 * @param OutAABBs an array of 2d axis aligned bounding box of each boundary
+	 */
+	const void Get2DLoopSampling(TArray<TArray<FPoint2D>>& OutLoopSamplings) const;
+
+	// ======   Loop edge Functions   ======
+
+	/**
+	 * @return the twin edge of linked edge belonging this topological face
+	 */
+	const FTopologicalEdge* GetLinkedEdge(const FTopologicalEdge& LinkedEdge) const;
+
+	/**
+	 * Finds the boundary containing the twin edge of Edge
+	 * @param Edge
+	 * @param OutBoundaryIndex: the index of the boundary containing the twin edge
+	 * @param OutEdgeIndex: the index in the boundary of the twin edge
+	 */
+	void GetEdgeIndex(const FTopologicalEdge& Edge, int32& OutBoundaryIndex, int32& OutEdgeIndex) const;
+
+	/**
+	 * Count the edges.
+	 */
+	void EdgeCount(int32& EdgeCount) const
+	{
+		for (const TSharedPtr<FTopologicalLoop>& Loop : Loops)
 		{
-			return Loops;
+			EdgeCount += Loop->GetEdges().Num();
 		}
+	}
 
-		/**
-		 * Get a sampling of each loop of the face
-		 * @param OutLoopSamplings an array of 2d points
-		 * @param OutAABBs an array of 2d axis aligned bounding box of each boundary
-		 */
-		const void Get2DLoopSampling(TArray<TArray<FPoint2D>>& OutLoopSamplings) const;
-
-		// ======   Loop edge Functions   ======
-
-		/**
-		 * @return the twin edge of linked edge belonging this topological face
-		 */
-		const FTopologicalEdge* GetLinkedEdge(const FTopologicalEdge& LinkedEdge) const;
-
-		/**
-		 * Finds the boundary containing the twin edge of Edge
-		 * @param Edge
-		 * @param OutBoundaryIndex: the index of the boundary containing the twin edge
-		 * @param OutEdgeIndex: the index in the boundary of the twin edge
-		 */
-		void GetEdgeIndex(const FTopologicalEdge& Edge, int32& OutBoundaryIndex, int32& OutEdgeIndex) const;
-
-		/**
-		 * Count the edges.
-		 */
-		void EdgeCount(int32& EdgeCount) const 
+	/**
+	 * Add active Edge that has not marker 1 in the edge array.
+	 * Marker 1 has to be reset at the end.
+	 */
+	void GetActiveEdges(TArray<TSharedPtr<FTopologicalEdge>>& OutEdges) const
+	{
+		for (const TSharedPtr<FTopologicalLoop>& Loop : Loops)
 		{
-			for (const TSharedPtr<FTopologicalLoop>& Loop : Loops)
-			{
-				EdgeCount += Loop->GetEdges().Num();
-			}
+			Loop->GetActiveEdges(OutEdges);
 		}
+	}
 
-		/**
-		 * Add active Edge that has not marker 1 in the edge array.
-		 * Marker 1 has to be reset at the end.
-		 */
-		void GetActiveEdges(TArray<TSharedPtr<FTopologicalEdge>>& OutEdges) const
-		{
-			for (const TSharedPtr<FTopologicalLoop>& Loop : Loops)
-			{
-				Loop->GetActiveEdges(OutEdges);
-			}
-		}
+	// ======   Carrier Surface Functions   ======
 
-		// ======   Carrier Surface Functions   ======
-
-		TSharedRef<FSurface> GetCarrierSurface() const
-		{
-			return CarrierSurface.ToSharedRef();
-		}
+	TSharedRef<FSurface> GetCarrierSurface() const
+	{
+		return CarrierSurface.ToSharedRef();
+	}
 
-		// ======   Point evaluation Functions   ======
+	// ======   Point evaluation Functions   ======
 
-		void EvaluatePointGrid(const FCoordinateGrid& Coordinates, FSurfacicSampling& OutPoints, bool bComputeNormals = false) const
-		{
-			CarrierSurface->EvaluatePointGrid(Coordinates, OutPoints, bComputeNormals);
-		}
+	void EvaluatePointGrid(const FCoordinateGrid& Coordinates, FSurfacicSampling& OutPoints, bool bComputeNormals = false) const
+	{
+		CarrierSurface->EvaluatePointGrid(Coordinates, OutPoints, bComputeNormals);
+	}
 
-		void EvaluateGrid(FGrid& Grid) const;
+	void EvaluateGrid(FGrid& Grid) const;
 
-		// ======   Sample Functions   ======
+	// ======   Sample Functions   ======
 
-		/**
-		 * Generate a pre-sampling of the surface saved in CrossingCoordinate.
-		 * This sampling is light enough to allow a fast computation of the grid, precise enough to compute accurately meshing criteria
-		 */
-		void Presample();
+	/**
+	 * Generate a pre-sampling of the surface saved in CrossingCoordinate.
+	 * This sampling is light enough to allow a fast computation of the grid, precise enough to compute accurately meshing criteria
+	 */
+	void Presample();
 
-		// ======   Topo Functions   ======
+	// ======   Topo Functions   ======
 
-		/**
-		 * Checks if the face and the other face have the same boundaries i.e. each non degenerated edge is linked to an edge of the other face
-		 */
-		bool HasSameBoundariesAs(const TSharedPtr<FTopologicalFace>& OtherFace) const;
+	/**
+	 * Checks if the face and the other face have the same boundaries i.e. each non degenerated edge is linked to an edge of the other face
+	 */
+	bool HasSameBoundariesAs(const TSharedPtr<FTopologicalFace>& OtherFace) const;
 
-		/**
-		 * Disconnects the face of its neighbors
-		 */
-		void RemoveLinksWithNeighbours();
+	/**
+	 * Disconnects the face of its neighbors
+	 */
+	void RemoveLinksWithNeighbours();
 
-		virtual void FillTopologyReport(FTopologyReport& Report) const override;
+	virtual void FillTopologyReport(FTopologyReport& Report) const override;
 
-		// ======   Meshing Function   ======
+	// ======   Meshing Function   ======
 
-		TSharedRef<FFaceMesh> GetOrCreateMesh(const TSharedRef<FModelMesh>& ShellMesh);
+	TSharedRef<FFaceMesh> GetOrCreateMesh(const TSharedRef<FModelMesh>& ShellMesh);
 
-		const bool HasTesselation() const
-		{
-			return Mesh.IsValid();
-		}
+	const bool HasTesselation() const
+	{
+		return Mesh.IsValid();
+	}
 
-		const TSharedRef<FFaceMesh> GetMesh() const
-		{
-			ensureCADKernel(Mesh.IsValid());
-			return Mesh.ToSharedRef();
-		}
+	const TSharedRef<FFaceMesh> GetMesh() const
+	{
+		ensureCADKernel(Mesh.IsValid());
+		return Mesh.ToSharedRef();
+	}
 
-		void InitDeltaUs();
+	void InitDeltaUs();
 
-		void ChooseFinalDeltaUs();
+	void ChooseFinalDeltaUs();
 
 
-		const TArray<double>& GetCuttingCoordinatesAlongIso(EIso Iso) const
-		{
-			return MeshCuttingCoordinates[Iso];
-		}
+	const TArray<double>& GetCuttingCoordinatesAlongIso(EIso Iso) const
+	{
+		return MeshCuttingCoordinates[Iso];
+	}
 
-		TArray<double>& GetCuttingCoordinatesAlongIso(EIso Iso)
-		{
-			return MeshCuttingCoordinates[Iso];
-		}
+	TArray<double>& GetCuttingCoordinatesAlongIso(EIso Iso)
+	{
+		return MeshCuttingCoordinates[Iso];
+	}
 
-		const FCoordinateGrid& GetCuttingPointCoordinates() const
-		{
-			return MeshCuttingCoordinates;
-		}
+	const FCoordinateGrid& GetCuttingPointCoordinates() const
+	{
+		return MeshCuttingCoordinates;
+	}
 
-		FCoordinateGrid& GetCuttingPointCoordinates()
-		{
-			return MeshCuttingCoordinates;
-		}
+	FCoordinateGrid& GetCuttingPointCoordinates()
+	{
+		return MeshCuttingCoordinates;
+	}
 
-		const FCoordinateGrid& GetCrossingPointCoordinates() const
-		{
-			return CrossingCoordinates;
-		}
+	const FCoordinateGrid& GetCrossingPointCoordinates() const
+	{
+		return CrossingCoordinates;
+	}
 
-		const TArray<double>& GetCrossingPointCoordinates(EIso Iso) const
-		{
-			return CrossingCoordinates[Iso];
-		}
+	const TArray<double>& GetCrossingPointCoordinates(EIso Iso) const
+	{
+		return CrossingCoordinates[Iso];
+	}
 
-		TArray<double>& GetCrossingPointCoordinates(EIso Iso)
-		{
-			return CrossingCoordinates[Iso];
-		}
+	TArray<double>& GetCrossingPointCoordinates(EIso Iso)
+	{
+		return CrossingCoordinates[Iso];
+	}
 
 
-		const TArray<double>& GetCrossingPointDeltaMins(EIso Iso) const
-		{
-			return CrossingPointDeltaMins[Iso];
-		}
+	const TArray<double>& GetCrossingPointDeltaMins(EIso Iso) const
+	{
+		return CrossingPointDeltaMins[Iso];
+	}
 
-		TArray<double>& GetCrossingPointDeltaMins(EIso Iso)
-		{
-			return CrossingPointDeltaMins[Iso];
-		}
+	TArray<double>& GetCrossingPointDeltaMins(EIso Iso)
+	{
+		return CrossingPointDeltaMins[Iso];
+	}
 
-		const TArray<double>& GetCrossingPointDeltaMaxs(EIso Iso) const
-		{
-			return CrossingPointDeltaMaxs[Iso];
-		}
+	const TArray<double>& GetCrossingPointDeltaMaxs(EIso Iso) const
+	{
+		return CrossingPointDeltaMaxs[Iso];
+	}
 
-		TArray<double>& GetCrossingPointDeltaMaxs(EIso Iso)
-		{
-			return CrossingPointDeltaMaxs[Iso];
-		}
+	TArray<double>& GetCrossingPointDeltaMaxs(EIso Iso)
+	{
+		return CrossingPointDeltaMaxs[Iso];
+	}
 
-		// ======   State, Type Functions   ======
+	// ======   State, Type Functions   ======
 
-		const EQuadType GetQuadType() const
-		{
-			return QuadType;
-		}
+	const EQuadType GetQuadType() const
+	{
+		return QuadType;
+	}
 
-		const bool HasThinZone() const
-		{
-			return ((States & EHaveStates::ThinZone) == EHaveStates::ThinZone);
-		}
+	const bool HasThinZone() const
+	{
+		return ((States & EHaveStates::ThinZone) == EHaveStates::ThinZone);
+	}
 
-		void SetHasThinZone()
-		{
-			States |= EHaveStates::ThinZone;
-		}
+	void SetHasThinZone()
+	{
+		States |= EHaveStates::ThinZone;
+	}
 
-		void ResetHasThinSurface()
-		{
-			States &= ~EHaveStates::ThinZone;
-		}
+	void ResetHasThinSurface()
+	{
+		States &= ~EHaveStates::ThinZone;
+	}
 
-		bool IsBackOriented() const
-		{
-			return ((States & EHaveStates::IsBackOriented) == EHaveStates::IsBackOriented);
-		}
+	bool IsBackOriented() const
+	{
+		return ((States & EHaveStates::IsBackOriented) == EHaveStates::IsBackOriented);
+	}
 
-		void SetBackOriented() const
-		{
-			States |= EHaveStates::IsBackOriented;
-		}
+	void SetBackOriented() const
+	{
+		States |= EHaveStates::IsBackOriented;
+	}
 
-		void ResetBackOriented() const
-		{
-			States &= ~EHaveStates::IsBackOriented;
-		}
+	void ResetBackOriented() const
+	{
+		States &= ~EHaveStates::IsBackOriented;
+	}
+
+
+	// =========================================================================================================================================================================================================
+	// =========================================================================================================================================================================================================
+	// =========================================================================================================================================================================================================
+	//
+	//
+	//                                                                            NOT YET REVIEWED
+	//
+	//
+	// =========================================================================================================================================================================================================
+	// =========================================================================================================================================================================================================
+	// =========================================================================================================================================================================================================
+
+
+
+	// ======================================================================================================================================================================================================================
+	// Quad properties for meshing scheduling ===============================================================================================================================================================================
+	// ======================================================================================================================================================================================================================
+private:
+	TArray<TSharedPtr<FTopologicalVertex>> SurfaceCorners;
+	TArray<int32> StartSideIndices;
+	TArray<FEdge2DProperties> SideProperties;
+	int32 NumOfMeshedSide = 0;
+	double LoopLength;
+	double LengthOfMeshedSide = 0;
+	double QuadCriteria = 0;
+	FSurfaceCurvature Curvatures;
+	EQuadType QuadType = EQuadType::Unset;
+
+public:
+	void ComputeQuadCriteria();
+	double GetQuadCriteria();
+
+	const FSurfaceCurvature& GetCurvatures() const
+	{
+		return Curvatures;
+	}
 
+	FSurfaceCurvature& GetCurvatures()
+	{
+		return Curvatures;
+	}
 
-		// =========================================================================================================================================================================================================
-		// =========================================================================================================================================================================================================
-		// =========================================================================================================================================================================================================
-		//
-		//
-		//                                                                            NOT YET REVIEWED
-		//
-		//
-		// =========================================================================================================================================================================================================
-		// =========================================================================================================================================================================================================
-		// =========================================================================================================================================================================================================
-
-
-
-		// ======================================================================================================================================================================================================================
-		// Quad properties for meshing scheduling ===============================================================================================================================================================================
-		// ======================================================================================================================================================================================================================
-	private:
-		TArray<TSharedPtr<FTopologicalVertex>> SurfaceCorners;
-		TArray<int32> StartSideIndices;
-		TArray<FEdge2DProperties> SideProperties;
-		int32 NumOfMeshedSide = 0;
-		double LoopLength;
-		double LengthOfMeshedSide = 0;
-		double QuadCriteria = 0;
-		FSurfaceCurvature Curvatures;
-		EQuadType QuadType = EQuadType::Unset;
-
-	public:
-		void ComputeQuadCriteria();
-		double GetQuadCriteria();
-
-		const FSurfaceCurvature& GetCurvatures() const
-		{
-			return Curvatures;
-		}
+	const FIsoCurvature& GetCurvature(EIso Iso) const
+	{
+		return Curvatures[Iso];
+	}
 
-		FSurfaceCurvature& GetCurvatures()
-		{
-			return Curvatures;
-		}
+	void ComputeSurfaceSideProperties();
 
-		const FIsoCurvature& GetCurvature(EIso Iso) const
-		{
-			return Curvatures[Iso];
-		}
+	/**
+	 * Defines if the surface is either EQuadType::QUAD, either EQuadType::TRIANGULAR or EQuadType::OTHER
+	 */
+	void DefineSurfaceType();
 
-		void ComputeSurfaceSideProperties();
+	const TArray<FEdge2DProperties>& GetSideProperties() const
+	{
+		return SideProperties;
+	}
 
-		/**
-		 * Defines if the surface is either EQuadType::QUAD, either EQuadType::TRIANGULAR or EQuadType::OTHER
-		 */
-		void DefineSurfaceType();
+	FEdge2DProperties& GetSideProperty(int32 Index)
+	{
+		return SideProperties[Index];
+	}
 
-		const TArray<FEdge2DProperties>& GetSideProperties() const
-		{
-			return SideProperties;
-		}
+	const FEdge2DProperties& GetSideProperty(int32 Index) const
+	{
+		return SideProperties[Index];
+	}
 
-		FEdge2DProperties& GetSideProperty(int32 Index)
-		{
-			return SideProperties[Index];
-		}
+	int32& MeshedSideNum()
+	{
+		return NumOfMeshedSide;
+	}
 
-		const FEdge2DProperties& GetSideProperty(int32 Index) const
-		{
-			return SideProperties[Index];
-		}
+	const int32& MeshedSideNum() const
+	{
+		return NumOfMeshedSide;
+	}
 
-		int32& MeshedSideNum()
-		{
-			return NumOfMeshedSide;
-		}
+	void AddMeshedLength(double Length)
+	{
+		LengthOfMeshedSide += Length;
+	}
 
-		const int32& MeshedSideNum() const
-		{
-			return NumOfMeshedSide;
-		}
+	double MeshedSideRatio() const
+	{
+		return LengthOfMeshedSide / LoopLength;
+	}
 
-		void AddMeshedLength(double Length)
-		{
-			LengthOfMeshedSide += Length;
-		}
+	int32 GetStartEdgeIndexOfSide(int32 Index) const
+	{
+		return StartSideIndices[Index];
+	}
 
-		double MeshedSideRatio() const
-		{
-			return LengthOfMeshedSide / LoopLength;
-		}
+	const TArray<int32>& GetStartSideIndices() const
+	{
+		return StartSideIndices;
+	}
 
-		int32 GetStartEdgeIndexOfSide(int32 Index) const
+	int32 GetSideIndex(FTopologicalEdge& Edge) const
+	{
+		int32 EdgeIndex = Loops[0]->GetEdgeIndex(Edge);
+		if (EdgeIndex < 0)
 		{
-			return StartSideIndices[Index];
+			return -1;
 		}
+		return GetSideIndex(EdgeIndex);
+	}
 
-		const TArray<int32>& GetStartSideIndices() const
+	int32 GetSideIndex(int32 EdgeIndex) const
+	{
+		if (StartSideIndices.Num() == 0)
 		{
-			return StartSideIndices;
+			return -1;
 		}
 
-		int32 GetSideIndex(FTopologicalEdge& Edge) const
+		if (StartSideIndices[0] > EdgeIndex)
 		{
-			int32 EdgeIndex = Loops[0]->GetEdgeIndex(Edge);
-			if (EdgeIndex < 0)
-			{
-				return -1;
-			}
-			return GetSideIndex(EdgeIndex);
+			return (int32)StartSideIndices.Num() - 1;
 		}
-
-		int32 GetSideIndex(int32 EdgeIndex) const
+		else
 		{
-			if (StartSideIndices.Num() == 0)
-			{
-				return -1;
-			}
-
-			if (StartSideIndices[0] > EdgeIndex)
+			for (int32 SideIndex = 0; SideIndex < StartSideIndices.Num() - 1; ++SideIndex)
 			{
-				return (int32)StartSideIndices.Num() - 1;
-			}
-			else
-			{
-				for (int32 SideIndex = 0; SideIndex < StartSideIndices.Num() - 1; ++SideIndex)
+				if (StartSideIndices[SideIndex] <= EdgeIndex && EdgeIndex < StartSideIndices[SideIndex + 1])
 				{
-					if (StartSideIndices[SideIndex] <= EdgeIndex && EdgeIndex < StartSideIndices[SideIndex + 1])
-					{
-						return SideIndex;
-					}
+					return SideIndex;
 				}
-				return (int32)StartSideIndices.Num() - 1;
 			}
+			return (int32)StartSideIndices.Num() - 1;
 		}
+	}
 
-	};
+};
 
-	struct FFaceSubset
-	{
-		TArray<TSharedPtr<FTopologicalFace>> Faces;
-		int32 BorderEdgeCount = 0;
-		int32 NonManifoldEdgeCount = 0;
-		TWeakPtr<FShell> MainShell;
-		FString MainName;
-		uint32 MainColor;
-		uint32 MainMaterial;
-	};
+struct FFaceSubset
+{
+	TArray<FTopologicalFace*> Faces;
+	int32 BorderEdgeCount = 0;
+	int32 NonManifoldEdgeCount = 0;
+	FShell* MainShell;
+	FBody* MainBody;
+	FString MainName;
+	uint32 MainColor;
+	uint32 MainMaterial;
+
+	void SetMainShell(TMap<FShell*, int32>& ShellToFaceCount);
+	void SetMainBody(TMap<FBody*, int32>& BodyToFaceCount);
+	void SetMainName(TMap<FString, int32>& NameToFaceCount);
+	void SetMainColor(TMap<uint32, int32>& ColorToFaceCount);
+	void SetMainMaterial(TMap<uint32, int32>& MaterialToFaceCount);
+};
 
 
 } // namespace CADKernel
diff --git a/Engine/Source/Runtime/Datasmith/CADKernel/Public/CADKernel/Topo/TopologicalLink.h b/Engine/Source/Runtime/Datasmith/CADKernel/Public/CADKernel/Topo/TopologicalLink.h
--- a/Engine/Source/Runtime/Datasmith/CADKernel/Public/CADKernel/Topo/TopologicalLink.h
+++ b/Engine/Source/Runtime/Datasmith/CADKernel/Public/CADKernel/Topo/TopologicalLink.h
@@ -116,6 +116,15 @@ namespace CADKernel
 			}
 		}
 
+		void UnlinkTwinEntities()
+		{
+			for (EntityType* Entity : TwinsEntities)
+			{
+				Entity->ResetTopologicalLink();
+			}
+			TwinsEntities.Empty();
+		}
+
 #ifdef CADKERNEL_DEV
 		virtual FInfoEntity& GetInfo(FInfoEntity& Info) const override;
 #endif
diff --git a/Engine/Source/Runtime/Datasmith/CADKernel/Public/CADKernel/Topo/TopologicalLoop.h b/Engine/Source/Runtime/Datasmith/CADKernel/Public/CADKernel/Topo/TopologicalLoop.h
--- a/Engine/Source/Runtime/Datasmith/CADKernel/Public/CADKernel/Topo/TopologicalLoop.h
+++ b/Engine/Source/Runtime/Datasmith/CADKernel/Public/CADKernel/Topo/TopologicalLoop.h
@@ -11,213 +11,214 @@
 namespace CADKernel
 {
 
-	class FTopologicalFace;
-	class FTopologicalVertex;
+class FTopologicalFace;
+class FTopologicalVertex;
 
-	class CADKERNEL_API FOrientedEdge : public TOrientedEntity<FTopologicalEdge>
+class CADKERNEL_API FOrientedEdge : public TOrientedEntity<FTopologicalEdge>
+{
+public:
+	FOrientedEdge(TSharedPtr<FTopologicalEdge>& InEntity, EOrientation InDirection)
+		: TOrientedEntity(InEntity, InDirection)
 	{
-	public:
-		FOrientedEdge(TSharedPtr<FTopologicalEdge>& InEntity, EOrientation InDirection)
-			: TOrientedEntity(InEntity, InDirection)
-		{
-		}
-
-		FOrientedEdge(const FOrientedEdge& OrientiredEntity)
-			: TOrientedEntity(OrientiredEntity)
-		{
-		}
+	}
 
-		FOrientedEdge()
-			: TOrientedEntity()
-		{
-		}
-
-		bool operator==(const FOrientedEdge& Edge) const
-		{
-			return Entity == Edge.Entity;
-		}
+	FOrientedEdge(const FOrientedEdge& OrientiredEntity)
+		: TOrientedEntity(OrientiredEntity)
+	{
+	}
 
-	};
+	FOrientedEdge()
+		: TOrientedEntity()
+	{
+	}
 
-	class CADKERNEL_API FTopologicalLoop : public FTopologicalEntity
+	bool operator==(const FOrientedEdge& Edge) const
 	{
-		friend class FEntity;
-		friend class FTopologicalFace;
-		friend class FTopologicalFace;
-		friend class FTopologicalEdge;
+		return Entity == Edge.Entity;
+	}
 
-	public:
-		FSurfacicBoundary Boundary;
+};
 
-	protected:
-		TArray<FOrientedEdge> Edges;
+class CADKERNEL_API FTopologicalLoop : public FTopologicalEntity
+{
+	friend class FEntity;
+	friend class FTopologicalFace;
+	friend class FTopologicalFace;
+	friend class FTopologicalEdge;
 
-		TWeakPtr<FTopologicalFace> Face;
-		bool bExternalLoop;
+public:
+	FSurfacicBoundary Boundary;
 
-		FTopologicalLoop(const TArray<TSharedPtr<FTopologicalEdge>>& Edges, const TArray<EOrientation>& EdgeDirections);
+protected:
+	TArray<FOrientedEdge> Edges;
 
-		FTopologicalLoop() = default;
+	FTopologicalFace* Face;
+	bool bExternalLoop;
 
-	private:
+	FTopologicalLoop(const TArray<TSharedPtr<FTopologicalEdge>>& Edges, const TArray<EOrientation>& EdgeDirections);
 
-		void SetSurface(TSharedRef<FTopologicalFace> NewDomain)
-		{
-			Face = NewDomain;
-		}
+	FTopologicalLoop() = default;
 
-		void ResetSurface()
-		{
-			Face.Reset();
-		}
+private:
 
-	public:
+	void SetSurface(FTopologicalFace* HostedFace)
+	{
+		Face = HostedFace;
+	}
 
-		~FTopologicalLoop() = default;
+	void ResetSurface()
+	{
+		Face = nullptr;
+	}
 
-		static TSharedPtr<FTopologicalLoop> Make(const TArray<TSharedPtr<FTopologicalEdge>>& EdgeList, const TArray<EOrientation>& EdgeDirections, const double GeometricTolerance);
+public:
 
-		void DeleteLoopEdges();
+	~FTopologicalLoop() = default;
 
-		virtual void Serialize(FCADKernelArchive& Ar) override
-		{
-			FTopologicalEntity::Serialize(Ar);
-			SerializeIdents(Ar, (TArray<TOrientedEntity<FEntity>>&) Edges);
-			SerializeIdent(Ar, Face);
-			Ar << bExternalLoop;
-		}
+	static TSharedPtr<FTopologicalLoop> Make(const TArray<TSharedPtr<FTopologicalEdge>>& EdgeList, const TArray<EOrientation>& EdgeDirections, const double GeometricTolerance);
 
-		virtual void SpawnIdent(FDatabase& Database) override
-		{
-			if (!FEntity::SetId(Database))
-			{
-				return;
-			}
+	void DeleteLoopEdges();
 
-			SpawnIdentOnEntities((TArray<TOrientedEntity<FEntity>>&) Edges, Database);
-		}
+	virtual void Serialize(FCADKernelArchive& Ar) override
+	{
+		FTopologicalEntity::Serialize(Ar);
+		SerializeIdents(Ar, (TArray<TOrientedEntity<FEntity>>&) Edges);
+		SerializeIdent(Ar, &Face);
+		Ar << bExternalLoop;
+	}
 
-		virtual void ResetMarkersRecursively() override
+	virtual void SpawnIdent(FDatabase& Database) override
+	{
+		if (!FEntity::SetId(Database))
 		{
-			ResetMarkers();
-			ResetMarkersRecursivelyOnEntities((TArray<TOrientedEntity<FEntity>>&) Edges);
+			return;
 		}
 
+		SpawnIdentOnEntities((TArray<TOrientedEntity<FEntity>>&) Edges, Database);
+	}
+
+	virtual void ResetMarkersRecursively() override
+	{
+		ResetMarkers();
+		ResetMarkersRecursivelyOnEntities((TArray<TOrientedEntity<FEntity>>&) Edges);
+	}
+
 #ifdef CADKERNEL_DEV
-		virtual FInfoEntity& GetInfo(FInfoEntity&) const override;
+	virtual FInfoEntity& GetInfo(FInfoEntity&) const override;
 #endif
 
-		virtual EEntity GetEntityType() const override
-		{
-			return EEntity::TopologicalLoop;
-		}
-
-		const int32 EdgeCount() const
-		{
-			return Edges.Num();
-		}
+	virtual EEntity GetEntityType() const override
+	{
+		return EEntity::TopologicalLoop;
+	}
 
-		const TArray<FOrientedEdge>& GetEdges() const
-		{
-			return Edges;
-		}
+	const int32 EdgeCount() const
+	{
+		return Edges.Num();
+	}
 
-		TArray<FOrientedEdge>& GetEdges()
-		{
-			return Edges;
-		}
+	const TArray<FOrientedEdge>& GetEdges() const
+	{
+		return Edges;
+	}
 
-		/**
-		 * Add active Edge that has not marker 1 in the edge array.
-		 * Marker 1 has to be reset at the end.
-		 */
-		void GetActiveEdges(TArray<TSharedPtr<FTopologicalEdge>>& OutEdges) const
+	TArray<FOrientedEdge>& GetEdges()
+	{
+		return Edges;
+	}
+
+	/**
+	 * Add active Edge that has not marker 1 in the edge array.
+	 * Marker 1 has to be reset at the end.
+	 */
+	void GetActiveEdges(TArray<TSharedPtr<FTopologicalEdge>>& OutEdges) const
+	{
+		for (const FOrientedEdge& Edge : Edges)
 		{
-			for (const FOrientedEdge& Edge : Edges)
+			TSharedPtr<FTopologicalEdge> ActiveEdge = Edge.Entity->GetLinkActiveEdge();
+			if (!ActiveEdge->HasMarker1())
 			{
-				TSharedPtr<FTopologicalEdge> ActiveEdge = Edge.Entity->GetLinkActiveEdge();
-				if (!ActiveEdge->HasMarker1())
-				{
-					ActiveEdge->SetMarker1();
-					OutEdges.Emplace(ActiveEdge);
-				}
+				ActiveEdge->SetMarker1();
+				OutEdges.Emplace(ActiveEdge);
 			}
 		}
+	}
 
-		TSharedRef<FTopologicalFace> GetFace() const
-		{
-			return Face.Pin().ToSharedRef();
-		}
+	FTopologicalFace* GetFace() const
+	{
+		return Face;
+	}
 
-		/*
-		 * @return false if the orientation is doubtful
-		 */ 
-		bool Orient();
-		void SwapOrientation();
+	/*
+	 * @return false if the orientation is doubtful
+	 */
+	bool Orient();
+	void SwapOrientation();
 
-		void SetAsInnerBoundary()
-		{
-			bExternalLoop = false;
-		}
+	void SetAsInnerBoundary()
+	{
+		bExternalLoop = false;
+	}
 
-		void ReplaceEdge(TSharedPtr<FTopologicalEdge>& OldEdge, TSharedPtr<FTopologicalEdge>& NewEdge);
-		void ReplaceEdge(TSharedPtr<FTopologicalEdge>& Edge, TArray<TSharedPtr<FTopologicalEdge>>& NewEdges);
-		void ReplaceEdges(TArray<FOrientedEdge>& Candidates, TSharedPtr<FTopologicalEdge>& NewEdge);
+	void ReplaceEdge(TSharedPtr<FTopologicalEdge>& OldEdge, TSharedPtr<FTopologicalEdge>& NewEdge);
+	void ReplaceEdge(TSharedPtr<FTopologicalEdge>& Edge, TArray<TSharedPtr<FTopologicalEdge>>& NewEdges);
+	void ReplaceEdges(TArray<FOrientedEdge>& Candidates, TSharedPtr<FTopologicalEdge>& NewEdge);
 
-		/**
-		 * The Edge is split in two edges : Edge + NewEdge 
-		 * @param bNewEdgeIsFirst == true => StartVertex Connected to Edge, EndVertexConnected to NewEdge
-		 * According to the direction of Edge, if bNewEdgeIsFirst == true, NewEdge is added in the loop after (EOrientation::Front) or before (EOrientation::Back)
-		 */
-		void SplitEdge(FTopologicalEdge& Edge, TSharedPtr<FTopologicalEdge> NewEdge, bool bNewEdgeIsFirst);
+	/**
+	 * The Edge is split in two edges : Edge + NewEdge
+	 * @param bNewEdgeIsFirst == true => StartVertex Connected to Edge, EndVertexConnected to NewEdge
+	 * According to the direction of Edge, if bNewEdgeIsFirst == true, NewEdge is added in the loop after (EOrientation::Front) or before (EOrientation::Back)
+	 */
+	void SplitEdge(FTopologicalEdge& Edge, TSharedPtr<FTopologicalEdge> NewEdge, bool bNewEdgeIsFirst);
 
-		void RemoveEdge(TSharedPtr<FTopologicalEdge>& Edge);
-		//void ReplaceEdgesWithMergedEdge(TArray<TSharedPtr<FTopologicalEdge>>& OldEdges, TSharedPtr<FTopologicalVertex>& MiddleVertex, TSharedPtr<FTopologicalEdge>& NewEdge);
+	void RemoveEdge(TSharedPtr<FTopologicalEdge>& Edge);
+	//void ReplaceEdgesWithMergedEdge(TArray<TSharedPtr<FTopologicalEdge>>& OldEdges, TSharedPtr<FTopologicalVertex>& MiddleVertex, TSharedPtr<FTopologicalEdge>& NewEdge);
 
-		EOrientation GetDirection(TSharedPtr<FTopologicalEdge>& Edge, bool bAllowLinkedEdge = false) const;
+	EOrientation GetDirection(TSharedPtr<FTopologicalEdge>& Edge, bool bAllowLinkedEdge = false) const;
 
-		EOrientation GetDirection(int32 Index) const
-		{
-			return Edges[Index].Direction;
-		}
+	EOrientation GetDirection(int32 Index) const
+	{
+		return Edges[Index].Direction;
+	}
 
-		const TSharedPtr<FTopologicalEdge>& GetEdge(int32 Index) const
-		{
-			return Edges[Index].Entity;
-		}
+	const TSharedPtr<FTopologicalEdge>& GetEdge(int32 Index) const
+	{
+		return Edges[Index].Entity;
+	}
 
-		int32 GetEdgeIndex(const FTopologicalEdge& Edge) const
+	int32 GetEdgeIndex(const FTopologicalEdge& Edge) const
+	{
+		for (int32 Index = 0; Index < Edges.Num(); ++Index)
 		{
-			for (int32 Index = 0; Index < Edges.Num(); ++Index)
+			if (&Edge == Edges[Index].Entity.Get())
 			{
-				if (&Edge == Edges[Index].Entity.Get())
-				{
-					return Index;
-				}
+				return Index;
 			}
-			return -1;
 		}
+		return -1;
+	}
+
+	void Get2DSampling(TArray<FPoint2D>& LoopSampling);
 
-		void Get2DSampling(TArray<FPoint2D>& LoopSampling);
+	/**
+	 * The idea is to remove degenerated edges of the loop i.e. where the surface is degenerated
+	 * - so where projections on these area are hazardous
+	 * - so where 2d curve computation based on hazardous projection is hazardous...
+	 * - so where the sampling could be in self-intersecting
+	 * @return false if the loop is degenerated
+	 */
+	bool Get2DSamplingWithoutDegeneratedEdges(TArray<FPoint2D>& LoopSampling);
 
-		/**
-		 * The idea is to remove degenerated edges of the loop i.e. where the surface is degenerated 
-		 * - so where projections on these area are hazardous
-		 * - so where 2d curve computation based on hazardous projection is hazardous...
-		 * - so where the sampling could be in self-intersecting
-		 * @return false if the loop is degenerated
-		 */
-		bool Get2DSamplingWithoutDegeneratedEdges(TArray<FPoint2D>& LoopSampling);
+	void FindSurfaceCorners(TArray<TSharedPtr<FTopologicalVertex>>& OutCorners, TArray<int32>& OutStartSideIndex) const;
+	void FindBreaks(TArray<TSharedPtr<FTopologicalVertex>>& Ruptures, TArray<int32>& OutStartSideIndex, TArray<double>& RuptureValues) const;
 
-		void FindSurfaceCorners(TArray<TSharedPtr<FTopologicalVertex>>& OutCorners, TArray<int32>& OutStartSideIndex) const;
-		void FindBreaks(TArray<TSharedPtr<FTopologicalVertex>>& Ruptures, TArray<int32>& OutStartSideIndex, TArray<double>& RuptureValues) const;
+	void ComputeBoundaryProperties(const TArray<int32>& StartSideIndex, TArray<FEdge2DProperties>& OutSideProperties) const;
 
-		void ComputeBoundaryProperties(const TArray<int32>& StartSideIndex, TArray<FEdge2DProperties>& OutSideProperties) const;
+	void EnsureLogicalClosing(const double GeometricTolerance);
 
-		void EnsureLogicalClosing(const double GeometricTolerance);
+	void CheckEdgesOrientation();
+	void CheckLoopWithTwoEdgesOrientation();
+	void RemoveDegeneratedEdges();
+};
 
-		void CheckEdgesOrientation();
-		void CheckLoopWithTwoEdgesOrientation();
-		void RemoveDegeneratedEdges();
-	};
 }
