Fix for static analysis warning

diff --git a/Engine/Source/Runtime/Datasmith/CADKernel/Private/CADKernel/Topo/TopologicalFace.cpp b/Engine/Source/Runtime/Datasmith/CADKernel/Private/CADKernel/Topo/TopologicalFace.cpp
--- a/Engine/Source/Runtime/Datasmith/CADKernel/Private/CADKernel/Topo/TopologicalFace.cpp
+++ b/Engine/Source/Runtime/Datasmith/CADKernel/Private/CADKernel/Topo/TopologicalFace.cpp
@@ -403,6 +403,9 @@ void FTopologicalFace::ComputeSurfaceSideProperties()
 
 void FTopologicalFace::DefineSurfaceType()
 {
+	if (CarrierSurface == nullptr)
+		return;
+
 	const double Tolerance3D = CarrierSurface->Get3DTolerance();
 	const double GeometricTolerance = 20.0 * Tolerance3D;
 
@@ -455,7 +458,7 @@ void FTopologicalFace::DefineSurfaceType()
 				ensure(Neighbor != nullptr);
 
 				// it's not a quad surface
-				if (Neighbor->SurfaceCorners.Num() == 0)
+				if (Neighbor==nullptr || Neighbor->SurfaceCorners.Num() == 0)
 				{
 					return;
 				}
