Rename Joiner into Topomaker

diff --git a/Engine/Plugins/Enterprise/DatasmithCADImporter/Source/CADInterfaces/Private/CoreTechFileParserImpl.cpp b/Engine/Plugins/Enterprise/DatasmithCADImporter/Source/CADInterfaces/Private/CoreTechFileParserImpl.cpp
--- a/Engine/Plugins/Enterprise/DatasmithCADImporter/Source/CADInterfaces/Private/CoreTechFileParserImpl.cpp
+++ b/Engine/Plugins/Enterprise/DatasmithCADImporter/Source/CADInterfaces/Private/CoreTechFileParserImpl.cpp
@@ -34,7 +34,7 @@
 
 #include "CADKernel/Topo/Model.h"
 #include "CADKernel/Topo/Body.h"
-#include "CADKernel/Topo/Joiner.h"
+#include "CADKernel/Topo/Topomaker.h"
 
 #include "HAL/FileManager.h"
 #include "Internationalization/Text.h"
diff --git a/Engine/Plugins/Enterprise/DatasmithCADImporter/Source/CADKernelSurface/Public/CADModelToCADKernelConverterBase.h b/Engine/Plugins/Enterprise/DatasmithCADImporter/Source/CADKernelSurface/Public/CADModelToCADKernelConverterBase.h
--- a/Engine/Plugins/Enterprise/DatasmithCADImporter/Source/CADKernelSurface/Public/CADModelToCADKernelConverterBase.h
+++ b/Engine/Plugins/Enterprise/DatasmithCADImporter/Source/CADKernelSurface/Public/CADModelToCADKernelConverterBase.h
@@ -12,7 +12,7 @@
 
 #include "CADKernel/Core/Session.h"
 #include "CADKernel/Topo/Model.h"
-#include "CADKernel/Topo/Joiner.h"
+#include "CADKernel/Topo/Topomaker.h"
 
 struct FDatasmithMeshElementPayload;
 struct FDatasmithTessellationOptions;
diff --git a/Engine/Source/Runtime/Datasmith/CADKernel/Private/CADKernel/Topo/Joiner.cpp b/Engine/Source/Runtime/Datasmith/CADKernel/Private/CADKernel/Topo/Joiner.cpp
deleted file mode 100644
diff --git a/Engine/Source/Runtime/Datasmith/CADKernel/Private/CADKernel/Topo/Model.cpp b/Engine/Source/Runtime/Datasmith/CADKernel/Private/CADKernel/Topo/Model.cpp
--- a/Engine/Source/Runtime/Datasmith/CADKernel/Private/CADKernel/Topo/Model.cpp
+++ b/Engine/Source/Runtime/Datasmith/CADKernel/Private/CADKernel/Topo/Model.cpp
@@ -3,7 +3,7 @@
 
 #include "CADKernel/Core/EntityGeom.h"
 #include "CADKernel/Topo/Body.h"
-#include "CADKernel/Topo/Joiner.h"
+#include "CADKernel/Topo/Topomaker.h"
 #include "CADKernel/Topo/TopologicalEdge.h"
 #include "CADKernel/Topo/TopologicalFace.h"
 #include "CADKernel/Topo/TopologicalVertex.h"
diff --git a/Engine/Source/Runtime/Datasmith/CADKernel/Private/CADKernel/Topo/Topomaker.cpp b/Engine/Source/Runtime/Datasmith/CADKernel/Private/CADKernel/Topo/Topomaker.cpp
new file mode 100644
--- /dev/null
+++ b/Engine/Source/Runtime/Datasmith/CADKernel/Private/CADKernel/Topo/Topomaker.cpp
@@ -0,0 +1,1178 @@
+// Copyright Epic Games, Inc. All Rights Reserved.
+
+#include "CADKernel/Topo/Topomaker.h"
+
+#include "CADKernel/Core/Chrono.h"
+#include "CADKernel/Core/Session.h"
+#include "CADKernel/Geo/Curves/RestrictionCurve.h"
+#include "CADKernel/Topo/Body.h"
+#include "CADKernel/Topo/Model.h"
+#include "CADKernel/Topo/Shell.h"
+#include "CADKernel/Topo/TopologicalEdge.h"
+#include "CADKernel/Topo/TopologicalLink.h"
+#include "CADKernel/Topo/TopologicalFace.h"
+#include "CADKernel/Topo/TopologicalLoop.h"
+#include "CADKernel/Topo/TopologicalVertex.h"
+#include "CADKernel/UI/Display.h"
+#include "CADKernel/UI/Message.h"
+#include "CADKernel/Utils/Util.h"
+
+namespace CADKernel
+{
+
+FJoiner::FJoiner(FSession& InSession, const TSharedRef<FModel>& InModel, double InTolerance)
+	: Session(InSession)
+	, JoiningTolerance(InTolerance* UE_SQRT_2)
+	, JoiningToleranceSquare(FMath::Square(JoiningTolerance))
+{
+	int32 ShellCount = 0;
+	for (const TSharedPtr<FBody>& Body : InModel->GetBodies())
+	{
+		ShellCount += Body->GetShells().Num();
+	}
+	Shells.Reserve(ShellCount);
+
+	for (const TSharedPtr<FBody>& Body : InModel->GetBodies())
+	{
+		Shells.Append(Body->GetShells());
+	}
+
+	InitFaces();
+}
+
+FJoiner::FJoiner(FSession& InSession, const TArray<TSharedPtr<FTopologicalFace>>& InFaces, double InTolerance)
+	: Session(InSession)
+	, Faces(InFaces)
+	, JoiningTolerance(InTolerance * UE_SQRT_2)
+	, JoiningToleranceSquare(FMath::Square(JoiningTolerance))
+{
+}
+
+FJoiner::FJoiner(FSession& InSession, const TArray<TSharedPtr<FShell>>& InShells, double InTolerance)
+	: Session(InSession)
+	, Shells(InShells)
+	, JoiningTolerance(InTolerance)
+	, JoiningToleranceSquare(FMath::Square(JoiningTolerance))
+{
+	InitFaces();
+}
+
+void FJoiner::InitFaces()
+{
+	int32 FaceCount = 0;
+	for (const TSharedPtr<FShell>& Shell : Shells)
+	{
+		Shell->CompleteMetadata();
+		FaceCount += Shell->FaceCount();
+	}
+	Faces.Reserve(FaceCount);
+
+	for (const TSharedPtr<FShell>& Shell : Shells)
+	{
+		Shell->SpreadBodyOrientation();
+		for (const FOrientedFace& Face : Shell->GetFaces())
+		{
+			Faces.Add(Face.Entity);
+		}
+		for (const FOrientedFace& Face : Shell->GetFaces())
+		{
+			Face.Entity->CompleteMetadata();
+		}
+	}
+
+	for (const TSharedPtr<FTopologicalFace>& Face : Faces)
+	{
+		Face->ResetMarker2();
+	}
+}
+
+void FJoiner::RemoveFacesFromShell()
+{
+	// remove faces from their shells
+	TSet<FShell*> ShellSet;
+	for (const TSharedPtr<FTopologicalFace>& Face : Faces)
+	{
+		if (Face->GetHost() != nullptr)
+		{
+			ShellSet.Add(Face->GetHost());
+			Face->ResetHost();
+		}	
+	}
+
+	for (FShell* Shell : ShellSet)
+	{
+		bool bIsOutter = Shell->IsOutter();
+
+		TArray<FOrientedFace> ShellFace;
+		ShellFace.Reserve(Shell->FaceCount());
+		for (const FOrientedFace& Face : Shell->GetFaces())
+		{
+			if (!Face.Entity->GetHost())
+			{
+				if (bIsOutter != (Face.Direction == EOrientation::Front))
+				{
+					Face.Entity->SetBackOriented();
+				}
+			}
+			else
+			{
+				ShellFace.Emplace(Face);
+			}
+		}
+		Shell->ReplaceFaces(ShellFace);
+	}
+}
+
+void FJoiner::EmptyShells()
+{
+	for (const TSharedPtr<FShell>& Shell : Shells)
+	{
+		Shell->Empty();
+	}
+}
+
+void FJoiner::JoinFaces()
+{
+	FTimePoint StartJoinTime = FChrono::Now();
+
+	//RemoveIsolatedEdges();
+
+	TArray<FTopologicalVertex*> BorderVertices;
+	GetBorderVertices(BorderVertices);
+	MergeCoincidentVertices(BorderVertices);
+	//RemoveIsolatedEdges();
+
+	CheckSelfConnectedEdge();
+
+	//RemoveIsolatedEdges();
+
+	//Wait();
+	MergeCoincidentEdges(BorderVertices);
+
+	//Wait();
+	MergeUnconnectedAdjacentEdges();
+
+	//StitchParallelEdges(BorderVertices);
+
+	FDuration InitKernelIODuration = FChrono::Elapse(StartJoinTime);
+	FChrono::PrintClockElapse(EVerboseLevel::Log, TEXT(""), TEXT("Join"), InitKernelIODuration);
+}
+
+void FJoiner::GetVertices(TArray<FTopologicalVertex*>& Vertices)
+{
+	Vertices.Empty(10 * Faces.Num());
+
+	for (TSharedPtr<FTopologicalFace> Face : Faces)
+	{
+		for (TSharedPtr<FTopologicalLoop> Loop : Face->GetLoops())
+		{
+			for (const FOrientedEdge& OrientedEdge : Loop->GetEdges())
+			{
+				const TSharedPtr<FTopologicalEdge>& Edge = OrientedEdge.Entity;
+				if (!Edge->GetStartVertex()->GetLinkActiveEntity()->HasMarker1())
+				{
+					Edge->GetStartVertex()->GetLinkActiveEntity()->SetMarker1();
+					Vertices.Add(&*Edge->GetStartVertex()->GetLinkActiveEntity());
+				}
+				if (!Edge->GetEndVertex()->GetLinkActiveEntity()->HasMarker1())
+				{
+					Edge->GetEndVertex()->GetLinkActiveEntity()->SetMarker1();
+					Vertices.Add(&*Edge->GetEndVertex()->GetLinkActiveEntity());
+				}
+			}
+		}
+	}
+	for (FTopologicalVertex* Vertex : Vertices)
+	{
+		Vertex->ResetMarker1();
+	}
+}
+
+void FJoiner::GetBorderVertices(TArray<FTopologicalVertex*>& BorderVertices)
+{
+	TArray<FTopologicalVertex*> Vertices;
+	GetVertices(Vertices);
+
+	BorderVertices.Empty(Vertices.Num());
+
+	for (FTopologicalVertex* Vertex : Vertices)
+	{
+		if (Vertex->IsBorderVertex())
+		{
+			BorderVertices.Add(Vertex);
+		}
+	}
+}
+
+void FJoiner::MergeCoincidentVertices(TArray<FTopologicalVertex*>& VerticesToMerge)
+{
+	FTimePoint StartTime = FChrono::Now();
+
+	const double JoiningVerticesToleranceSquare = 2 * JoiningToleranceSquare;
+	const double WeigthTolerance = 3 * JoiningTolerance;
+
+	int32 VertexNum = (int32)VerticesToMerge.Num();
+
+	TArray<double> VerticesWeight;
+	VerticesWeight.Reserve(VertexNum);
+
+	TArray<int32> SortedVertexIndices;
+	SortedVertexIndices.Reserve(VertexNum);
+
+	for (FTopologicalVertex* Vertex : VerticesToMerge)
+	{
+		VerticesWeight.Add(Vertex->GetCoordinates().X + Vertex->GetCoordinates().Y + Vertex->GetCoordinates().Z);
+	}
+
+	for (int32 Index = 0; Index < VerticesToMerge.Num(); ++Index)
+	{
+		SortedVertexIndices.Add(Index);
+	}
+	SortedVertexIndices.Sort([&VerticesWeight](const int32& Index1, const int32& Index2) { return VerticesWeight[Index1] < VerticesWeight[Index2]; });
+
+	for (int32 IndexI = 0; IndexI < VertexNum; ++IndexI)
+	{
+		FTopologicalVertex* Vertex = VerticesToMerge[SortedVertexIndices[IndexI]];
+		if (Vertex->HasMarker1())
+		{
+			continue;
+		}
+
+		ensureCADKernel(Vertex->IsActiveEntity());
+
+		Vertex->SetMarker1();
+
+		double VertexWeigth = VerticesWeight[SortedVertexIndices[IndexI]];
+		FPoint Barycenter = Vertex->GetBarycenter();
+
+		for (int32 IndexJ = IndexI + 1; IndexJ < VertexNum; ++IndexJ)
+		{
+			FTopologicalVertex* OtherVertex = VerticesToMerge[SortedVertexIndices[IndexJ]];
+			if (OtherVertex->HasMarker1())
+			{
+				continue;
+			}
+
+			double OtherVertexWeigth = VerticesWeight[SortedVertexIndices[IndexJ]];
+			if ((OtherVertexWeigth - VertexWeigth) > WeigthTolerance)
+			{
+				break;
+			}
+
+			double DistanceSqr = OtherVertex->GetLinkActiveEntity()->SquareDistance(Barycenter);
+			if (DistanceSqr < JoiningVerticesToleranceSquare)
+			{
+				OtherVertex->SetMarker1();
+				Vertex->Link(*OtherVertex);
+				Barycenter = Vertex->GetBarycenter();
+			}
+		}
+	}
+
+	for (FTopologicalVertex* Vertex : VerticesToMerge)
+	{
+		Vertex->ResetMarker1();
+	}
+
+	TArray<FTopologicalVertex*> ActiveVertices;
+	ActiveVertices.Reserve(VertexNum);
+
+	for (FTopologicalVertex* Vertex : VerticesToMerge)
+	{
+		FTopologicalVertex& ActiveVertex = *Vertex->GetLinkActiveEntity();
+		if (ActiveVertex.HasMarker1())
+		{
+			continue;
+		}
+		ActiveVertex.SetMarker1();
+		ActiveVertices.Add(&ActiveVertex);
+	}
+
+	for (FTopologicalVertex* Vertex : ActiveVertices)
+	{
+		Vertex->ResetMarker1();
+	}
+
+	Swap(ActiveVertices, VerticesToMerge);
+
+	FDuration Duration = FChrono::Elapse(StartTime);
+	FChrono::PrintClockElapse(EVerboseLevel::Log, TEXT("    "), TEXT("Merge Coincident vertices"), Duration);
+
+}
+
+void FJoiner::MergeBorderVerticesWithCoincidentOtherVertices(TArray<FTopologicalVertex*>& Vertices)
+{
+	int32 VertexNum = (int32)Vertices.Num();
+
+	TArray<double> VerticesWeight;
+	VerticesWeight.Reserve(VertexNum);
+	TArray<int32> SortedVertexIndices;
+	SortedVertexIndices.Reserve(VertexNum);
+
+	for (FTopologicalVertex* Vertex : Vertices)
+	{
+		VerticesWeight.Add(Vertex->GetCoordinates().X + Vertex->GetCoordinates().Y + Vertex->GetCoordinates().Z);
+	}
+
+	for (int32 Index = 0; Index < Vertices.Num(); ++Index)
+	{
+		SortedVertexIndices.Add(Index);
+	}
+	SortedVertexIndices.Sort([&VerticesWeight](const int32& Index1, const int32& Index2) { return VerticesWeight[Index1] < VerticesWeight[Index2]; });
+
+	int32 StartIndexJ = 1;
+	for (int32 IndexI = 0; IndexI < VertexNum; ++IndexI)
+	{
+		FTopologicalVertex* Vertex = Vertices[SortedVertexIndices[IndexI]];
+		if (!Vertex->IsBorderVertex())
+		{
+			continue;
+		}
+
+		double VertexWeigth = VerticesWeight[SortedVertexIndices[IndexI]];
+		FPoint Barycenter = Vertex->GetBarycenter();
+		for (int32 IndexJ = StartIndexJ; IndexJ < VertexNum; ++IndexJ)
+		{
+			FTopologicalVertex& OtherVertex = *Vertices[SortedVertexIndices[IndexJ]];
+			if (Vertex->GetLink() == OtherVertex.GetLink())
+			{
+				continue;
+			}
+
+			double OtherVertexWeigth = VerticesWeight[SortedVertexIndices[IndexJ]];
+
+			if (OtherVertexWeigth + JoiningTolerance - VertexWeigth < 0)
+			{
+				StartIndexJ = IndexJ;
+				continue;
+			}
+
+			if (OtherVertexWeigth - JoiningTolerance - VertexWeigth > 0)
+			{
+				break;
+			}
+
+			double DistanceSqr = OtherVertex.GetLinkActiveEntity()->SquareDistance(Barycenter);
+			if (DistanceSqr < JoiningToleranceSquare)
+			{
+				TArray<FTopologicalEdge*> CommonEdges;
+				Vertex->GetConnectedEdges(OtherVertex, CommonEdges);
+				if (CommonEdges.Num() > 0)
+				{
+					continue;
+				}
+
+				OtherVertex.SetMarker1();
+				Vertex->Link(OtherVertex);
+				Barycenter = Vertex->GetBarycenter();
+			}
+		}
+	}
+
+	TArray<FTopologicalVertex*> NewVertices;
+	NewVertices.Empty(VertexNum);
+
+	for (FTopologicalVertex* Vertex : Vertices)
+	{
+		FTopologicalVertex& ActiveVertex = *Vertex->GetLinkActiveEntity();
+		if (ActiveVertex.HasMarker1())
+		{
+			continue;
+		}
+		ActiveVertex.SetMarker1();
+		NewVertices.Add(&ActiveVertex);
+	}
+
+	for (FTopologicalVertex* Vertex : NewVertices)
+	{
+		Vertex->ResetMarker1();
+	}
+
+	Swap(NewVertices, Vertices);
+}
+
+void FJoiner::MergeCoincidentEdges(TArray<FTopologicalVertex*>& VerticesToProcess)
+{
+	FTimePoint StartTime = FChrono::Now();
+
+	for (FTopologicalVertex* Vertex : VerticesToProcess)
+	{
+		TArray<FTopologicalEdge*> ConnectedEdges;
+		Vertex->GetConnectedEdges(ConnectedEdges);
+		int32 ConnectedEdgeCount = ConnectedEdges.Num();
+		if (ConnectedEdgeCount == 1)
+		{
+			continue;
+		}
+
+		for (int32 EdgeI = 0; EdgeI < ConnectedEdgeCount - 1; ++EdgeI)
+		{
+			FTopologicalEdge* Edge = ConnectedEdges[EdgeI];
+			if (!Edge->IsActiveEntity())
+			{
+				continue;
+			}
+			bool bFirstEdgeBorder = Edge->IsBorder();
+			FTopologicalVertex& EndVertex = *Edge->GetOtherVertex(*Vertex)->GetLinkActiveEntity();
+
+			for (int32 EdgeJ = EdgeI + 1; EdgeJ < ConnectedEdgeCount; ++EdgeJ)
+			{
+				FTopologicalEdge* SecondEdge = ConnectedEdges[EdgeJ];
+				if (!SecondEdge->IsActiveEntity())
+				{
+					continue;
+				}
+				bool bSecondEdgeBorder = Edge->IsBorder();
+
+				// Process only if at least one edge is Border
+				if (!bFirstEdgeBorder && !bSecondEdgeBorder)
+				{
+					continue;
+				}
+
+				FTopologicalVertex& OtherEdgeEndVertex = *SecondEdge->GetOtherVertex(*Vertex)->GetLinkActiveEntity();
+
+				if (&OtherEdgeEndVertex == &EndVertex)
+				{
+					FPoint StartTangentEdge = Edge->GetTangentAt(*Vertex);
+					FPoint StartTangentOtherEdge = SecondEdge->GetTangentAt(*Vertex);
+
+					double CosAngle = StartTangentEdge.ComputeCosinus(StartTangentOtherEdge);
+					if (CosAngle < 0.9)
+					{
+						continue;
+					}
+
+					if (Edge->GetFace() != SecondEdge->GetFace())
+					{
+						Edge->Link(*SecondEdge, JoiningTolerance);
+					}
+				}
+			}
+		}
+	}
+
+	FDuration Duration = FChrono::Elapse(StartTime);
+
+	FChrono::PrintClockElapse(EVerboseLevel::Log, TEXT("    "), TEXT("Merge coincident edges"), Duration);
+
+}
+
+TSharedPtr<FTopologicalVertex> FJoiner::SplitAndLink(FTopologicalVertex& StartVertex, FTopologicalEdge& EdgeToLink, FTopologicalEdge& EdgeToSplit)
+{
+	FTopologicalVertex* VertexToLink = EdgeToLink.GetOtherVertex(StartVertex);
+
+	FPoint ProjectedPoint;
+	double UProjectedPoint = EdgeToSplit.ProjectPoint(VertexToLink->GetBarycenter(), ProjectedPoint);
+
+	double SquareDistanceToProjectedPoint = ProjectedPoint.SquareDistance(VertexToLink->GetBarycenter());
+	if (SquareDistanceToProjectedPoint > JoiningToleranceSquare)
+	{
+		return TSharedPtr<FTopologicalVertex>();
+	}
+
+	// Check if the ProjectedPoint is not nearly equal edge boundary
+	FTopologicalVertex* EndVertex = EdgeToSplit.GetOtherVertex(StartVertex);
+	if (EndVertex->SquareDistance(ProjectedPoint) < JoiningToleranceSquare)
+	{
+		VertexToLink->Link(*EndVertex);
+		EdgeToLink.Link(EdgeToSplit, JoiningTolerance);
+		// TSharedPtr<FTopologicalVertex>() is returned as EndVertex is not new
+		return TSharedPtr<FTopologicalVertex>();
+	}
+
+	// JoinParallelEdges process all edges connected to startVertex (ConnectedEdges).
+	// Connected edges must remain compliant i.e. all edges of ConnectedEdges must be connected to StartVertex
+	// EdgeToSplit->SplitAt() must keep EdgeToSplit connected to StartVertex
+	bool bKeepStartVertexConnectivity = (StartVertex.GetLink() == EdgeToSplit.GetStartVertex()->GetLink());
+
+	TSharedPtr<FTopologicalEdge> NewEdge;
+	TSharedPtr<FTopologicalVertex> NewVertex = EdgeToSplit.SplitAt(UProjectedPoint, ProjectedPoint, bKeepStartVertexConnectivity, NewEdge);
+	if (!NewVertex.IsValid())
+	{
+		return TSharedPtr<FTopologicalVertex>();
+	}
+
+	VertexToLink->Link(*NewVertex);
+	EdgeToLink.Link(EdgeToSplit, JoiningTolerance);
+
+	return NewVertex;
+}
+
+void FJoiner::StitchParallelEdges(TArray<FTopologicalVertex*>& VerticesToProcess)
+{
+	FTimePoint StartTime = FChrono::Now();
+
+	for (int32 VertexI = 0; VertexI < VerticesToProcess.Num(); ++VertexI)
+	{
+		FTopologicalVertex* Vertex = VerticesToProcess[VertexI];
+
+		if (!Vertex->IsBorderVertex())
+		{
+			continue;
+		}
+
+		TArray<FTopologicalEdge*> ConnectedEdges;
+		Vertex->GetConnectedEdges(ConnectedEdges);
+		int32 ConnectedEdgeCount = ConnectedEdges.Num();
+		if (ConnectedEdgeCount == 1)
+		{
+			continue;
+		}
+
+		for (int32 EdgeI = 0; EdgeI < ConnectedEdgeCount - 1; ++EdgeI)
+		{
+			FTopologicalEdge* Edge = ConnectedEdges[EdgeI];
+			ensureCADKernel(Edge->GetLoop() != nullptr);
+
+			if (Edge->IsDegenerated())
+			{
+				continue;
+			}
+
+			if (!Edge->IsActiveEntity())
+			{
+				continue;
+			}
+			bool bFirstEdgeBorder = Edge->IsBorder();
+
+			for (int32 EdgeJ = EdgeI + 1; EdgeJ < ConnectedEdgeCount; ++EdgeJ)
+			{
+				FTopologicalEdge* SecondEdge = ConnectedEdges[EdgeJ];
+				if (Edge->IsDegenerated())
+				{
+					continue;
+				}
+				if (!SecondEdge->IsActiveEntity())
+				{
+					continue;
+				}
+				bool bSecondEdgeBorder = Edge->IsBorder();
+
+				// Process only if at least one edge is Border
+				if (!bFirstEdgeBorder && !bSecondEdgeBorder)
+				{
+					continue;
+				}
+
+				FPoint StartTangentEdge = Edge->GetTangentAt(*Vertex);
+				FPoint StartTangentOtherEdge = SecondEdge->GetTangentAt(*Vertex);
+
+				double CosAngle = StartTangentEdge.ComputeCosinus(StartTangentOtherEdge);
+				if (CosAngle < 0.9)
+				{
+					continue;
+				}
+
+				FTopologicalVertex& EndVertex = *Edge->GetOtherVertex(*Vertex)->GetLinkActiveEntity();
+				FTopologicalVertex& OtherEdgeEndVertex = *SecondEdge->GetOtherVertex(*Vertex)->GetLinkActiveEntity();
+				if (&EndVertex == &OtherEdgeEndVertex)
+				{
+					Edge->Link(*SecondEdge, JoiningTolerance);
+				}
+				else
+				{
+					double EdgeLength = Edge->Length();
+					double OtherEdgeLength = SecondEdge->Length();
+
+					TSharedPtr<FTopologicalVertex> NewVertex;
+					if (EdgeLength < OtherEdgeLength)
+					{
+						NewVertex = SplitAndLink(*Vertex, *Edge, *SecondEdge);
+					}
+					else
+					{
+						NewVertex = SplitAndLink(*Vertex, *SecondEdge, *Edge);
+					}
+
+					if (NewVertex.IsValid())
+					{
+						VerticesToProcess.Add(NewVertex.Get());
+					}
+				}
+			}
+		}
+	}
+
+	FDuration Duration = FChrono::Elapse(StartTime);
+	FChrono::PrintClockElapse(EVerboseLevel::Log, TEXT("    "), TEXT("Stitch Parallel Edges"), Duration);
+}
+
+void FJoiner::MergeUnconnectedAdjacentEdges()
+{
+	FTimePoint StartTime = FChrono::Now();
+
+	for (TSharedPtr<FTopologicalFace> FacePtr : Faces)
+	{
+		if (!FacePtr.IsValid())
+		{
+			continue;
+		}
+
+		FTopologicalFace* Face = FacePtr.Get();
+
+		TArray<TArray<FOrientedEdge>> ArrayOfCandidates;
+		ArrayOfCandidates.Reserve(10);
+
+		// For each loop, find unconnected successive edges...
+		for (TSharedPtr<FTopologicalLoop> Loop : Face->GetLoops())
+		{
+			TArray<FOrientedEdge>& Edges = Loop->GetEdges();
+			int32 EdgeCount = Edges.Num();
+
+
+			// Find the starting edge i.e. the next edge of the first edge that its ending vertex is connecting to 3 or more edges 
+			// The algorithm start to the last edge of the loop, if it verifies the criteria then the first edge is the edges[0]
+			int32 EndIndex = EdgeCount;
+			{
+				TSharedPtr<FTopologicalVertex> EndVertex;
+				do
+				{
+					EndIndex--;
+					EndVertex = Edges[EndIndex].Direction == EOrientation::Front ? Edges[EndIndex].Entity->GetEndVertex() : Edges[EndIndex].Entity->GetStartVertex();
+				} while (EndVertex->ConnectedEdgeCount() == 2 && EndIndex > 0);
+			}
+			EndIndex++;
+
+			// First step
+			// For the loop, find all arrays of successive unconnected edges
+			TArray<FOrientedEdge>* Candidates = &ArrayOfCandidates.Emplace_GetRef();
+			Candidates->Reserve(10);
+			bool bCanStop = false;
+			for (int32 Index = EndIndex; bCanStop == (Index != EndIndex); ++Index)
+			{
+				if (Index == EdgeCount)
+				{
+					Index = 0;
+				}
+				bCanStop = true;
+
+				FOrientedEdge& Edge = Edges[Index];
+				if (Edge.Entity->GetTwinsEntityCount() == 1)
+				{
+					TSharedPtr<FTopologicalVertex> EndVertex = Edge.Direction == EOrientation::Front ? Edge.Entity->GetEndVertex() : Edge.Entity->GetStartVertex();
+
+					TArray<FTopologicalEdge*> ConnectedEdges;
+					EndVertex->GetConnectedEdges(ConnectedEdges);
+
+					bool bEdgeIsNotTheLast = false;
+					if (ConnectedEdges.Num() == 2)
+					{
+						// check if the edges are tangents
+						FPoint StartTangentEdge = ConnectedEdges[0]->GetTangentAt(*EndVertex);
+						FPoint StartTangentOtherEdge = ConnectedEdges[1]->GetTangentAt(*EndVertex);
+
+						double CosAngle = StartTangentEdge.ComputeCosinus(StartTangentOtherEdge);
+						if (CosAngle < -0.9)
+						{
+							bEdgeIsNotTheLast = true;
+						}
+					}
+
+					if (bEdgeIsNotTheLast || Candidates->Num() > 0)
+					{
+						Candidates->Add(Edges[Index]);
+					}
+
+					if (!bEdgeIsNotTheLast && Candidates->Num() > 0)
+					{
+						Candidates = &ArrayOfCandidates.Emplace_GetRef();
+						Candidates->Reserve(10);
+					}
+				}
+			}
+		}
+
+		// Second step, 
+		// Each array of edges are merged to generated a single edge, the loop is updated
+		for (TArray<FOrientedEdge>& Candidates : ArrayOfCandidates)
+		{
+			if (Candidates.Num())
+			{
+				TSharedPtr<FTopologicalVertex> StartVertexPtr = Candidates[0].Direction == EOrientation::Front ? Candidates[0].Entity->GetStartVertex() : Candidates[0].Entity->GetEndVertex();
+				TSharedRef<FTopologicalVertex> StartVertex = StartVertexPtr.ToSharedRef();
+				TSharedPtr<FTopologicalVertex> EndVertex = Candidates.Last().Direction == EOrientation::Front ? Candidates.Last().Entity->GetEndVertex() : Candidates.Last().Entity->GetStartVertex();
+
+				FPoint StartTangentEdge = Candidates[0].Entity->GetTangentAt(*StartVertex);
+
+				TArray<FTopologicalEdge*> Edges;
+				StartVertex->GetConnectedEdges(*EndVertex, Edges);
+
+				// Remove edge that is not parallel edge
+				FTopologicalEdge* ParallelEdge = nullptr;
+				for (FTopologicalEdge*& Edge : Edges)
+				{
+					if (!Edge)
+					{
+						continue;
+					}
+
+					if (Edge->GetFace() == Face)
+					{
+						Edge = nullptr;
+						continue;
+					}
+					
+					FPoint StartTangentOtherEdge = Edge->GetTangentAt(*StartVertex);
+					double CosAngle = StartTangentEdge.ComputeCosinus(StartTangentOtherEdge);
+					if (CosAngle > 0.9)
+					{
+						ParallelEdge = Edge;
+						break;
+					}
+				}
+
+				if (ParallelEdge)
+				{
+					//FMessage::Printf(Log, TEXT("Face %d UnconnectedAdjacentEdges\n"), Face->GetId());
+					//for (auto Edge : Candidates)
+					//{
+					//	FMessage::Printf(Log, TEXT("       - %d \n"), Edge.Entity->GetId());
+					//}
+
+					TSharedPtr<FTopologicalEdge> NewEdge = FTopologicalEdge::CreateEdgeByMergingEdges(Candidates, StartVertex, EndVertex.ToSharedRef());
+					if (!NewEdge.IsValid())
+					{
+						// the edges cannot be merge, they will be connected to the parallel edge with "SplitAndLink"
+						// i.e. instead of merging edges to be one and link to the parallel one
+						// the parallel one is split at the extremities of each edges, and each new edge is linked to its parallel edge 
+						break;
+					}
+
+					// Link to the parallel edge
+					// New edge is link to the first parallel one as the other parallel should already be linked together
+					for (FTopologicalEdge* Edge : Edges)
+					{
+						if (Edge)
+						{
+							Edge->Link(*NewEdge, JoiningTolerance);
+							break;
+						}
+					}
+				}
+			}
+		}
+	}
+
+	FDuration Duration = FChrono::Elapse(StartTime);
+
+	FChrono::PrintClockElapse(EVerboseLevel::Log, TEXT("    "), TEXT("Merge unconnected adjacent edges"), Duration);
+
+}
+
+void FJoiner::RemoveIsolatedEdges()
+{
+	FTimePoint StartTime = FChrono::Now();
+
+	TArray<FTopologicalEdge*> IsolatedEdges;
+
+	TArray<FTopologicalVertex*> Vertices;
+	GetVertices(Vertices);
+
+	for (const FTopologicalVertex* Vertex : Vertices)
+	{
+		for (const FTopologicalVertex* TwinVertex : Vertex->GetTwinsEntities())
+		{
+			TArray<FTopologicalEdge*> Edges = Vertex->GetDirectConnectedEdges();
+			for (FTopologicalEdge* Edge : Edges)
+			{
+				if (Edge->GetLoop() == nullptr)
+				{
+					IsolatedEdges.Add(Edge);
+				}
+			}
+		}
+	}
+	FDuration Duration = FChrono::Elapse(StartTime);
+
+	FChrono::PrintClockElapse(EVerboseLevel::Log, TEXT("    "), TEXT("Remove Isolated Edges"), Duration);
+
+	FMessage::Printf(EVerboseLevel::Log, TEXT("\n\nIsolatedEdges count %d\n\n\n"), IsolatedEdges.Num());
+}
+
+void FJoiner::UnlinkNonManifoldVertex()
+{
+	TArray<FTopologicalVertex*> Vertices;
+	GetVertices(Vertices);
+
+	TMap<FShell*, int32> ShellToVertexCount;
+	for (FTopologicalVertex* Vertex : Vertices)
+	{
+		ShellToVertexCount.Empty(Vertex->GetTwinsEntityCount());
+		for (const FTopologicalVertex* TwinVertex : Vertex->GetTwinsEntities())
+		{
+			ensureCADKernel(!TwinVertex->GetDirectConnectedEdges().IsEmpty());
+			FTopologicalEdge* Edge = TwinVertex->GetDirectConnectedEdges()[0];
+			FShell* Shell = Edge->GetLoop()->GetFace()->GetHost();
+			if (Shell != nullptr)
+			{
+				ShellToVertexCount.FindOrAdd(Shell)++;
+			}
+		}
+
+		if (ShellToVertexCount.Num() > 1)
+		{
+			TMap<FShell*, TArray<FTopologicalVertex*>> ShellToVertices;
+			ShellToVertices.Reserve(ShellToVertexCount.Num());
+			for (TPair<FShell*, int32> Pair : ShellToVertexCount)
+			{
+				ShellToVertices.FindOrAdd(Pair.Key).Reserve(Pair.Value);
+			}
+
+			for (FTopologicalVertex* TwinVertex : Vertex->GetTwinsEntities())
+			{
+				ensureCADKernel(!TwinVertex->GetDirectConnectedEdges().IsEmpty());
+				FTopologicalEdge* Edge = TwinVertex->GetDirectConnectedEdges()[0];
+				FShell* Shell = Edge->GetLoop()->GetFace()->GetHost();
+				if (Shell != nullptr)
+				{
+					ShellToVertices.FindOrAdd(Shell).Add(TwinVertex);
+				}
+			}
+
+			Vertex->UnlinkTwinEntities();
+			for (TPair<FShell*, TArray<FTopologicalVertex*>> ShellVertices : ShellToVertices)
+			{
+				if (ShellVertices.Value.Num() > 1)
+				{
+					TArray<FTopologicalVertex*>& Twins = ShellVertices.Value;
+					FTopologicalVertex* FirstVertex = Twins[0];
+					for (FTopologicalVertex* TwinVertex : Twins)
+					{
+						FirstVertex->Link(*TwinVertex);
+					}
+				}
+			}
+		}
+	}
+}
+
+// =========================================================================================================================================================================================================
+// =========================================================================================================================================================================================================
+// =========================================================================================================================================================================================================
+//
+//
+//                                                                            NOT YET REVIEWED
+//
+//
+// =========================================================================================================================================================================================================
+// =========================================================================================================================================================================================================
+// =========================================================================================================================================================================================================
+
+void FJoiner::CheckSelfConnectedEdge()
+{
+	FTimePoint StartTime = FChrono::Now();
+
+	FMessage::Printf(Log, TEXT("    Self connected edges\n"));
+	for (TSharedPtr<FTopologicalFace> Face : Faces)
+	{
+		for (TSharedPtr<FTopologicalLoop> Loop : Face->GetLoops())
+		{
+			for (const FOrientedEdge& OrientedEdge : Loop->GetEdges())
+			{
+				TSharedPtr<FTopologicalEdge> Edge = OrientedEdge.Entity;
+				if (Edge->GetStartVertex()->IsLinkedTo(Edge->GetEndVertex()))
+				{
+					if (!Edge->IsDegenerated() && Edge->Length() < 2 * JoiningTolerance)
+					{
+						FMessage::Printf(Debug, TEXT("Face %d Edge %d was self connected\n"), Face->GetId(), Edge->GetId());
+						Edge->GetStartVertex()->UnlinkTo(*Edge->GetEndVertex());
+					}
+				}
+			}
+		}
+	}
+	FDuration Duration = FChrono::Elapse(StartTime);
+
+	FChrono::PrintClockElapse(EVerboseLevel::Log, TEXT("    "), TEXT("Unconnect Self connected edges"), Duration);
+}
+
+void FJoiner::SplitIntoConnectedShell()
+{
+	// Processed1 : Surfaces added in CandidateSurfacesForMesh
+
+	int32 TopologicalFaceCount = Faces.Num();
+	// Is closed ?
+	// Is one shell ?
+
+	TArray<FFaceSubset> SubShells;
+
+	int32 ProcessFaceCount = 0;
+
+	TArray<FTopologicalFace*> Front;
+	TFunction<void(const FTopologicalFace&, FFaceSubset&)> GetNeighboringFaces = [&](const FTopologicalFace& Face, FFaceSubset& Shell)
+	{
+		for (const TSharedPtr<FTopologicalLoop>& Loop : Face.GetLoops())
+		{
+			for (const FOrientedEdge& OrientedEdge : Loop->GetEdges())
+			{
+				const TSharedPtr<FTopologicalEdge>& Edge = OrientedEdge.Entity;
+				if (Edge->HasMarker1())
+				{
+					continue;
+				}
+				Edge->SetMarker1();
+
+				if (Edge->GetTwinsEntityCount() == 1)
+				{
+					if (!Edge->IsDegenerated())
+					{
+						Shell.BorderEdgeCount++;
+					}
+					continue;
+				}
+
+				if (Edge->GetTwinsEntityCount() > 2)
+				{
+					Shell.NonManifoldEdgeCount++;
+				}
+
+				for (FTopologicalEdge* NextEdge : Edge->GetTwinsEntities())
+				{
+					if (NextEdge->HasMarker1())
+					{
+						continue;
+					}
+					NextEdge->SetMarker1();
+
+					FTopologicalFace* NextFace = NextEdge->GetFace();
+					if ((NextFace == nullptr) || NextFace->HasMarker1())
+					{
+						continue;
+					}
+
+					NextFace->SetMarker1();
+					Front.Add(NextFace);
+				}
+			}
+		}
+	};
+
+	TFunction<void(FFaceSubset&)> SpreadFront = [&](FFaceSubset& Shell)
+	{
+		while (Front.Num())
+		{
+			FTopologicalFace* Face = Front.Pop();
+			if (Face == nullptr)
+			{
+				continue;
+			}
+
+			Shell.Faces.Add(Face);
+			GetNeighboringFaces(*Face, Shell);
+		}
+	};
+
+	for (const TSharedPtr<FTopologicalFace>& Face : Faces)
+	{
+		if (Face->HasMarker1())
+		{
+			continue;
+		}
+
+		FFaceSubset& Shell = SubShells.Emplace_GetRef();
+		Shell.Faces.Reserve(TopologicalFaceCount - ProcessFaceCount);
+		Front.Empty(TopologicalFaceCount);
+
+		Face->SetMarker1();
+		Front.Add(Face.Get());
+		SpreadFront(Shell);
+		ProcessFaceCount += Shell.Faces.Num();
+
+		if (ProcessFaceCount == TopologicalFaceCount)
+		{
+			break;
+		}
+	}
+
+	// reset Marker
+	for (const TSharedPtr<FTopologicalFace>& Face : Faces)
+	{
+		Face->ResetMarkers();
+		for (const TSharedPtr<FTopologicalLoop>& Loop : Face->GetLoops())
+		{
+			for (const FOrientedEdge& OrientedEdge : Loop->GetEdges())
+			{
+				OrientedEdge.Entity->ResetMarkers();
+			}
+		}
+	}
+
+	// for each FaceSubset, find the main shell
+	for (FFaceSubset& FaceSubset : SubShells)
+	{
+		TMap<FBody*, int32> BodyToFaceCount;
+		TMap<FShell*, int32> ShellToFaceCount;
+		TMap<uint32, int32> ColorToFaceCount;
+		TMap<FString, int32> NameToFaceCount;
+
+		for (FTopologicalFace* Face : FaceSubset.Faces)
+		{
+			FShell* Shell = Face->GetHost();
+			FBody* Body = Shell->GetHost();
+
+			ShellToFaceCount.FindOrAdd(Shell)++;
+			BodyToFaceCount.FindOrAdd(Body)++;
+			ColorToFaceCount.FindOrAdd(Face->GetColorId())++;
+			NameToFaceCount.FindOrAdd(Face->GetName())++;
+		}
+
+		FaceSubset.SetMainShell(ShellToFaceCount);
+		FaceSubset.SetMainBody(BodyToFaceCount);
+		FaceSubset.SetMainName(NameToFaceCount);
+		FaceSubset.SetMainColor(ColorToFaceCount);
+	}
+
+	if (Shells.Num())
+	{
+		EmptyShells();
+	}
+	else
+	{
+		RemoveFacesFromShell();
+	}
+
+	// for each FaceSubset, process the Shell
+	for (FFaceSubset FaceSubset : SubShells)
+	{
+		if (FaceSubset.MainShell != nullptr)
+		{
+			FShell* Shell = FaceSubset.MainShell;
+			Shell->Empty(FaceSubset.Faces.Num());
+			Shell->Add(FaceSubset.Faces);
+		}
+		else
+		{
+			FBody* Body = FaceSubset.MainBody;
+			if (Body == nullptr)
+			{
+				TSharedRef<FBody> SharedBody = FEntity::MakeShared<FBody>();
+				Session.GetModel()->Add(SharedBody);
+				Body = &SharedBody.Get();
+				Body->SetName(FaceSubset.MainName);
+				Body->SetColorId(FaceSubset.MainColor);
+			}
+
+			TSharedRef<FShell> Shell = FEntity::MakeShared<FShell>();
+			Body->AddShell(Shell);
+
+			Shell->Add(FaceSubset.Faces);
+			Shell->SetName(FaceSubset.MainName);
+			Shell->SetColorId(FaceSubset.MainColor);
+		}
+	}
+
+	UnlinkNonManifoldVertex();
+
+	Session.GetModel()->RemoveEmptyBodies();
+}
+
+
+
+//void FJoiner::FixCollapsedEdges()
+//{
+//	TArray<TSharedPtr<FTopologicalEdge>> CollapsedEdges;
+//
+//	int32 NbDegeneratedEdges = 0;
+//	for (const TSharedPtr<FTopologicalFace>& Face : Faces)
+//	{
+//		for (const TSharedPtr<FTopologicalLoop>& Loop : Face->GetLoops())
+//		{
+//			for (const FOrientedEdge& OrientedEdge : Loop->GetEdges())
+//			{
+//				const TSharedPtr<FTopologicalEdge>& Edge = OrientedEdge.Entity;
+//				for (const TWeakPtr<FTopologicalEdge>& TwinEdge : Edge->GetTwinsEntities())
+//				{
+//					TSharedPtr<FTopologicalFace> Face = TwinEdge.Pin()->GetFace();
+//					if (!Face.IsValid())
+//					{
+//						continue;
+//					}
+//				}
+//
+//				if (Edge->GetStartVertex()->GetLink() == Edge->GetEndVertex()->GetLink())
+//				{
+//					if (!Edge->IsDegenerated())
+//					{
+//						for (const TWeakPtr<FTopologicalEdge>& TwinEdge : Edge->GetTwinsEntities())
+//						{
+//							TwinEdge.Pin()->SetAsDegenerated();
+//						}
+//					}
+//				}
+//			}
+//		}
+//	}
+//
+//	for (TSharedPtr<FTopologicalEdge>& Edge : CollapsedEdges)
+//	{
+//		TSharedPtr<FTopologicalLoop> Loop = Edge->GetLoop();
+//		Loop->RemoveEdge(Edge);
+//
+//		TSharedPtr<FTopologicalVertex> StartVertex = Edge->GetStartVertex();
+//		TSharedPtr<FTopologicalVertex> EndVertex = Edge->GetEndVertex();
+//
+//		Edge->Delete();
+//
+//		StartVertex->DeleteIfIsolated();
+//		EndVertex->DeleteIfIsolated();
+//
+//		if (Loop->GetEdges().Num() == 0)
+//		{
+//			TSharedPtr<FTopologicalFace> Face = Loop->GetFace();
+//			Face->RemoveLoop(Loop);
+//		}
+//	}
+//}
+
+
+
+
+
+
+/*
+NewEdgeVertices.Reserve(AllVertices.Num());
+for (TSharedPtr<FTopologicalVertex> Vertex : AllVertices)
+{
+	TArray<TWeakPtr<FTopologicalEdge>> Edges = Vertex->GetConnectedEdges();
+	if (Edges.Num() == 2)
+	{
+		// TODO Work in 2D
+		FPoint Edge0Tangent;
+		Edges[0].Pin()->GetTangentAt(Vertex, Edge0Tangent);
+
+		FPoint Edge1Tangent;
+		Edges[1].Pin()->GetTangentAt(Vertex, Edge1Tangent);
+
+		double CosAngle = Edge0Tangent.ComputeCosinus(Edge1Tangent);
+		if (CosAngle < -0.75)
+		{
+			TArray<EOrientation> EdgeDirections;
+			EdgeDirections.Reserve(2);
+			EdgeDirections.Add(Edges[0].Pin()->GetEndVertex() == Vertex ? EOrientation::Front : EOrientation::Back);
+			EdgeDirections.Add(Edges[1].Pin()->GetStartVertex() == Vertex ? EOrientation::Front : EOrientation::Back);
+			TSharedPtr<FTopologicalEdge> ExtendedEdge = nullptr;
+
+			// TODO Check better if the both curve are similar (2DCurve as to be similar)
+			if (Edges[0].Pin()->GetCurve() != Edges[1].Pin()->GetCurve())
+			{
+				//ExtendedEdge = FEdge::ReplaceEdgesByOne(Edges, EdgeDirections, GetTolerance());
+			}
+			else
+			{
+				//ExtendedEdge = FEdge::CreateEdgeToMerge2Edges(Edges, EdgeDirections);
+			}
+
+			if (ExtendedEdge)
+			{
+				NewEdgeVertices.Add(ExtendedEdge->GetStartVertex());
+				NewEdgeVertices.Add(ExtendedEdge->GetEndVertex());
+			}
+			continue;
+		}
+	}
+}
+*/
+
+
+}
diff --git a/Engine/Source/Runtime/Datasmith/CADKernel/Public/CADKernel/Topo/Joiner.h b/Engine/Source/Runtime/Datasmith/CADKernel/Public/CADKernel/Topo/Joiner.h
deleted file mode 100644
diff --git a/Engine/Source/Runtime/Datasmith/CADKernel/Public/CADKernel/Topo/Topomaker.h b/Engine/Source/Runtime/Datasmith/CADKernel/Public/CADKernel/Topo/Topomaker.h
new file mode 100644
--- /dev/null
+++ b/Engine/Source/Runtime/Datasmith/CADKernel/Public/CADKernel/Topo/Topomaker.h
@@ -0,0 +1,130 @@
+// Copyright Epic Games, Inc. All Rights Reserved.
+#pragma once
+
+#include "CADKernel/Core/Types.h"
+
+namespace CADKernel
+{
+
+class FModel;
+class FSession;
+class FShell;
+class FTopologicalEdge;
+class FTopologicalFace;
+class FTopologicalVertex;
+
+class CADKERNEL_API FJoiner
+{
+	friend class FThinZone;
+	friend class FThinZoneFinder;
+
+protected:
+
+	FSession& Session;
+
+	TArray<TSharedPtr<FShell>> Shells;
+	TArray<TSharedPtr<FTopologicalFace>> Faces;
+
+	double JoiningTolerance;
+	double JoiningToleranceSquare;
+
+public:
+
+	FJoiner(FSession& InSession, const TSharedRef<FModel>& Shells, double Tolerance);
+	FJoiner(FSession& InSession, const TArray<TSharedPtr<FShell>>& Shells, double Tolerance);
+	FJoiner(FSession& InSession, const TArray<TSharedPtr<FTopologicalFace>>& Surfaces, double Tolerance);
+
+	void JoinFaces();
+	//void JoinFaces(bool bProcessOnlyBorderEdges, bool bProcessOnlyNonManifoldEdges);
+
+	/**
+	 * Check topology of each body
+	 */
+	void CheckTopology();
+
+	//void MergeInto(TSharedPtr<FBody> Body, TArray<TSharedPtr<FTopologicalEntity>>& InEntities);
+	//void SortByShell(TSharedPtr<FBody> Body, TArray<TSharedPtr<FBody>>& OutNewBody);
+	//void Join(TArray<TSharedPtr<FBody>> Bodies, double Tolerance);
+
+	/**
+	 * Split into connected shell and put each shell into the appropriate body
+	 */
+	void SplitIntoConnectedShell();
+
+	/**
+	 * Unlink Non-Manifold Vertex i.e. Vertex belong to tow or more shell
+	 */
+	void UnlinkNonManifoldVertex();
+
+private:
+
+	/**
+	 * Call by constructor.
+	 * For each shell, add their faces into Faces array, complete the metadata and set the states for the joining process
+	 */
+	void InitFaces();
+
+	void EmptyShells();
+
+	void RemoveFacesFromShell();
+
+	/**
+	 * Return an array of active vertices.
+	 */
+	void GetVertices(TArray<FTopologicalVertex*>& Vertices);
+
+	/**
+	 * Return an array of active border vertices.
+	 */
+	void GetBorderVertices(TArray<FTopologicalVertex*>& BorderVertices);
+
+	/**
+	 * Merge Border Vertices with other vertices.
+	 * @param Vertices: the initial array of active vertices, this array is updated at the end of the process
+	 */
+	void MergeCoincidentVertices(TArray<FTopologicalVertex*>& VerticesToMerge);
+
+	/**
+	 * Merge Border Vertices with other vertices.
+	 * @param Vertices: the initial array of active vertices, this array is updated at the end of the process
+	 */
+	void MergeBorderVerticesWithCoincidentOtherVertices(TArray<FTopologicalVertex*>& Vertices);
+
+	//void FixCollapsedEdges();
+	void CheckSelfConnectedEdge();
+	void RemoveIsolatedEdges(); // Useful ???
+
+	/**
+	 * First step, trivial edge merge i.e. couple of edges with same extremity vertex
+	 */
+	void MergeCoincidentEdges(TArray<FTopologicalVertex*>& Vertices);
+
+	/**
+	 * Second step, parallel edges but with different length. the longest must be split
+	 */
+	void StitchParallelEdges(TArray<FTopologicalVertex*>& Vertices);
+
+	/**
+	 * Split the edge to split except if one side is too small (JoiningTolerance). In this case (too small), the EdgeToSplit is not split but linked to EdgeToLink.
+	 * @return the created vertex or TSharedPtr<FTopologicalVertex>()
+	 */
+	TSharedPtr<FTopologicalVertex> SplitAndLink(FTopologicalVertex& StartVertex, FTopologicalEdge& EdgeToLink, FTopologicalEdge& EdgeToSplit);
+
+	/**
+	 * For each loop of each surface, check if successive edges are unconnected and if their common vertices are connected only to them.
+	 * These edges are merged into one edge.
+	 * E.g. :
+	 * Face A has 3 successive unconnected edges. If these 3 edges are merged to give only one edge, the new edge could be linked to its parallel edge of Face B
+	 *
+	 *              \                          Face A                                   |    Face C
+	 *               \                                                                  |
+	 *    Face E     CV ------------------- UV ------------------ UV ----------------- CV --------------------
+	 *               CV -------------------------------------------------------------- CV --------------------
+	 *               /                         Face B                                   |    Face D
+	 *              /                                                                   |
+	 */
+	void MergeUnconnectedAdjacentEdges();
+
+};
+
+} // namespace CADKernel
