Update CADKernel : add oriention, fix one infinite loop

diff --git a/Engine/Plugins/Enterprise/DatasmithCADImporter/Source/CADInterfaces/Private/CoreTechBridge.cpp b/Engine/Plugins/Enterprise/DatasmithCADImporter/Source/CADInterfaces/Private/CoreTechBridge.cpp
--- a/Engine/Plugins/Enterprise/DatasmithCADImporter/Source/CADInterfaces/Private/CoreTechBridge.cpp
+++ b/Engine/Plugins/Enterprise/DatasmithCADImporter/Source/CADInterfaces/Private/CoreTechBridge.cpp
@@ -94,6 +94,7 @@ FCoreTechBridge::FCoreTechBridge(FSession& InSession, FCADFileReport& InReport)
 
 TSharedRef<FBody> FCoreTechBridge::AddBody(CT_OBJECT_ID CTBodyId)
 {
+
 	TSharedPtr<FEntity>* EntityPtr = CTIdToEntity.Find((uint32)CTBodyId);
 	if (EntityPtr)
 	{
@@ -107,7 +108,7 @@ TSharedRef<FBody> FCoreTechBridge::AddBody(CT_OBJECT_ID CTBodyId)
 	TSharedRef<FBody> Body = FEntity::MakeShared<FBody>();
 	Report.BodyCount++;
 
-	AddMetadata(CTBodyId, Body);
+	AddMetadata(CTBodyId, *Body);
 
 #ifdef CORETECHBRIDGE_DEBUG
 	Body->CtKioId = (FIdent)CTBodyId;
@@ -144,7 +145,7 @@ TSharedRef<FBody> FCoreTechBridge::AddBody(CT_OBJECT_ID CTBodyId)
 		Body->AddShell(Shell);
 		Report.ShellCount++;
 
-		AddMetadata(CTShellId, Shell);
+		AddMetadata(CTShellId, *Shell);
 #ifdef CORETECHBRIDGE_DEBUG
 		Shell->CtKioId = (FIdent)CTShellId;
 #endif
@@ -267,7 +268,7 @@ void FCoreTechBridge::AddFace(CT_OBJECT_ID CTFaceId, TSharedRef<FShell>& Shell)
 	}
 
 	TSharedRef<FTopologicalFace> Face = FEntity::MakeShared<FTopologicalFace>(Surface);
-	AddMetadata(CTFaceId, Face);
+	AddMetadata(CTFaceId, *Face);
 	Face->SetPatchId((int32)CTFaceId);
 
 #ifdef CORETECHBRIDGE_DEBUG
@@ -675,6 +676,13 @@ TSharedPtr<FSurface> FCoreTechBridge::AddLinearTransfoSurface(CT_OBJECT_ID CTSur
 	TSharedPtr<FSurface> BaseSurface = AddSurface(CTBaseSurfaceId, FSurfacicBoundary());
 	TSharedPtr<FSurface> Surface = StaticCastSharedPtr<FSurface>(BaseSurface->ApplyMatrix(Matrix));
 
+//#ifdef CADKERNEL_DEV 
+//	const FMatrixH& BaseSurfaceMatrix = GetParamSpaceTransform(BaseSurface);
+//	if (!BaseSurfaceMatrix.IsId())
+//	{
+//		SetParamSpaceTransorm(Surface, BaseSurfaceMatrix);
+//	}
+//#endif
 	return Surface;
 }
 
@@ -999,12 +1007,41 @@ TSharedPtr<FCurve> FCoreTechBridge::AddCurveOnSurface(CT_OBJECT_ID CTCurveId)
 		return TSharedPtr<FCurve>();
 	}
 
+//#ifdef CADKERNEL_DEV
+//	// Curves on cylinders have UV param. and Z axis inversed:
+//	// Z'=-Z and {UV}' = {VU}
+//	if (CTSurfaceType == CT_CYLINDER_TYPE)
+//	{
+//		ensureCADKernel(false);
+//		FMatrixH MatrixTransformZandUV;
+//		MatrixTransformZandUV.SetIdentity();
+//		MatrixTransformZandUV[0] = 0;
+//		MatrixTransformZandUV[1] = -1;
+//		MatrixTransformZandUV[4] = 1;
+//		MatrixTransformZandUV[5] = 0;
+//		SetParamSpaceTransorm(Surface, MatrixTransformZandUV);
+//	}
+//#endif
+
 	TSharedPtr<FCurve> Curve2D = AddCurve(CTParametricCurveId, CTBaseSurfaceId);
 	if (!Curve2D.IsValid())
 	{
 		return TSharedPtr<FCurve>();
 	}
 
+//#ifdef CADKERNEL_DEV
+//	FMatrixH Matrix = GetParamSpaceTransform(Surface);
+//	if (!Matrix.IsId())
+//	{
+//		Curve2D = StaticCastSharedPtr<FCurve>(Curve2D->ApplyMatrix(Matrix));
+//	}
+//
+//	if (!Curve2D.IsValid())
+//	{
+//		return TSharedPtr<FCurve>();
+//	}
+//#endif
+
 	return FEntity::MakeShared<FSurfacicCurve>(Curve2D.ToSharedRef(), Surface.ToSharedRef());
 }
 
@@ -1105,9 +1142,9 @@ static void GetAttributeValue(CT_ATTRIB_TYPE AttributType, int IthField, EValueT
 	}
 }
 
-void FCoreTechBridge::AddMetadata(CT_OBJECT_ID CTNodeId, TSharedRef<FMetadataDictionary> Entity)
+void FCoreTechBridge::AddMetadata(CT_OBJECT_ID CTNodeId, FTopologicalShapeEntity& ShapeEntity)
 {
-	Entity->SetHostId((uint32) CTNodeId);
+	ShapeEntity.SetHostId((uint32) CTNodeId);
 
 	CT_UINT32 IthAttrib = 0;
 	while (CT_OBJECT_IO::SearchAttribute(CTNodeId, CT_ATTRIB_ALL, IthAttrib++) == IO_OK)
@@ -1122,12 +1159,12 @@ void FCoreTechBridge::AddMetadata(CT_OBJECT_ID CTNodeId, TSharedRef<FMetadataDic
 		{
 		case CT_ATTRIB_NAME:
 		{
-			if (Entity->GetName() == nullptr) // otherwise the name has been set with ORIGINAL_NAME that is prefered in all case
+			if (ShapeEntity.HasName()) // otherwise the name has been set with ORIGINAL_NAME that is prefered in all case
 			{
 				CT_STR NameStrValue;
 				if (CT_CURRENT_ATTRIB_IO::AskStrField(ITH_NAME_VALUE, NameStrValue) == IO_OK)
 				{
-					Entity->SetName(AsFString(NameStrValue));
+					ShapeEntity.SetName(AsFString(NameStrValue));
 				}
 				break;
 			}
@@ -1138,7 +1175,7 @@ void FCoreTechBridge::AddMetadata(CT_OBJECT_ID CTNodeId, TSharedRef<FMetadataDic
 			CT_STR NameStrValue;
 			if (CT_CURRENT_ATTRIB_IO::AskStrField(ITH_NAME_VALUE, NameStrValue) == IO_OK)
 			{
-				Entity->SetName(AsFString(NameStrValue));
+				ShapeEntity.SetName(AsFString(NameStrValue));
 			}
 			break;
 		}
@@ -1150,7 +1187,7 @@ void FCoreTechBridge::AddMetadata(CT_OBJECT_ID CTNodeId, TSharedRef<FMetadataDic
 			{
 				LayerId = 0;
 			}
-			Entity->SetLayer(LayerId);
+			ShapeEntity.SetLayer(LayerId);
 			break;
 
 		}
@@ -1174,7 +1211,7 @@ void FCoreTechBridge::AddMetadata(CT_OBJECT_ID CTNodeId, TSharedRef<FMetadataDic
 			}
 
 			uint32 ColorHId = CADLibrary::BuildColorId(ColorId, Alpha);
-			Entity->SetColorId(ColorHId);
+			ShapeEntity.SetColorId(ColorHId);
 		}
 
 		case CT_ATTRIB_MATERIALID:
@@ -1182,7 +1219,7 @@ void FCoreTechBridge::AddMetadata(CT_OBJECT_ID CTNodeId, TSharedRef<FMetadataDic
 			int32 MaterialId = 0;
 			if (CT_CURRENT_ATTRIB_IO::AskIntField(ITH_MATERIALID_VALUE, MaterialId) == IO_OK)
 			{
-				Entity->SetMaterialId(MaterialId);
+				ShapeEntity.SetMaterialId(MaterialId);
 			}
 			break;
 		}
diff --git a/Engine/Plugins/Enterprise/DatasmithCADImporter/Source/CADInterfaces/Private/CoreTechBridge.h b/Engine/Plugins/Enterprise/DatasmithCADImporter/Source/CADInterfaces/Private/CoreTechBridge.h
--- a/Engine/Plugins/Enterprise/DatasmithCADImporter/Source/CADInterfaces/Private/CoreTechBridge.h
+++ b/Engine/Plugins/Enterprise/DatasmithCADImporter/Source/CADInterfaces/Private/CoreTechBridge.h
@@ -1,5 +1,6 @@
 // Copyright Epic Games, Inc. All Rights Reserved.
 #pragma once
+
 #include "CADKernel/Core/Types.h"
 
 #ifdef USE_KERNEL_IO_SDK 
@@ -32,6 +33,7 @@ class FShell;
 class FSurface;
 class FTopologicalEdge;
 class FTopologicalLoop;
+class FTopologicalShapeEntity;
 
 class FCoreTechBridge
 {
@@ -98,28 +100,7 @@ private:
 
 	FMatrixH CreateCoordinateSystem(const CT_COORDINATE& InOrigin, const CT_VECTOR& InDirection, const CT_VECTOR& InURef);
 
-	void AddMetadata(CT_OBJECT_ID NodeId, TSharedRef<FMetadataDictionary> Entity);
-
-	template<typename EntityType>
-	void AddMetadata(CT_OBJECT_ID NodeId, TSharedRef<EntityType> Entity)
-	{
-		AddMetadata(NodeId, (TSharedRef<FMetadataDictionary>) Entity);
-	}
-
-//#ifdef CADKERNEL_DEV
-//	TMap<const uint32, FMatrixH> SurfaceToMatrix;
-//
-//	const FMatrixH& GetParamSpaceTransform(TSharedPtr<FSurface>& Surface)
-//	{
-//		FMatrixH* Matrix = SurfaceToMatrix.Find(Surface->GetId());
-//		return Matrix == nullptr ? FMatrixH::Identity : *Matrix;
-//	}
-//
-//	void SetParamSpaceTransorm(TSharedPtr<FSurface>& Surface, const FMatrixH Matrix)
-//	{
-//		SurfaceToMatrix.Add(Surface->GetId(), Matrix);
-//	}
-//#endif
+	void AddMetadata(CT_OBJECT_ID NodeId, FTopologicalShapeEntity& Entity);
 
 };
 }
diff --git a/Engine/Plugins/Enterprise/DatasmithCADImporter/Source/CADInterfaces/Private/CoreTechFileParser.h b/Engine/Plugins/Enterprise/DatasmithCADImporter/Source/CADInterfaces/Private/CoreTechFileParser.h
--- a/Engine/Plugins/Enterprise/DatasmithCADImporter/Source/CADInterfaces/Private/CoreTechFileParser.h
+++ b/Engine/Plugins/Enterprise/DatasmithCADImporter/Source/CADInterfaces/Private/CoreTechFileParser.h
@@ -77,7 +77,7 @@ namespace CADLibrary
 		bool BuildStaticMeshData(CT_OBJECT_ID NodeId, CT_OBJECT_ID ParentId, uint32 ParentMaterialHash);
 		
 		void ReadAndSewBodies(const TArray<CT_OBJECT_ID>& Bodies, CT_OBJECT_ID ParentId, uint32 ParentMaterialHash, TArray<FCadId>& OutChildren);
-		void BuildStaticMeshData(CADKernel::FSession& CADKernelSession, TSharedRef<CADKernel::FBody> CADKernelBody, CT_OBJECT_ID ParentId, uint32 ParentMaterialHash);
+		void BuildStaticMeshData(CADKernel::FSession& CADKernelSession, CADKernel::FBody& CADKernelBody, CT_OBJECT_ID ParentId, uint32 ParentMaterialHash);
 
 		uint32 GetMaterialNum();
 		void ReadMaterials();
diff --git a/Engine/Plugins/Enterprise/DatasmithCADImporter/Source/CADInterfaces/Private/CoreTechFileParserImpl.cpp b/Engine/Plugins/Enterprise/DatasmithCADImporter/Source/CADInterfaces/Private/CoreTechFileParserImpl.cpp
--- a/Engine/Plugins/Enterprise/DatasmithCADImporter/Source/CADInterfaces/Private/CoreTechFileParserImpl.cpp
+++ b/Engine/Plugins/Enterprise/DatasmithCADImporter/Source/CADInterfaces/Private/CoreTechFileParserImpl.cpp
@@ -32,8 +32,9 @@
 #include "CADKernel/Mesh/Meshers/ParametricMesher.h"
 #include "CADKernel/Mesh/Structure/ModelMesh.h"
 
-#include "CADKernel/Topo/Model.h"
 #include "CADKernel/Topo/Body.h"
+#include "CADKernel/Topo/Model.h"
+#include "CADKernel/Topo/TopologicalShapeEntity.h"
 #include "CADKernel/Topo/Topomaker.h"
 
 #include "HAL/FileManager.h"
@@ -721,20 +722,20 @@ namespace CADLibrary
 			double GeometricTolerance = 0.00001 / CADFileData.GetImportParameters().GetMetricUnit();
 			CADKernel::FSession CADKernelSession(GeometricTolerance);
 
-			TSharedRef<CADKernel::FModel> CADKernelModel = CADKernelSession.GetModel();
+			CADKernel::FModel& CADKernelModel = CADKernelSession.GetModel();
 
 			CADKernel::FCADFileReport Report;
 			CADKernel::FCoreTechBridge CoreTechBridge(CADKernelSession, Report);
 
 			TSharedRef<CADKernel::FBody> CADKernelBody = CoreTechBridge.AddBody(BodyId);
-			CADKernelModel->Add(CADKernelBody);
+			CADKernelModel.Add(CADKernelBody);
 
 			// Repair if needed
 			if(CADFileData.GetImportParameters().GetStitchingTechnique() != EStitchingTechnique::StitchingNone)
 			{
 				double Tolerance = CADFileData.GetImportParameters().ConvertMMToImportUnit(0.1); // mm
-				CADKernel::FJoiner Joiner(CADKernelSession, CADKernelModel, Tolerance);
-				Joiner.JoinFaces();
+				CADKernel::FTopomaker Topomaker(CADKernelSession, Tolerance);
+				Topomaker.Sew();
 			}
 
 #ifdef CORETECHBRIDGE_DEBUG
@@ -752,9 +753,9 @@ namespace CADLibrary
 			// Tessellate the body
 			TSharedRef<CADKernel::FModelMesh> CADKernelModelMesh = CADKernel::FEntity::MakeShared<CADKernel::FModelMesh>();
 
-			FCADKernelTools::DefineMeshCriteria(CADKernelModelMesh, CADFileData.GetImportParameters(), GeometricTolerance);
+			FCADKernelTools::DefineMeshCriteria(*CADKernelModelMesh, CADFileData.GetImportParameters(), GeometricTolerance);
 
-			CADKernel::FParametricMesher Mesher(CADKernelModelMesh);
+			CADKernel::FParametricMesher Mesher(*CADKernelModelMesh);
 			Mesher.MeshEntity(CADKernelModel);
 
 			TFunction<void(FObjectDisplayDataId, FObjectDisplayDataId, int32)> ProcessFace;
@@ -763,7 +764,7 @@ namespace CADLibrary
 				SetFaceMainMaterial(FaceMaterial, BodyMaterial, BodyMesh, Index);
 			};
 
-			FCADKernelTools::GetBodyTessellation(CADKernelModelMesh, CADKernelBody, BodyMesh, DefaultMaterialHash, ProcessFace);
+			FCADKernelTools::GetBodyTessellation(*CADKernelModelMesh, *CADKernelBody, BodyMesh, DefaultMaterialHash, ProcessFace);
 
 			CADKernelSession.Clear();
 		}
@@ -773,14 +774,14 @@ namespace CADLibrary
 		return true;
 	}
 
-	void FCoreTechFileParser::BuildStaticMeshData(CADKernel::FSession& CADKernelSession, TSharedRef<CADKernel::FBody> CADKernelBody, CT_OBJECT_ID ParentId, uint32 DefaultMaterialHash)
+	void FCoreTechFileParser::BuildStaticMeshData(CADKernel::FSession& CADKernelSession, CADKernel::FBody& CADKernelBody, CT_OBJECT_ID ParentId, uint32 DefaultMaterialHash)
 	{
-		int32 Index = CADFileData.AddBody(CADKernelBody->GetHostId());
+		int32 Index = CADFileData.AddBody(CADKernelBody.GetHostId());
 		FArchiveBody& ArchiveBody = CADFileData.GetBodyAt(Index);
 
-		CADKernelBody->GetMetaData(ArchiveBody.MetaData);
+		CADKernelBody.ExtractMetaData(ArchiveBody.MetaData);
 
-		FBodyMesh& BodyMesh = CADFileData.AddBodyMesh(CADKernelBody->GetHostId(), ArchiveBody);
+		FBodyMesh& BodyMesh = CADFileData.AddBodyMesh(CADKernelBody.GetHostId(), ArchiveBody);
 		if (uint32 MaterialHash = GetObjectMaterial(ArchiveBody))
 		{
 			DefaultMaterialHash = MaterialHash;
@@ -801,9 +802,9 @@ namespace CADLibrary
 		// Tessellate the body
 		TSharedRef<CADKernel::FModelMesh> CADKernelModelMesh = CADKernel::FEntity::MakeShared<CADKernel::FModelMesh>();
 
-		FCADKernelTools::DefineMeshCriteria(CADKernelModelMesh, CADFileData.GetImportParameters(), CADKernelSession.GetGeometricTolerance());
+		FCADKernelTools::DefineMeshCriteria(CADKernelModelMesh.Get(), CADFileData.GetImportParameters(), CADKernelSession.GetGeometricTolerance());
 
-		CADKernel::FParametricMesher Mesher(CADKernelModelMesh);
+		CADKernel::FParametricMesher Mesher(*CADKernelModelMesh);
 		Mesher.MeshEntity(CADKernelBody);
 
 		TFunction<void(FObjectDisplayDataId, FObjectDisplayDataId, int32)> ProcessFace;
@@ -812,7 +813,7 @@ namespace CADLibrary
 			SetFaceMainMaterial(FaceMaterial, BodyMaterial, BodyMesh, Index);
 		};
 
-		FCADKernelTools::GetBodyTessellation(CADKernelModelMesh, CADKernelBody, BodyMesh, DefaultMaterialHash, ProcessFace);
+		FCADKernelTools::GetBodyTessellation(CADKernelModelMesh.Get(), CADKernelBody, BodyMesh, DefaultMaterialHash, ProcessFace);
 
 		ArchiveBody.ColorFaceSet = BodyMesh.ColorSet;
 		ArchiveBody.MaterialFaceSet = BodyMesh.MaterialSet;
@@ -823,7 +824,7 @@ namespace CADLibrary
 		double GeometricTolerance = CADFileData.GetImportParameters().ConvertMMToImportUnit(0.01); // mm
 		CADKernel::FSession CADKernelSession(GeometricTolerance);
 		CADKernelSession.SetFirstNewHostId(LastHostIdUsed);
-		TSharedRef<CADKernel::FModel> CADKernelModel = CADKernelSession.GetModel();
+		CADKernel::FModel& CADKernelModel = CADKernelSession.GetModel();
 
 		CADKernel::FCADFileReport Report;
 		CADKernel::FCoreTechBridge CoreTechBridge(CADKernelSession, Report);
@@ -831,27 +832,28 @@ namespace CADLibrary
 		for(CT_OBJECT_ID BodyId : Bodies)
 		{
 			TSharedRef<CADKernel::FBody> CADKernelBody = CoreTechBridge.AddBody(BodyId);
-			CADKernelModel->Add(CADKernelBody);
+			CADKernelModel.Add(CADKernelBody);
 		}
 
 		// Repair if needed
 		if (CADFileData.GetImportParameters().GetStitchingTechnique() != EStitchingTechnique::StitchingNone)
 		{
 			double Tolerance = CADFileData.GetImportParameters().ConvertMMToImportUnit(0.1); // mm
-			CADKernel::FJoiner Joiner(CADKernelSession, CADKernelModel, Tolerance);
-			Joiner.JoinFaces();
-			Joiner.SplitIntoConnectedShell();
+			CADKernel::FTopomaker Topomaker(CADKernelSession, Tolerance);
+			Topomaker.Sew();
+			Topomaker.SplitIntoConnectedShells();
+			Topomaker.OrientShells();
 		}
 
 		// Get CADKernel new bodies
-		for(const TSharedPtr<CADKernel::FBody>& CADKernelBody : CADKernelModel->GetBodies())
+		for(const TSharedPtr<CADKernel::FBody>& CADKernelBody : CADKernelModel.GetBodies())
 		{
 			if (!CADKernelBody.IsValid())
 			{
 				continue;
 			}
 
-			BuildStaticMeshData(CADKernelSession, CADKernelBody.ToSharedRef(), ParentId, ParentMaterialHash);
+			BuildStaticMeshData(CADKernelSession, *CADKernelBody, ParentId, ParentMaterialHash);
 			OutChildren.Add(CADKernelBody->GetHostId());
 		}
 
diff --git a/Engine/Plugins/Enterprise/DatasmithCADImporter/Source/CADInterfaces/Private/TechSoftBridge.cpp b/Engine/Plugins/Enterprise/DatasmithCADImporter/Source/CADInterfaces/Private/TechSoftBridge.cpp
--- a/Engine/Plugins/Enterprise/DatasmithCADImporter/Source/CADInterfaces/Private/TechSoftBridge.cpp
+++ b/Engine/Plugins/Enterprise/DatasmithCADImporter/Source/CADInterfaces/Private/TechSoftBridge.cpp
@@ -729,7 +729,7 @@ TSharedRef<FBody> FTechSoftBridge::AddBody(A3DRiBrepModel* BrepModel, FEntityDat
 	TSharedRef<FBody> Body = FEntity::MakeShared<FBody>();
 	TechSoftToEntity.Add(BrepModel, Body);
 
-	AddMetadata(BrepMetaData, Body);
+	AddMetadata(BrepMetaData, *Body);
 
 	CADLibrary::TUniqueTSObj<A3DRiBrepModelData> BRepModelData(BrepModel);
 	if (BRepModelData.IsValid())
@@ -797,7 +797,7 @@ void FTechSoftBridge::TraverseShell(const A3DTopoShell* A3DShell, TSharedRef<FBo
 	Body->AddShell(Shell);
 	Report.ShellCount++;
 
-	AddMetadata(MetaData, Shell);
+	AddMetadata(MetaData, *Shell);
 
 	CADLibrary::TUniqueTSObj<A3DTopoShellData> ShellData(A3DShell);
 	if (ShellData.IsValid())
@@ -1374,7 +1374,7 @@ void FTechSoftBridge::AddFace(const A3DTopoFace* A3DFace, TSharedRef<FShell>& Sh
 		}
 	}
 
-	//AddMetadata(CTFaceId, Face);
+	AddMetadata(MetaData, *Face);
 	//Face->SetPatchId((int32)CTFaceId);
 
 	EOrientation Orientation = EOrientation::Front; //CTOrientation == CT_FORWARD ? EOrientation::Front : EOrientation::Back;
@@ -1938,17 +1938,17 @@ TSharedPtr<FSurface> FTechSoftBridge::AddSurfaceAsNurbs(const A3DSurfBase* Surfa
 
 }
 
-void FTechSoftBridge::AddMetadata(FEntityData& MetaData, TSharedRef<FMetadataDictionary> Entity)
+void FTechSoftBridge::AddMetadata(FEntityData& MetaData, FTopologicalShapeEntity& Entity)
 {
 	//Entity->SetHostId(CTNodeId);
-	Entity->SetName(MetaData.Name);
-	//Entity->SetLayer(LayerId);
+	Entity.SetName(*MetaData.Name);
+	//Entity.SetLayer(LayerId);
 	//uint32 ColorHId = CADLibrary::BuildColorId(ColorId, Alpha);
-	//Entity->SetColorId(ColorHId);
-	//Entity->SetPatchId();
+	//Entity.SetColorId(ColorHId);
+	//Entity.SetPatchId();
 }
 
 }
 
 
-#endif
\ No newline at end of file
+#endif // USE_TECHSOFT_SDK
\ No newline at end of file
diff --git a/Engine/Plugins/Enterprise/DatasmithCADImporter/Source/CADInterfaces/Private/TechSoftBridge.h b/Engine/Plugins/Enterprise/DatasmithCADImporter/Source/CADInterfaces/Private/TechSoftBridge.h
--- a/Engine/Plugins/Enterprise/DatasmithCADImporter/Source/CADInterfaces/Private/TechSoftBridge.h
+++ b/Engine/Plugins/Enterprise/DatasmithCADImporter/Source/CADInterfaces/Private/TechSoftBridge.h
@@ -217,7 +217,6 @@ class FCriterion;
 class FCurve;
 class FEntity;
 class FFaceMesh;
-class FMetadataDictionary;
 class FModel;
 class FPoint;
 class FRestrictionCurve;
@@ -226,6 +225,7 @@ class FShell;
 class FSurface;
 class FTopologicalEdge;
 class FTopologicalLoop;
+class FTopologicalShapeEntity;
 
 class FTechSoftBridge
 {
@@ -306,10 +306,10 @@ private:
 
 	TSharedPtr<FCurve> AddCurveAsNurbs(const A3DCrvBase* A3DCurve, const TechSoftUtils::FUVReparameterization& UVReparameterization);
 
-	void AddMetadata(FEntityData& EntityData, TSharedRef<FMetadataDictionary> Entity);
+	void AddMetadata(FEntityData& EntityData, FTopologicalShapeEntity& Entity);
 };
 }
 
 
 
-#endif // USE_KERNEL_IO_SDK
+#endif // USE_TECHSOFT_SDK
diff --git a/Engine/Plugins/Enterprise/DatasmithCADImporter/Source/CADKernelSurface/Private/CADKernelSurfaceExtension.cpp b/Engine/Plugins/Enterprise/DatasmithCADImporter/Source/CADKernelSurface/Private/CADKernelSurfaceExtension.cpp
--- a/Engine/Plugins/Enterprise/DatasmithCADImporter/Source/CADKernelSurface/Private/CADKernelSurfaceExtension.cpp
+++ b/Engine/Plugins/Enterprise/DatasmithCADImporter/Source/CADKernelSurface/Private/CADKernelSurfaceExtension.cpp
@@ -59,15 +59,14 @@ bool UCADKernelParametricSurfaceData::Tessellate(UStaticMesh& StaticMesh, const
 		TSharedRef<CADKernel::FSession> CADKernelSession = MakeShared<CADKernel::FSession>(0.00001 / ImportParameters.GetMetricUnit());
 		CADKernelSession->AddDatabase(RawData);
 
-		TSharedRef<CADKernel::FModel> CADKernelModel = CADKernelSession->GetModel();
-		TArray<TSharedPtr<CADKernel::FBody>> CADKernelBodies = CADKernelModel->GetBodies();
+		CADKernel::FModel& CADKernelModel = CADKernelSession->GetModel();
+		TArray<TSharedPtr<CADKernel::FBody>> CADKernelBodies = CADKernelModel.GetBodies();
 		if (CADKernelBodies.Num() != 1)
 		{
 			return bSuccessfulTessellation;
 		}
 
-		TSharedRef<CADKernel::FTopologicalEntity> CADKernelEntity = CADKernelModel;
-		if(CADLibrary::FCADKernelTools::Tessellate(CADKernelEntity, ImportParameters, CadMeshParameters, MeshDescription))
+		if(CADLibrary::FCADKernelTools::Tessellate(CADKernelModel, ImportParameters, CadMeshParameters, MeshDescription))
 		{
 			// To update the SectionInfoMap 
 			TPolygonGroupAttributesConstRef<FName> MaterialSlotNames = MeshDescriptionAttributes.GetPolygonGroupMaterialSlotNames();
diff --git a/Engine/Plugins/Enterprise/DatasmithCADImporter/Source/CADKernelSurface/Public/CADModelToCADKernelConverterBase.h b/Engine/Plugins/Enterprise/DatasmithCADImporter/Source/CADKernelSurface/Public/CADModelToCADKernelConverterBase.h
--- a/Engine/Plugins/Enterprise/DatasmithCADImporter/Source/CADKernelSurface/Public/CADModelToCADKernelConverterBase.h
+++ b/Engine/Plugins/Enterprise/DatasmithCADImporter/Source/CADKernelSurface/Public/CADModelToCADKernelConverterBase.h
@@ -41,9 +41,8 @@ public:
 		{
 			// the joining tolerance is set to 0.1 mm until the user can specify it
 			double JoiningTolerance = ImportParameters.ConvertMMToImportUnit(0.1);
-			TSharedRef<CADKernel::FModel> Model = CADKernelSession.GetModel();
-			CADKernel::FJoiner Joiner(CADKernelSession, Model, JoiningTolerance);
-			Joiner.JoinFaces();
+			CADKernel::FTopomaker Topomaker(CADKernelSession, JoiningTolerance);
+			Topomaker.Sew();
 		}
 
 		return true;
@@ -59,7 +58,7 @@ public:
 
 	virtual bool Tessellate(const CADLibrary::FMeshParameters& InMeshParameters, FMeshDescription& OutMeshDescription) override
 	{
-		TSharedRef<CADKernel::FTopologicalEntity> Model = CADKernelSession.GetModel();
+		CADKernel::FModel& Model = CADKernelSession.GetModel();
 		return CADLibrary::FCADKernelTools::Tessellate(Model, ImportParameters, InMeshParameters, OutMeshDescription);
 	}
 
diff --git a/Engine/Plugins/Enterprise/DatasmithCADImporter/Source/CADLibrary/Private/CADKernelTools.cpp b/Engine/Plugins/Enterprise/DatasmithCADImporter/Source/CADLibrary/Private/CADKernelTools.cpp
--- a/Engine/Plugins/Enterprise/DatasmithCADImporter/Source/CADLibrary/Private/CADKernelTools.cpp
+++ b/Engine/Plugins/Enterprise/DatasmithCADImporter/Source/CADLibrary/Private/CADKernelTools.cpp
@@ -23,7 +23,7 @@
 #include "CADKernel/Topo/Body.h"
 #include "CADKernel/Topo/Shell.h"
 #include "CADKernel/Topo/TopologicalEdge.h"
-#include "CADKernel/Topo/TopologicalEntity.h"
+#include "CADKernel/Topo/TopologicalShapeEntity.h"
 #include "CADKernel/Topo/TopologicalFace.h"
 #include "CADKernel/Topo/TopologicalVertex.h"
 
@@ -275,15 +275,15 @@ namespace CADLibrary
 		return MeshDescription.Polygons().Num() > 0;
 	}
 
-	bool FCADKernelTools::Tessellate(TSharedRef<CADKernel::FTopologicalEntity>& CADTopologicalEntity, const FImportParameters& ImportParameters, const FMeshParameters& MeshParameters, FMeshDescription& OutMeshDescription)
+	bool FCADKernelTools::Tessellate(CADKernel::FTopologicalShapeEntity& CADTopologicalEntity, const FImportParameters& ImportParameters, const FMeshParameters& MeshParameters, FMeshDescription& OutMeshDescription)
 	{
 		using namespace CADKernel;
 
 		// Tessellate the model
 		TSharedRef<FModelMesh> CADKernelModelMesh = FEntity::MakeShared<FModelMesh>();
-		FParametricMesher Mesher(CADKernelModelMesh);
+		FParametricMesher Mesher(*CADKernelModelMesh);
 
-		DefineMeshCriteria(CADKernelModelMesh, ImportParameters, 0.01);
+		DefineMeshCriteria(*CADKernelModelMesh, ImportParameters, 0.01);
 		Mesher.MeshEntity(CADTopologicalEntity);
 
 		FMeshConversionContext Context(ImportParameters, MeshParameters);
@@ -292,41 +292,39 @@ namespace CADLibrary
 	}
 
 
-	uint32 FCADKernelTools::GetFaceTessellation(const TSharedRef<CADKernel::FFaceMesh>& FaceMesh, FBodyMesh& OutBodyMesh)
+	uint32 FCADKernelTools::GetFaceTessellation(CADKernel::FFaceMesh& FaceMesh, FBodyMesh& OutBodyMesh)
 	{
 		// Something wrong happened, either an error or no data to collect
-		if (FaceMesh->TrianglesVerticesIndex.Num() == 0)
+		if (FaceMesh.TrianglesVerticesIndex.Num() == 0)
 		{
 			return 0;
 		}
 
 		FTessellationData& Tessellation = OutBodyMesh.Faces.Emplace_GetRef();
 
-		const CADKernel::FTopologicalFace& HaveMetadata = (const CADKernel::FTopologicalFace&) FaceMesh->GetGeometricEntity();
-		Tessellation.PatchId = HaveMetadata.GetPatchId();
+		const CADKernel::FTopologicalFace& TopologicalFace = (const CADKernel::FTopologicalFace&) FaceMesh.GetGeometricEntity();
+		Tessellation.PatchId = TopologicalFace.GetPatchId();
 
-		Tessellation.PositionIndices = MoveTemp(FaceMesh->VerticesGlobalIndex);
-		Tessellation.VertexIndices = MoveTemp(FaceMesh->TrianglesVerticesIndex);
+		Tessellation.PositionIndices = MoveTemp(FaceMesh.VerticesGlobalIndex);
+		Tessellation.VertexIndices = MoveTemp(FaceMesh.TrianglesVerticesIndex);
 
-		Tessellation.NormalArray = MoveTemp(FaceMesh->Normals);
-		Tessellation.TexCoordArray = MoveTemp(FaceMesh->UVMap);
+		Tessellation.NormalArray = MoveTemp(FaceMesh.Normals);
+		Tessellation.TexCoordArray = MoveTemp(FaceMesh.UVMap);
 
 		return Tessellation.VertexIndices.Num() / 3;
 	}
 
-	template<class ClassType>
-	void GetDisplayDataIds(const TSharedRef<ClassType>& Entity, FObjectDisplayDataId& DisplayDataId)
+	void GetDisplayDataIds(const CADKernel::FTopologicalShapeEntity& ShapeEntity, FObjectDisplayDataId& DisplayDataId)
 	{
-		const TSharedRef<CADKernel::FMetadataDictionary>& HaveMetadata = StaticCastSharedRef<CADKernel::FMetadataDictionary>(Entity);
-		DisplayDataId.Color = HaveMetadata->GetColorId();
-		DisplayDataId.Material = HaveMetadata->GetMaterialId();
+		DisplayDataId.Color = ShapeEntity.GetColorId();
+		DisplayDataId.Material = ShapeEntity.GetMaterialId();
 	}
 
-	void FCADKernelTools::GetBodyTessellation(const TSharedRef<CADKernel::FModelMesh>& ModelMesh, const TSharedRef<CADKernel::FBody>& Body, FBodyMesh& OutBodyMesh, uint32 DefaultMaterialHash, TFunction<void(FObjectDisplayDataId, FObjectDisplayDataId, int32)> SetFaceMainMaterial)
+	void FCADKernelTools::GetBodyTessellation(const CADKernel::FModelMesh& ModelMesh, const CADKernel::FBody& Body, FBodyMesh& OutBodyMesh, uint32 DefaultMaterialHash, TFunction<void(FObjectDisplayDataId, FObjectDisplayDataId, int32)> SetFaceMainMaterial)
 	{
-		ModelMesh->GetNodeCoordinates(OutBodyMesh.VertexArray);
+		ModelMesh.GetNodeCoordinates(OutBodyMesh.VertexArray);
 
-		uint32 FaceSize = Body->FaceCount();
+		uint32 FaceSize = Body.FaceCount();
 
 		// Allocate memory space for tessellation data
 		OutBodyMesh.Faces.Reserve(FaceSize);
@@ -340,7 +338,7 @@ namespace CADLibrary
 
 		// Loop through the face of bodies and collect all tessellation data
 		int32 FaceIndex = 0;
-		for (const TSharedPtr<CADKernel::FShell>& Shell : Body->GetShells())
+		for (const TSharedPtr<CADKernel::FShell>& Shell : Body.GetShells())
 		{
 			if (!Shell.IsValid())
 			{
@@ -348,7 +346,7 @@ namespace CADLibrary
 			}
 
 			FObjectDisplayDataId ShellMaterial = BodyMaterial;
-			GetDisplayDataIds(Shell.ToSharedRef(), ShellMaterial);
+			GetDisplayDataIds(*Shell, ShellMaterial);
 
 			for (const CADKernel::FOrientedFace& Face : Shell->GetFaces())
 			{
@@ -363,9 +361,9 @@ namespace CADLibrary
 				}
 
 				FObjectDisplayDataId FaceMaterial;
-				GetDisplayDataIds(Face.Entity.ToSharedRef(), FaceMaterial);
+				GetDisplayDataIds(*Face.Entity, FaceMaterial);
 
-				uint32 TriangleNum = GetFaceTessellation(Face.Entity->GetMesh(), OutBodyMesh);
+				uint32 TriangleNum = GetFaceTessellation(*Face.Entity->GetMesh(), OutBodyMesh);
 
 				if (TriangleNum == 0)
 				{
@@ -383,32 +381,32 @@ namespace CADLibrary
 		}
 	}
 
-	void FCADKernelTools::DefineMeshCriteria(TSharedRef<CADKernel::FModelMesh>& MeshModel, const FImportParameters& ImportParameters, double GeometricTolerance)
+	void FCADKernelTools::DefineMeshCriteria(CADKernel::FModelMesh& MeshModel, const FImportParameters& ImportParameters, double GeometricTolerance)
 	{
 		{
 			TSharedPtr<CADKernel::FCriterion> CurvatureCriterion = CADKernel::FCriterion::CreateCriterion(CADKernel::ECriterion::CADCurvature);
-			MeshModel->AddCriterion(CurvatureCriterion);
+			MeshModel.AddCriterion(CurvatureCriterion);
 
 			TSharedPtr<CADKernel::FCriterion> MinSizeCriterion = CADKernel::FCriterion::CreateCriterion(CADKernel::ECriterion::MinSize, 2. * GeometricTolerance);
-			MeshModel->AddCriterion(MinSizeCriterion);
+			MeshModel.AddCriterion(MinSizeCriterion);
 		}
 
 		if (ImportParameters.GetMaxEdgeLength() > SMALL_NUMBER)
 		{
 			TSharedPtr<CADKernel::FCriterion> MaxSizeCriterion = CADKernel::FCriterion::CreateCriterion(CADKernel::ECriterion::MaxSize, ImportParameters.GetMaxEdgeLength() / ImportParameters.GetScaleFactor());
-			MeshModel->AddCriterion(MaxSizeCriterion);
+			MeshModel.AddCriterion(MaxSizeCriterion);
 		}
 
 		if (ImportParameters.GetChordTolerance() > SMALL_NUMBER)
 		{
 			TSharedPtr<CADKernel::FCriterion> ChordCriterion = CADKernel::FCriterion::CreateCriterion(CADKernel::ECriterion::Sag, ImportParameters.GetChordTolerance() / ImportParameters.GetScaleFactor());
-			MeshModel->AddCriterion(ChordCriterion);
+			MeshModel.AddCriterion(ChordCriterion);
 		}
 
 		if (ImportParameters.GetMaxNormalAngle() > SMALL_NUMBER)
 		{
 			TSharedPtr<CADKernel::FCriterion> MaxNormalAngleCriterion = CADKernel::FCriterion::CreateCriterion(CADKernel::ECriterion::Angle, ImportParameters.GetMaxNormalAngle());
-			MeshModel->AddCriterion(MaxNormalAngleCriterion);
+			MeshModel.AddCriterion(MaxNormalAngleCriterion);
 		}
 	}
 
diff --git a/Engine/Plugins/Enterprise/DatasmithCADImporter/Source/CADLibrary/Public/CADKernelTools.h b/Engine/Plugins/Enterprise/DatasmithCADImporter/Source/CADLibrary/Public/CADKernelTools.h
--- a/Engine/Plugins/Enterprise/DatasmithCADImporter/Source/CADLibrary/Public/CADKernelTools.h
+++ b/Engine/Plugins/Enterprise/DatasmithCADImporter/Source/CADLibrary/Public/CADKernelTools.h
@@ -11,7 +11,7 @@ namespace CADKernel
 	class FBody;
 	class FBodyMesh;
 	class FFaceMesh;
-	class FTopologicalEntity;
+	class FTopologicalShapeEntity;
 	class FModelMesh;
 }
 
@@ -25,15 +25,15 @@ namespace CADLibrary
 	class CADLIBRARY_API FCADKernelTools
 	{
 	public:
-		static void DefineMeshCriteria(TSharedRef<CADKernel::FModelMesh>& MeshModel, const FImportParameters& ImportParameters, double GeometricTolerance);
+		static void DefineMeshCriteria(CADKernel::FModelMesh& MeshModel, const FImportParameters& ImportParameters, double GeometricTolerance);
 
-		static void GetBodyTessellation(const TSharedRef<CADKernel::FModelMesh>& ModelMesh, const TSharedRef<CADKernel::FBody>& Body, FBodyMesh& OutBodyMesh, uint32 DefaultMaterialHash, TFunction<void(FObjectDisplayDataId, FObjectDisplayDataId, int32)> SetFaceMainMaterial);
+		static void GetBodyTessellation(const CADKernel::FModelMesh& ModelMesh, const CADKernel::FBody& Body, FBodyMesh& OutBodyMesh, uint32 DefaultMaterialHash, TFunction<void(FObjectDisplayDataId, FObjectDisplayDataId, int32)> SetFaceMainMaterial);
 
 		/**
 		 * Tessellate a CADKernel entity and update the MeshDescription
 		 */
-		static bool Tessellate(TSharedRef<CADKernel::FTopologicalEntity>& InCADKernelEntity, const FImportParameters& ImportParameters, const FMeshParameters& MeshParameters, FMeshDescription& OutMeshDescription);
+		static bool Tessellate(CADKernel::FTopologicalShapeEntity& InCADKernelEntity, const FImportParameters& ImportParameters, const FMeshParameters& MeshParameters, FMeshDescription& OutMeshDescription);
 
-		static uint32 GetFaceTessellation(const TSharedRef<CADKernel::FFaceMesh>& FaceMesh, FBodyMesh& OutBodyMesh);
+		static uint32 GetFaceTessellation(CADKernel::FFaceMesh& FaceMesh, FBodyMesh& OutBodyMesh);
 	};
 }
diff --git a/Engine/Plugins/Enterprise/DatasmithCADImporter/Source/DatasmithOpenNurbsTranslator/Private/OpenNurbsBRepToCADKernelConverter.cpp b/Engine/Plugins/Enterprise/DatasmithCADImporter/Source/DatasmithOpenNurbsTranslator/Private/OpenNurbsBRepToCADKernelConverter.cpp
--- a/Engine/Plugins/Enterprise/DatasmithCADImporter/Source/DatasmithOpenNurbsTranslator/Private/OpenNurbsBRepToCADKernelConverter.cpp
+++ b/Engine/Plugins/Enterprise/DatasmithCADImporter/Source/DatasmithOpenNurbsTranslator/Private/OpenNurbsBRepToCADKernelConverter.cpp
@@ -327,7 +327,7 @@ bool FOpenNurbsBRepToCADKernelConverter::AddBRep(ON_Brep& BRep, const ON_3dVecto
 
 	BRep.Translate(-Offset);
 
-	CADKernelSession.GetModel()->Add(Body);
+	CADKernelSession.GetModel().Add(Body);
 
 	return true;
 }
diff --git a/Engine/Plugins/Enterprise/DatasmithCADImporter/Source/DatasmithWireTranslator/Private/AliasModelToCADKernelConverter.cpp b/Engine/Plugins/Enterprise/DatasmithCADImporter/Source/DatasmithWireTranslator/Private/AliasModelToCADKernelConverter.cpp
--- a/Engine/Plugins/Enterprise/DatasmithCADImporter/Source/DatasmithWireTranslator/Private/AliasModelToCADKernelConverter.cpp
+++ b/Engine/Plugins/Enterprise/DatasmithCADImporter/Source/DatasmithWireTranslator/Private/AliasModelToCADKernelConverter.cpp
@@ -323,7 +323,7 @@ bool FAliasModelToCADKernelConverter::AddBRep(AlDagNode& DagNode, EAliasObjectRe
 	}
 
 	// Create body from faces
-	CADKernelSession.GetModel()->Add(CADKernelBody);
+	CADKernelSession.GetModel().Add(CADKernelBody);
 	return true;
 }
 
diff --git a/Engine/Source/Runtime/Datasmith/CADKernel/Private/CADKernel/Core/Database.cpp b/Engine/Source/Runtime/Datasmith/CADKernel/Private/CADKernel/Core/Database.cpp
--- a/Engine/Source/Runtime/Datasmith/CADKernel/Private/CADKernel/Core/Database.cpp
+++ b/Engine/Source/Runtime/Datasmith/CADKernel/Private/CADKernel/Core/Database.cpp
@@ -32,14 +32,14 @@ FDatabase::FDatabase()
 	AddEntity(Model.ToSharedRef());
 }
 
-TSharedRef<FModel> FDatabase::GetModel()
+FModel& FDatabase::GetModel()
 {
 	if (!Model.IsValid())
 	{
 		Model = FEntity::MakeShared<FModel>();
 		AddEntity(Model.ToSharedRef());
 	}
-	return Model.ToSharedRef();
+	return *Model;
 }
 
 FIdent FDatabase::CreateId()
@@ -138,6 +138,80 @@ void FDatabase::GetEntities(const TArray<FIdent>& EntityIds, TArray<TSharedPtr<F
 	}
 }
 
+void FDatabase::GetTopologicalEntities(const TArray<FIdent>& EntityIds, TArray<FTopologicalEntity*>& Entities) const
+{
+	Entities.Empty(EntityIds.Num());
+
+	// to avoid duplicated components
+	TSet<FIdent> AddedComponents;
+	AddedComponents.Reserve(EntityIds.Num());
+
+	for (FIdent EntityId : EntityIds)
+	{
+		TSharedPtr<FEntity> Entity = GetEntity(EntityId);
+		if (!Entity.IsValid() || Entity->IsDeleted())
+		{
+			continue;
+		}
+
+		switch (Entity->GetEntityType())
+		{
+		case EEntity::EdgeLink:
+		case EEntity::VertexLink:
+		case EEntity::TopologicalEdge:
+		case EEntity::TopologicalFace:
+		case EEntity::TopologicalLink:
+		case EEntity::TopologicalLoop:
+		case EEntity::TopologicalVertex:
+		case EEntity::Shell:
+		case EEntity::Body:
+		case EEntity::Model:
+			if (AddedComponents.Find(EntityId) == nullptr)
+			{
+				AddedComponents.Add(EntityId);
+				Entities.Add((FTopologicalShapeEntity*)Entity.Get());
+			}
+			break;
+		default:
+			break;
+		}
+	}
+}
+
+void FDatabase::GetTopologicalShapeEntities(const TArray<FIdent>& EntityIds, TArray<FTopologicalShapeEntity*>& Entities) const
+{
+	Entities.Empty(EntityIds.Num());
+
+	// to avoid duplicated components
+	TSet<FIdent> AddedComponents;
+	AddedComponents.Reserve(EntityIds.Num());
+
+	for (FIdent EntityId : EntityIds)
+	{
+		TSharedPtr<FEntity> Entity = GetEntity(EntityId);
+		if (!Entity.IsValid() || Entity->IsDeleted())
+		{
+			continue;
+		}
+
+		switch (Entity->GetEntityType())
+		{
+		case EEntity::Body:
+		case EEntity::Model:
+		case EEntity::Shell:
+		case EEntity::TopologicalFace:
+			if (AddedComponents.Find(EntityId) == nullptr)
+			{
+				AddedComponents.Add(EntityId);
+				Entities.Add((FTopologicalShapeEntity*) Entity.Get());
+			}
+			break;
+		default:
+			break;
+		}
+	}
+}
+
 void FCADKernelArchive::SetReferencedEntityOrAddToWaitingList(FIdent ArchiveId, FEntity** Entity)
 {
 	Session.Database.SetReferencedEntityOrAddToWaitingList(ArchiveId, Entity);
@@ -240,7 +314,7 @@ void FDatabase::Serialize(FCADKernelArchive& Ar)
 
 	ensure(Ar.IsSaving());
 
-	SpawnEntityIdent((TSharedPtr<FEntity>&) Model, true);
+	SpawnEntityIdent(*Model, true);
 
 	Ar.Session.Serialize(Ar);
 
@@ -320,7 +394,7 @@ void FDatabase::CleanArchiveEntities()
 			{
 				TSharedPtr<FEdgeLink> EdgeLink = StaticCastSharedPtr<FEdgeLink>(Entity);
 				EdgeLink->CleanLink();
-				ensureCADKernel(EdgeLink->GetTwinsEntitieNum());
+				ensureCADKernel(EdgeLink->GetTwinEntityNum());
 				break;
 			}
 			case EEntity::VertexLink:
@@ -328,7 +402,7 @@ void FDatabase::CleanArchiveEntities()
 				TSharedPtr<FVertexLink> VertexLink = StaticCastSharedPtr<FVertexLink>(Entity);
 				VertexLink->CleanLink();
 #ifdef CADKERNEL_DEV
- 				ensureCADKernel(VertexLink->GetTwinsEntitieNum());
+ 				ensureCADKernel(VertexLink->GetTwinEntityNum());
 #endif
 				break;
 			}
@@ -422,7 +496,7 @@ void FDatabase::Empty()
 	AvailableIdents.Empty(DataBaseInitialSize);
 }
 
-uint32 FDatabase::SpawnEntityIdent(const TArray<TSharedPtr<FEntity>>& SelectedEntities, bool bInForceSpawning)
+uint32 FDatabase::SpawnEntityIdents(const TArray<TSharedPtr<FEntity>>& SelectedEntities, bool bInForceSpawning)
 {
 	bForceSpawning = bInForceSpawning;
 	EntityCount = 0;
@@ -439,16 +513,33 @@ uint32 FDatabase::SpawnEntityIdent(const TArray<TSharedPtr<FEntity>>& SelectedEn
 	return EntityCount;
 }
 
-uint32 FDatabase::SpawnEntityIdent(TSharedPtr<FEntity>& Entity, bool bInForceSpawning)
+uint32 FDatabase::SpawnEntityIdents(const TArray<FEntity*>& SelectedEntities, bool bInForceSpawning)
 {
-	if (!Entity.IsValid() || Entity->IsDeleted())
+	bForceSpawning = bInForceSpawning;
+	EntityCount = 0;
+	for (FEntity* Entity : SelectedEntities)
+	{
+		if (Entity == nullptr || Entity->IsDeleted())
+		{
+			continue;
+		}
+
+		Entity->SpawnIdent(*this);
+	}
+	bForceSpawning = false;
+	return EntityCount;
+}
+
+uint32 FDatabase::SpawnEntityIdent(FEntity& Entity, bool bInForceSpawning)
+{
+	if (Entity.IsDeleted())
 	{
 		return 0;
 	}
 
 	bForceSpawning = bInForceSpawning;
 	EntityCount = 0;
-	Entity->SpawnIdent(*this);
+	Entity.SpawnIdent(*this);
 	bForceSpawning = false;
 	return EntityCount;
 }
diff --git a/Engine/Source/Runtime/Datasmith/CADKernel/Private/CADKernel/Core/Session.cpp b/Engine/Source/Runtime/Datasmith/CADKernel/Private/CADKernel/Core/Session.cpp
--- a/Engine/Source/Runtime/Datasmith/CADKernel/Private/CADKernel/Core/Session.cpp
+++ b/Engine/Source/Runtime/Datasmith/CADKernel/Private/CADKernel/Core/Session.cpp
@@ -23,18 +23,17 @@ void FSession::SaveDatabase(const TCHAR* FileName)
 	Archive->Close();
 }
 
-TSharedRef<FModel> FSession::GetModel()
+FModel& FSession::GetModel()
 {
 	return Database.GetModel();
 }
 
-
 void FSession::SaveDatabase(const TCHAR* FileName, const TArray<TSharedPtr<FEntity>>& SelectedEntities)
 {
 	TArray<FIdent> EntityIds;
 	EntityIds.Reserve(SelectedEntities.Num());
 
-	SpawnEntityIdent(SelectedEntities, true);
+	SpawnEntityIdents(SelectedEntities, true);
 
 	for (const TSharedPtr<FEntity>& Entity : SelectedEntities)
 	{
@@ -56,7 +55,7 @@ void FSession::LoadDatabase(const TCHAR* FilePath)
 		return;
 	}
 
-	TSharedRef<FModel> SessionModel = GetModel();
+	FModel& Model = GetModel();
 
 	Database.Deserialize(*Archive.Get());
 	FModel* ArchiveModel = Archive->ArchiveModel;
@@ -77,7 +76,7 @@ void FSession::AddDatabase(const TArray<uint8>& InRawData)
 
 void FSession::SetGeometricTolerance(double NewTolerance)
 {
-	ensure(Database.GetModel()->EntityCount() == 0);
+	ensure(Database.GetModel().EntityCount() == 0);
 	GeometricTolerance = NewTolerance;
 }
 } // namespace CADKernel
diff --git a/Engine/Source/Runtime/Datasmith/CADKernel/Private/CADKernel/Core/Version.h b/Engine/Source/Runtime/Datasmith/CADKernel/Private/CADKernel/Core/Version.h
--- a/Engine/Source/Runtime/Datasmith/CADKernel/Private/CADKernel/Core/Version.h
+++ b/Engine/Source/Runtime/Datasmith/CADKernel/Private/CADKernel/Core/Version.h
@@ -2,8 +2,8 @@
 #pragma once
 
 //** ** Version du Toolkit** **//
-#define TOOLKIT_VERSION 5, 0, 18687460
-#define TOOLKIT_VERSION_ASCII "5.0.18687460"
+#define TOOLKIT_VERSION 5, 0, 18784723
+#define TOOLKIT_VERSION_ASCII "5.0.18784723"
 
 #define RELEASE_DATE	 __DATE__
 
diff --git a/Engine/Source/Runtime/Datasmith/CADKernel/Private/CADKernel/Geo/Sampler/SamplerAbstract.h b/Engine/Source/Runtime/Datasmith/CADKernel/Private/CADKernel/Geo/Sampler/SamplerAbstract.h
--- a/Engine/Source/Runtime/Datasmith/CADKernel/Private/CADKernel/Geo/Sampler/SamplerAbstract.h
+++ b/Engine/Source/Runtime/Datasmith/CADKernel/Private/CADKernel/Geo/Sampler/SamplerAbstract.h
@@ -13,221 +13,222 @@
 
 namespace CADKernel
 {
-	template<typename PolylineType, typename PointType>
-	class TCurveSamplerAbstract
-	{
-	public:
-
-		struct TSamplingPoint
-		{
-			const PolylineType* Sampling;
-			const PolylineType* Polyline;
-			int32 Index;
-
-			TSamplingPoint(const PolylineType& InSampling)
-				: Sampling(&InSampling)
-			{
-			}
 
-			void Set(const PolylineType& InPolyline, int32 InIndex)
-			{
-				Polyline = &InPolyline;
-				Index = InIndex;
-			}
+template<typename PolylineType, typename PointType>
+class TCurveSamplerAbstract
+{
+public:
 
-			constexpr bool IsCandidatePoint()
-			{
-				return Sampling != Polyline;
-			}
+	struct TSamplingPoint
+	{
+		const PolylineType* Sampling;
+		const PolylineType* Polyline;
+		int32 Index;
 
-			constexpr bool IsSamplingPoint()
-			{
-				return Sampling == Polyline;
-			}
-		};
-
-		TCurveSamplerAbstract(const FLinearBoundary& InBoundary, PolylineType& OutPolyline, double InDesiredChordError)
-			: Boundary(InBoundary)
-			, Sampling(OutPolyline)
-			, TmpPolylineCoordinates(CandidatePoints.GetCoordinates())
-			, SamplingCoordinates(Sampling.GetCoordinates())
-			, EndStudySegment(OutPolyline)
-			, DesiredChordError(InDesiredChordError)
+		TSamplingPoint(const PolylineType& InSampling)
+			: Sampling(&InSampling)
 		{
 		}
 
-		virtual ~TCurveSamplerAbstract() = default;
-
-		/**
-		 * Method to call to generate the curve sampling
-		 */
-		virtual void Sample()
+		void Set(const PolylineType& InPolyline, int32 InIndex)
 		{
-			SamplingInitalizing();
-			RunSampling();
+			Polyline = &InPolyline;
+			Index = InIndex;
 		}
 
-	protected:
-	
-		/** Evaluate the new candidate points corresponding to the NextCoordinates */
-		virtual void EvaluatesNewCandidatePoints() = 0;
-
-		/**
-		 * This function check the sub-polyline to verify if its respect the attended quality
-		 * e.g. Curve sampler implementation of this method check the chord error
-		 */
-		virtual int32 CheckSamplingError(int32 FirstIndex, int32 EndIndex) = 0;
-
-		/**
-		 * OutNotDerivableCoordinates must not include the boundaries of the curve. They are added after 
-		 */
-		virtual void GetNotDerivableCoordinates(TArray<double>& OutNotDerivableCoordinates)
+		constexpr bool IsCandidatePoint()
 		{
+			return Sampling != Polyline;
 		}
 
-		/**
-		 * Method to call to generate the curve sampling
-		 */
-		virtual void SamplingInitalizing()
+		constexpr bool IsSamplingPoint()
 		{
+			return Sampling == Polyline;
+		}
+	};
+
+	TCurveSamplerAbstract(const FLinearBoundary& InBoundary, PolylineType& OutPolyline, double InDesiredChordError)
+		: Boundary(InBoundary)
+		, Sampling(OutPolyline)
+		, TmpPolylineCoordinates(CandidatePoints.GetCoordinates())
+		, SamplingCoordinates(Sampling.GetCoordinates())
+		, EndStudySegment(OutPolyline)
+		, DesiredChordError(InDesiredChordError)
+	{
+	}
+
+	virtual ~TCurveSamplerAbstract() = default;
+
+	/**
+	 * Method to call to generate the curve sampling
+	 */
+	virtual void Sample()
+	{
+		SamplingInitalizing();
+		RunSampling();
+	}
+
+protected:
+
+	/** Evaluate the new candidate points corresponding to the NextCoordinates */
+	virtual void EvaluatesNewCandidatePoints() = 0;
+
+	/**
+	 * This function check the sub-polyline to verify if its respect the attended quality
+	 * e.g. Curve sampler implementation of this method check the chord error
+	 */
+	virtual int32 CheckSamplingError(int32 FirstIndex, int32 EndIndex) = 0;
+
+	/**
+	 * OutNotDerivableCoordinates must not include the boundaries of the curve. They are added after
+	 */
+	virtual void GetNotDerivableCoordinates(TArray<double>& OutNotDerivableCoordinates)
+	{
+	}
+
+	/**
+	 * Method to call to generate the curve sampling
+	 */
+	virtual void SamplingInitalizing()
+	{
 #ifdef DEBUG_CURVE_SAMPLING
-			++CurveIndex;
-			CurveToDisplay = CurveIndex;
+		++CurveIndex;
+		CurveToDisplay = CurveIndex;
 #endif
-			IsOptimalSegments.Empty(100);
-			Sampling.Empty(100);
+		IsOptimalSegments.Empty(100);
+		Sampling.Empty(100);
 
-			CandidatePoints.Empty(100);
+		CandidatePoints.Empty(100);
 
-			// Initialization of the algorithm with the not derivable coordinates of the curve and at least 5 temporary points
-			int32 ComplementatyPointCount = 0;
-			{
-				TArray<double> NotDerivableCoordinates;
-				GetNotDerivableCoordinates(NotDerivableCoordinates);
+		// Initialization of the algorithm with the not derivable coordinates of the curve and at least 5 temporary points
+		int32 ComplementatyPointCount = 0;
+		{
+			TArray<double> NotDerivableCoordinates;
+			GetNotDerivableCoordinates(NotDerivableCoordinates);
 
-				NotDerivableCoordinates.Insert(Boundary.Min, 0);
-				NotDerivableCoordinates.Add(Boundary.Max);
+			NotDerivableCoordinates.Insert(Boundary.Min, 0);
+			NotDerivableCoordinates.Add(Boundary.Max);
 
-				ComplementatyPointCount = NotDerivableCoordinates.Num() < 10 ? 10 : 1;
+			ComplementatyPointCount = NotDerivableCoordinates.Num() < 10 ? 10 : 1;
 
-				NextCoordinates.Empty(NotDerivableCoordinates.Num() * (ComplementatyPointCount + 1));
-				NextCoordinates.Add(Boundary.Min);
+			NextCoordinates.Empty(NotDerivableCoordinates.Num() * (ComplementatyPointCount + 1));
+			NextCoordinates.Add(Boundary.Min);
 
-				for (int32 Index = 0; Index < NotDerivableCoordinates.Num() - 1; ++Index)
-				{
-					AddIntermediateCoordinates(NotDerivableCoordinates[Index], NotDerivableCoordinates[Index + 1], ComplementatyPointCount);
-					NextCoordinates.Add(NotDerivableCoordinates[Index + 1]);
-				}
+			for (int32 Index = 0; Index < NotDerivableCoordinates.Num() - 1; ++Index)
+			{
+				AddIntermediateCoordinates(NotDerivableCoordinates[Index], NotDerivableCoordinates[Index + 1], ComplementatyPointCount);
+				NextCoordinates.Add(NotDerivableCoordinates[Index + 1]);
 			}
+		}
 
-			CandidatePoints.SwapCoordinates(NextCoordinates);
-			EvaluatesNewCandidatePoints();
+		CandidatePoints.SwapCoordinates(NextCoordinates);
+		EvaluatesNewCandidatePoints();
 
-			// Not Derivable point initialize the sampling
-			ComplementatyPointCount++;
-			for (int32 Index = 0, ISampling = 0; Index < CandidatePoints.Size(); Index += ComplementatyPointCount, ++ISampling)
-			{
-				Sampling.EmplaceAt(ISampling, CandidatePoints, Index);
-				IsOptimalSegments.Add(false);
-			}
-			IsOptimalSegments.Pop();
+		// Not Derivable point initialize the sampling
+		ComplementatyPointCount++;
+		for (int32 Index = 0, ISampling = 0; Index < CandidatePoints.Size(); Index += ComplementatyPointCount, ++ISampling)
+		{
+			Sampling.EmplaceAt(ISampling, CandidatePoints, Index);
+			IsOptimalSegments.Add(false);
+		}
+		IsOptimalSegments.Pop();
 
-			for (int32 Index = CandidatePoints.Size() - 1; Index >= 0; Index -= ComplementatyPointCount)
-			{
-				CandidatePoints.RemoveAt(Index);
-			}
+		for (int32 Index = CandidatePoints.Size() - 1; Index >= 0; Index -= ComplementatyPointCount)
+		{
+			CandidatePoints.RemoveAt(Index);
+		}
 
-			// first segment is not optimal
+		// first segment is not optimal
 
 #ifdef DEBUG_CURVE_SAMPLING
-			int32 Step = 0;
-			DisplaySampling(CurveIndex == CurveToDisplay, Step);
+		int32 Step = 0;
+		DisplaySampling(CurveIndex == CurveToDisplay, Step);
 #endif
-		}
+	}
 
-		/**
-		 * Method to call to generate the curve sampling
-		 */
-		void RunSampling()
+	/**
+	 * Method to call to generate the curve sampling
+	 */
+	void RunSampling()
+	{
+
+		// < 100000 : check to avoid that the process loops endlessly
+		int32 CandidatePointsCount = CandidatePoints.Size();
+		while (CandidatePointsCount && Sampling.Coordinates.Num() < 100000)
 		{
+			NextCoordinates.Empty();
 
-			// < 100000 : check to avoid that the process loops endlessly
-			int32 CandidatePointsCount = CandidatePoints.Size();
-			while (CandidatePointsCount && Sampling.Coordinates.Num() < 100000)
-			{
-				NextCoordinates.Empty();
+			StartSamplingSegmentIndex = 0;
+			EndStudySegment.Index = 0;
 
-				StartSamplingSegmentIndex = 0;
-				EndStudySegment.Index = 0;
+			// indices of TmpPoints and TmpCoordinates
+			int32 FirstCandidateIndex = 0;
+			int32 LastCandidateIndex = 1;
 
-				// indices of TmpPoints and TmpCoordinates
-				int32 FirstCandidateIndex = 0;
-				int32 LastCandidateIndex = 1;
+			int32 SegmentIndex = 0;
+
+			do
+			{
+				if (IsOptimalSegments[StartSamplingSegmentIndex])
+				{
+					StartSamplingSegmentIndex++;
+					continue;
+				}
 
-				int32 SegmentIndex = 0;
+				EndSamplingSegmentIndex = StartSamplingSegmentIndex + 1;
 
 				do
 				{
-					if (IsOptimalSegments[StartSamplingSegmentIndex])
+					if ((LastCandidateIndex < CandidatePointsCount && TmpPolylineCoordinates[LastCandidateIndex] < Sampling.Coordinates[EndSamplingSegmentIndex]))
 					{
-						StartSamplingSegmentIndex++;
-						continue;
+						EndStudySegment.Set(CandidatePoints, LastCandidateIndex);
 					}
-
-					EndSamplingSegmentIndex = StartSamplingSegmentIndex + 1;
-
-					do
+					else
 					{
-						if ((LastCandidateIndex < CandidatePointsCount && TmpPolylineCoordinates[LastCandidateIndex] < Sampling.Coordinates[EndSamplingSegmentIndex]))
-						{
-							EndStudySegment.Set(CandidatePoints, LastCandidateIndex);
-						}
-						else
-						{
-							EndStudySegment.Set(Sampling, EndSamplingSegmentIndex);
-						}
-
-						// if the next candidate point is biggest than the next sampling point, no chord check can be done. 
-						// an intermediate point is added. 
-						// the next segment is processed 
-						if ((FirstCandidateIndex >= CandidatePointsCount) || (FirstCandidateIndex < CandidatePointsCount && TmpPolylineCoordinates[FirstCandidateIndex] > EndStudySegment.Polyline->Coordinates[EndStudySegment.Index]))
-						{
-							AddIntermediateCoordinates(Sampling.Coordinates[StartSamplingSegmentIndex], Sampling.Coordinates[EndSamplingSegmentIndex], 1);
-							StartSamplingSegmentIndex = EndSamplingSegmentIndex;
-							break;
-						}
+						EndStudySegment.Set(Sampling, EndSamplingSegmentIndex);
+					}
 
-						int32 NeededPointNum = CheckSamplingError(FirstCandidateIndex, LastCandidateIndex);
-						if (NeededPointNum == 0 || (NeededPointNum == 1 && ((LastCandidateIndex + 1) == CandidatePointsCount || (((LastCandidateIndex + 1) < CandidatePointsCount) && TmpPolylineCoordinates[LastCandidateIndex + 1] > Sampling.Coordinates[EndSamplingSegmentIndex]))))
-						{
-							if (EndStudySegment.IsSamplingPoint())
-							{
-								IsOptimalSegments[StartSamplingSegmentIndex] = true;
-							}
-							else
-							{
-								LastCandidateIndex++;
-								continue;
-							}
-						}
-						else
-						{
-							CompletesPolyline(NeededPointNum, LastCandidateIndex - 1);
-						}
+					// if the next candidate point is biggest than the next sampling point, no chord check can be done. 
+					// an intermediate point is added. 
+					// the next segment is processed 
+					if ((FirstCandidateIndex >= CandidatePointsCount) || (FirstCandidateIndex < CandidatePointsCount && TmpPolylineCoordinates[FirstCandidateIndex] > EndStudySegment.Polyline->Coordinates[EndStudySegment.Index]))
+					{
+						AddIntermediateCoordinates(Sampling.Coordinates[StartSamplingSegmentIndex], Sampling.Coordinates[EndSamplingSegmentIndex], 1);
+						StartSamplingSegmentIndex = EndSamplingSegmentIndex;
+						break;
+					}
 
-						// An existing point is reached 
+					int32 NeededPointNum = CheckSamplingError(FirstCandidateIndex, LastCandidateIndex);
+					if (NeededPointNum == 0 || (NeededPointNum == 1 && ((LastCandidateIndex + 1) == CandidatePointsCount || (((LastCandidateIndex + 1) < CandidatePointsCount) && TmpPolylineCoordinates[LastCandidateIndex + 1] > Sampling.Coordinates[EndSamplingSegmentIndex]))))
+					{
 						if (EndStudySegment.IsSamplingPoint())
 						{
-							StartSamplingSegmentIndex = EndSamplingSegmentIndex;
-							FirstCandidateIndex = LastCandidateIndex;
-							LastCandidateIndex++;
+							IsOptimalSegments[StartSamplingSegmentIndex] = true;
 						}
 						else
 						{
-							StartSamplingSegmentIndex = EndSamplingSegmentIndex - 1;
-							FirstCandidateIndex = LastCandidateIndex + 1;
-							LastCandidateIndex += 2;
+							LastCandidateIndex++;
+							continue;
 						}
+					}
+					else
+					{
+						CompletesPolyline(NeededPointNum, LastCandidateIndex - 1);
+					}
+
+					// An existing point is reached 
+					if (EndStudySegment.IsSamplingPoint())
+					{
+						StartSamplingSegmentIndex = EndSamplingSegmentIndex;
+						FirstCandidateIndex = LastCandidateIndex;
+						LastCandidateIndex++;
+					}
+					else
+					{
+						StartSamplingSegmentIndex = EndSamplingSegmentIndex - 1;
+						FirstCandidateIndex = LastCandidateIndex + 1;
+						LastCandidateIndex += 2;
+					}
 
 				} while (EndStudySegment.IsCandidatePoint());
 
@@ -244,211 +245,211 @@ namespace CADKernel
 		}
 
 #ifdef CHECK_RESULT_MAKE_ISO_POLYLINE
-			CheckResult();
+		CheckResult();
 #endif
 
-		}
+	}
 
-		/** Adds coordinates of the next candidate points in NextCoordinates array */
-		void AddIntermediateCoordinates(double UMin, double UMax, int32 PointNum)
-		{
-			double DeltaCoord = (UMax - UMin) / (PointNum + 1);
-			double UCoord = UMin;
-			for (int32 Index = 0; Index < PointNum; Index++)
-			{
-				UCoord += DeltaCoord;
-				NextCoordinates.Add(UCoord);
-			}
-		};
-
-		/**
-		 * Completes Polyline and add new candidate coordinates for the next step.
-		 * @param NeededPointCount estimated number of needed points to respect the sampling error.
-		 * NeededPointNumber = 1: only the last point is needed
-		 * NeededPointNumber = 2: the intermediate point and last point is needed
-		 * NeededPointNumber > 2: the both points are not enought. Candidate points are added on each side of the intermediate point
-		 * @param IntermediateIndex
-		 */
-		void CompletesPolyline(int32 NeededPointNumber, int32 IntermediateIndex)
+	/** Adds coordinates of the next candidate points in NextCoordinates array */
+	void AddIntermediateCoordinates(double UMin, double UMax, int32 PointNum)
+	{
+		double DeltaCoord = (UMax - UMin) / (PointNum + 1);
+		double UCoord = UMin;
+		for (int32 Index = 0; Index < PointNum; Index++)
 		{
-			const TArray<double>& CandidateCoordinates = CandidatePoints.GetCoordinates();
-
-			if (NeededPointNumber > 1)
-			{
-				Sampling.EmplaceAt(EndSamplingSegmentIndex, CandidatePoints, IntermediateIndex);
-				IsOptimalSegments.EmplaceAt(EndSamplingSegmentIndex - 1, NeededPointNumber == 2 ? true : false);
-				EndSamplingSegmentIndex++;
-			}
+			UCoord += DeltaCoord;
+			NextCoordinates.Add(UCoord);
+		}
+	};
 
-			if (EndStudySegment.IsCandidatePoint())
-			{
-				Sampling.EmplaceAt(EndSamplingSegmentIndex, CandidatePoints, EndStudySegment.Index);
-				IsOptimalSegments.EmplaceAt(EndSamplingSegmentIndex - 1, NeededPointNumber <= 2 ? true : false);
-				EndSamplingSegmentIndex++;
-			}
-			else
-			{
-				EndStudySegment.Index = EndSamplingSegmentIndex;
-				if (NeededPointNumber <= 2)
-				{
-					IsOptimalSegments[EndSamplingSegmentIndex - 1] = true;
-				}
-			}
+	/**
+	 * Completes Polyline and add new candidate coordinates for the next step.
+	 * @param NeededPointCount estimated number of needed points to respect the sampling error.
+	 * NeededPointNumber = 1: only the last point is needed
+	 * NeededPointNumber = 2: the intermediate point and last point is needed
+	 * NeededPointNumber > 2: the both points are not enought. Candidate points are added on each side of the intermediate point
+	 * @param IntermediateIndex
+	 */
+	void CompletesPolyline(int32 NeededPointNumber, int32 IntermediateIndex)
+	{
+		const TArray<double>& CandidateCoordinates = CandidatePoints.GetCoordinates();
 
-			if (NeededPointNumber > 2)
-			{
-				NeededPointNumber -= 2;
-				AddIntermediateCoordinates(Sampling.Coordinates[StartSamplingSegmentIndex], CandidateCoordinates[IntermediateIndex], NeededPointNumber);
-  				AddIntermediateCoordinates(CandidateCoordinates[IntermediateIndex], EndStudySegment.Polyline->Coordinates[EndStudySegment.Index], NeededPointNumber);
-			}
-		};
+		if (NeededPointNumber > 1)
+		{
+			Sampling.EmplaceAt(EndSamplingSegmentIndex, CandidatePoints, IntermediateIndex);
+			IsOptimalSegments.EmplaceAt(EndSamplingSegmentIndex - 1, NeededPointNumber == 2 ? true : false);
+			EndSamplingSegmentIndex++;
+		}
 
-		int32 GetFirstNeighbor(int32 NeighborIndex, const double StartCoordinate, const PolylineType& Points, const int32 Increment, double& NeighborCoordinate)
+		if (EndStudySegment.IsCandidatePoint())
 		{
-			NeighborCoordinate = Points.Coordinates[NeighborIndex];
-			if (FMath::IsNearlyEqual(NeighborCoordinate, StartCoordinate) && Points.Coordinates.IsValidIndex(NeighborIndex + Increment))
+			Sampling.EmplaceAt(EndSamplingSegmentIndex, CandidatePoints, EndStudySegment.Index);
+			IsOptimalSegments.EmplaceAt(EndSamplingSegmentIndex - 1, NeededPointNumber <= 2 ? true : false);
+			EndSamplingSegmentIndex++;
+		}
+		else
+		{
+			EndStudySegment.Index = EndSamplingSegmentIndex;
+			if (NeededPointNumber <= 2)
 			{
-				NeighborIndex += Increment;
-				NeighborCoordinate = Points.Coordinates[NeighborIndex];;
+				IsOptimalSegments[EndSamplingSegmentIndex - 1] = true;
 			}
-			return NeighborIndex;
-		};
+		}
 
-		int32 CountOfNeededPointsToRespectChordError(const PointType& PointA, const PointType& PointB, double ChordError)
+		if (NeededPointNumber > 2)
 		{
-			double ABLength = PointA.Distance(PointB) / 2;
-			// Max segment length to respect the chord error
-			double LengthToHave = sqrt(FMath::Square(ABLength) * DesiredChordError / ChordError);
-			return (int32)(2 + (ABLength / LengthToHave + 0.5));
-		};
+			NeededPointNumber -= 2;
+			AddIntermediateCoordinates(Sampling.Coordinates[StartSamplingSegmentIndex], CandidateCoordinates[IntermediateIndex], NeededPointNumber);
+			AddIntermediateCoordinates(CandidateCoordinates[IntermediateIndex], EndStudySegment.Polyline->Coordinates[EndStudySegment.Index], NeededPointNumber);
+		}
+	};
 
-		int32 CheckTangentError(const PointType& APoint, double ACoordinate, const PointType& BPoint, double BCoordinate, int32 FirstIndex, int32 EndIndex, int32 InStartSamplingSegmentIndex)
+	int32 GetFirstNeighbor(int32 NeighborIndex, const double StartCoordinate, const PolylineType& Points, const int32 Increment, double& NeighborCoordinate)
+	{
+		NeighborCoordinate = Points.Coordinates[NeighborIndex];
+		if (FMath::IsNearlyEqual(NeighborCoordinate, StartCoordinate) && Points.Coordinates.IsValidIndex(NeighborIndex + Increment))
 		{
-			PointType Middle = CandidatePoints.GetPoints()[FirstIndex];
+			NeighborIndex += Increment;
+			NeighborCoordinate = Points.Coordinates[NeighborIndex];;
+		}
+		return NeighborIndex;
+	};
 
-			int32 CountOfNeededPoints = 2;
+	int32 CountOfNeededPointsToRespectChordError(const PointType& PointA, const PointType& PointB, double ChordError)
+	{
+		double ABLength = PointA.Distance(PointB) / 2;
+		// Max segment length to respect the chord error
+		double LengthToHave = sqrt(FMath::Square(ABLength) * DesiredChordError / ChordError);
+		return (int32)(2 + (ABLength / LengthToHave + 0.5));
+	};
 
-			if (FirstIndex > 0 && InStartSamplingSegmentIndex > 0)
-			{
-				// get the first previous point in Sampling or CandidatePoints
-				double PreviousSampingCoordinate;
-				int32 PreviousSampingIndex = GetFirstNeighbor(InStartSamplingSegmentIndex - 1, ACoordinate, Sampling, -1, PreviousSampingCoordinate);
+	int32 CheckTangentError(const PointType& APoint, double ACoordinate, const PointType& BPoint, double BCoordinate, int32 FirstIndex, int32 EndIndex, int32 InStartSamplingSegmentIndex)
+	{
+		PointType Middle = CandidatePoints.GetPoints()[FirstIndex];
 
-				double PreviousCandidateCoordinate;
-				int32 PreviousCandidateIndex = GetFirstNeighbor(FirstIndex - 1, ACoordinate, CandidatePoints, -1, PreviousCandidateCoordinate);
+		int32 CountOfNeededPoints = 2;
 
-				const PointType& PreviousPoint = PreviousSampingCoordinate > PreviousCandidateCoordinate ? Sampling.GetPoints()[PreviousSampingIndex] : CandidatePoints.GetPoints()[PreviousCandidateIndex];
+		if (FirstIndex > 0 && InStartSamplingSegmentIndex > 0)
+		{
+			// get the first previous point in Sampling or CandidatePoints
+			double PreviousSampingCoordinate;
+			int32 PreviousSampingIndex = GetFirstNeighbor(InStartSamplingSegmentIndex - 1, ACoordinate, Sampling, -1, PreviousSampingCoordinate);
 
-				double ChordA = DistanceOfPointToSegment(APoint, PreviousPoint, Middle);
-				if (ChordA > 4 * DesiredChordError)
-				{
-					CountOfNeededPoints = CountOfNeededPointsToRespectChordError(APoint, PreviousPoint, ChordA);
-				}
-			}
+			double PreviousCandidateCoordinate;
+			int32 PreviousCandidateIndex = GetFirstNeighbor(FirstIndex - 1, ACoordinate, CandidatePoints, -1, PreviousCandidateCoordinate);
+
+			const PointType& PreviousPoint = PreviousSampingCoordinate > PreviousCandidateCoordinate ? Sampling.GetPoints()[PreviousSampingIndex] : CandidatePoints.GetPoints()[PreviousCandidateIndex];
 
-			Middle = CandidatePoints.GetPoints()[EndIndex - 1];
-			if (EndIndex + 1 < CandidatePoints.Coordinates.Num())
+			double ChordA = DistanceOfPointToSegment(APoint, PreviousPoint, Middle);
+			if (ChordA > 4 * DesiredChordError)
 			{
-				double NextSampingCoordinate;
-				int32 NextSampingIndex = GetFirstNeighbor(StartSamplingSegmentIndex + 1, BCoordinate, Sampling, +1, NextSampingCoordinate);
+				CountOfNeededPoints = CountOfNeededPointsToRespectChordError(APoint, PreviousPoint, ChordA);
+			}
+		}
 
-				double NextCandidateCoordinate;
-				int32 NextCandidateIndex = GetFirstNeighbor(EndIndex, BCoordinate, CandidatePoints, +1, NextCandidateCoordinate);
+		Middle = CandidatePoints.GetPoints()[EndIndex - 1];
+		if (EndIndex + 1 < CandidatePoints.Coordinates.Num())
+		{
+			double NextSampingCoordinate;
+			int32 NextSampingIndex = GetFirstNeighbor(StartSamplingSegmentIndex + 1, BCoordinate, Sampling, +1, NextSampingCoordinate);
 
-				const PointType& NextPoint = NextSampingCoordinate < NextCandidateCoordinate ? Sampling.GetPoints()[NextSampingIndex] : CandidatePoints.GetPoints()[NextCandidateIndex];
+			double NextCandidateCoordinate;
+			int32 NextCandidateIndex = GetFirstNeighbor(EndIndex, BCoordinate, CandidatePoints, +1, NextCandidateCoordinate);
 
-				double ChordB = DistanceOfPointToSegment(BPoint, NextPoint, Middle);
-				if (ChordB > 4 * DesiredChordError)
-				{
-					CountOfNeededPoints = FMath::Max(CountOfNeededPointsToRespectChordError(BPoint, NextPoint, ChordB), CountOfNeededPoints);
-				}
+			const PointType& NextPoint = NextSampingCoordinate < NextCandidateCoordinate ? Sampling.GetPoints()[NextSampingIndex] : CandidatePoints.GetPoints()[NextCandidateIndex];
+
+			double ChordB = DistanceOfPointToSegment(BPoint, NextPoint, Middle);
+			if (ChordB > 4 * DesiredChordError)
+			{
+				CountOfNeededPoints = FMath::Max(CountOfNeededPointsToRespectChordError(BPoint, NextPoint, ChordB), CountOfNeededPoints);
 			}
+		}
 
-			return CountOfNeededPoints;
-		};
+		return CountOfNeededPoints;
+	};
 
 protected:
 
-		FLinearBoundary Boundary;
+	FLinearBoundary Boundary;
 
-		// Final sampling of the curve. This polyline will be completed at each iteration
-		PolylineType& Sampling;
+	// Final sampling of the curve. This polyline will be completed at each iteration
+	PolylineType& Sampling;
 
-		PolylineType CandidatePoints;
+	PolylineType CandidatePoints;
 
-		const TArray<double>& TmpPolylineCoordinates;
-		const TArray<double>& SamplingCoordinates;
+	const TArray<double>& TmpPolylineCoordinates;
+	const TArray<double>& SamplingCoordinates;
 
-		/** Array to indicated if the segment at Index defined by Sampling[Index] and Sampling[Index + 1] respects the desired criteria */
-		TArray<char> IsOptimalSegments;
+	/** Array to indicated if the segment at Index defined by Sampling[Index] and Sampling[Index + 1] respects the desired criteria */
+	TArray<char> IsOptimalSegments;
 
-		/**
-		 * StartSamplingPointIndex and EndSamplingPointIndex defined the extremity of the segment that is currently enriched with candidate points if needed
-		 */
-		int32 StartSamplingSegmentIndex;
+	/**
+	 * StartSamplingPointIndex and EndSamplingPointIndex defined the extremity of the segment that is currently enriched with candidate points if needed
+	 */
+	int32 StartSamplingSegmentIndex;
 
-		/**
-		 * EndSamplingPointIndex changes (+1) each time a candidate point is inserted in the sampling
-		 */
-		int32 EndSamplingSegmentIndex;
+	/**
+	 * EndSamplingPointIndex changes (+1) each time a candidate point is inserted in the sampling
+	 */
+	int32 EndSamplingSegmentIndex;
 
-		/**
-		 * The study segment starting at StartSamplingPoint and ending at EndStudySegment that is either a point of Sampling or a point of CandidatePoints
-		 * The sampling error is compute between the intermediate candidate points (candidate points localized between StartSamplingPoint and EndStudySegment) and this segment.
-		 * FirstCandidateIndex and (LastCandidateIndex - 1) are the index of the first and the last intermediate points
-		 */
-		TSamplingPoint EndStudySegment;
+	/**
+	 * The study segment starting at StartSamplingPoint and ending at EndStudySegment that is either a point of Sampling or a point of CandidatePoints
+	 * The sampling error is compute between the intermediate candidate points (candidate points localized between StartSamplingPoint and EndStudySegment) and this segment.
+	 * FirstCandidateIndex and (LastCandidateIndex - 1) are the index of the first and the last intermediate points
+	 */
+	TSamplingPoint EndStudySegment;
 
-		/** Coordinates of the new candidate points for the next iteration */
-		TArray<double> NextCoordinates;
+	/** Coordinates of the new candidate points for the next iteration */
+	TArray<double> NextCoordinates;
 
-		double DesiredChordError;
+	double DesiredChordError;
 
 #ifdef DEBUG_CURVE_SAMPLING
 public:
-		int32 CurveIndex = 0;
-		int32 CurveToDisplay = 1;
+	int32 CurveIndex = 0;
+	int32 CurveToDisplay = 1;
 
-		void DisplaySampling(bool bDisplay, int32 Step)
+	void DisplaySampling(bool bDisplay, int32 Step)
+	{
+		if (bDisplay)
 		{
-			if (bDisplay)
 			{
+				int32 Index = 0;
+				F3DDebugSession G(FString::Printf(TEXT("Curve %d"), Step));
+				for (const FPoint& Point : Sampling.GetPoints())
 				{
-					int32 Index = 0;
-					F3DDebugSession G(FString::Printf(TEXT("Curve %d"), Step));
-					for (const FPoint& Point : Sampling.GetPoints())
+					CADKernel::DisplayPoint(Point, EVisuProperty::BluePoint, Index++);
+				}
+
+				for (int32 Index = 0; Index < Sampling.GetPoints().Num() - 1; ++Index)
+				{
+					if (IsOptimalSegments[Index])
 					{
-						CADKernel::DisplayPoint(Point, EVisuProperty::BluePoint, Index++);
+						DisplaySegment(Sampling.GetPoints()[Index], Sampling.GetPoints()[Index + 1], Index, EVisuProperty::BlueCurve);
 					}
-
-					for (int32 Index = 0; Index < Sampling.GetPoints().Num() - 1; ++Index)
+					else
 					{
-						if (IsOptimalSegments[Index])
-						{
-							DisplaySegment(Sampling.GetPoints()[Index], Sampling.GetPoints()[Index + 1], Index, EVisuProperty::BlueCurve);
-						}
-						else
-						{
-							DisplaySegment(Sampling.GetPoints()[Index], Sampling.GetPoints()[Index + 1], Index, EVisuProperty::YellowCurve);
-						}
+						DisplaySegment(Sampling.GetPoints()[Index], Sampling.GetPoints()[Index + 1], Index, EVisuProperty::YellowCurve);
 					}
 				}
+			}
+			{
+				int32 Index = 0;
+				F3DDebugSession G(FString::Printf(TEXT("Next Point Tmp %d"), Step));
+				for (const FPoint& Point : CandidatePoints.GetPoints())
 				{
-					int32 Index = 0;
-					F3DDebugSession G(FString::Printf(TEXT("Next Point Tmp %d"), Step));
-					for (const FPoint& Point : CandidatePoints.GetPoints())
-					{
-						CADKernel::DisplayPoint(Point, EVisuProperty::YellowPoint, Index++);
-					}
+					CADKernel::DisplayPoint(Point, EVisuProperty::YellowPoint, Index++);
 				}
-				Wait();
 			}
+			Wait();
 		}
+	}
 #endif
 
 #ifdef CHECK_RESULT_MAKE_ISO_POLYLINE
-		void CheckResult()
-		{
+	void CheckResult()
+	{
 		// Check the result
 		NextCoordinates.Empty(Sampling.Coordinates.Num());
 
@@ -495,9 +496,8 @@ public:
 
 		FMessage::Printf(EVerboseLevel::Log, TEXT("Chord Desired Med Sd Min Max %f %f %f %f %f"), MaxError, ChordMed, ChordStandartDeviation, MinError[0], MaxError[0]);
 		FMessage::Printf(EVerboseLevel::Log, TEXT("Param Desired Med Sd Min Max %f %f %f %f %f"), MaxError, ParamMed, ParamStandartDeviation, MinError[1], MaxError[1]);
-		}
+	}
 #endif
+};
 
-
-	};
-}
\ No newline at end of file
+} // ns CADKernel
\ No newline at end of file
diff --git a/Engine/Source/Runtime/Datasmith/CADKernel/Private/CADKernel/Geo/Sampler/SamplerOnChord.h b/Engine/Source/Runtime/Datasmith/CADKernel/Private/CADKernel/Geo/Sampler/SamplerOnChord.h
--- a/Engine/Source/Runtime/Datasmith/CADKernel/Private/CADKernel/Geo/Sampler/SamplerOnChord.h
+++ b/Engine/Source/Runtime/Datasmith/CADKernel/Private/CADKernel/Geo/Sampler/SamplerOnChord.h
@@ -17,238 +17,239 @@
 namespace CADKernel
 {
 
-	template<class PointType>
-	class TSamplerBasedOnChordError : public TCurveSamplerAbstract<TPolyline<PointType>, PointType>
+template<class PointType>
+class TSamplerBasedOnChordError : public TCurveSamplerAbstract<TPolyline<PointType>, PointType>
+{
+public:
+	TSamplerBasedOnChordError(const FLinearBoundary& InBoundary, double InMaxAllowedError, TPolyline<PointType>& OutPolyline)
+		: TCurveSamplerAbstract<TPolyline<PointType>, PointType>(InBoundary, OutPolyline, InMaxAllowedError)
 	{
-	public:
-		TSamplerBasedOnChordError(const FLinearBoundary& InBoundary, double InMaxAllowedError, TPolyline<PointType>& OutPolyline)
-			: TCurveSamplerAbstract<TPolyline<PointType>, PointType>(InBoundary, OutPolyline, InMaxAllowedError)
-		{
-		};
+	};
 
-	protected:
-
-		/**
-		 * Check the chord error of the AB segment for each point of TmpPoints between FirstIndex and EndIndex
-		 * If the MaxChord is biggest to the DesiredChordError, an estimation of the needed point number to insert is done.
-		 * This estimation is bases on the approximation of local curve radius (R) :
-		 * R = square(L) / 8Sag
-		 * L2 = sqrt(square(L1) * S2 / S1)
-		 * S2 = S1 * square(L2 / L1)
-		 *
-		 * Sag & Angle criterion.pdf
-		 * https://docs.google.com/presentation/d/1bUnrRFWCW3sDn9ngb9ftfQS-2JxNJaUZlh783hZMMEw/edit?usp=sharing
-		 *
-		 * @return the number of point to add after PointA to PointB to respect the chord error.
-		 * if PointNum = 0, PointB is not necessary
-		 * if PointNum = 1, PointB has to be add to respect the chord error and the segment AB is respecting the chord error
-		 * if PointNum = 2, Point at FirstIndex and PointB have to be add to respect the chord error and both segments are respecting the chord error
-		 *                  In this case, a complementary check is done at each extremity to verify that the curve do not have a fast modification
-		 * if PointNum > 2, Point at FirstIndex and PointB have to be add but at least (PointNum - 2) points have to be added in both segments to respect the chord error
-		 */
-		virtual int32 CheckSamplingError(int32 FirstIndex, int32 EndIndex) override
-		{
-			const PointType& APoint = this->Sampling.Points[this->StartSamplingSegmentIndex];
-			double ACoordinate = this->Sampling.Coordinates[this->StartSamplingSegmentIndex];
+protected:
+
+	/**
+	 * Check the chord error of the AB segment for each point of TmpPoints between FirstIndex and EndIndex
+	 * If the MaxChord is biggest to the DesiredChordError, an estimation of the needed point number to insert is done.
+	 * This estimation is bases on the approximation of local curve radius (R) :
+	 * R = square(L) / 8Sag
+	 * L2 = sqrt(square(L1) * S2 / S1)
+	 * S2 = S1 * square(L2 / L1)
+	 *
+	 * Sag & Angle criterion.pdf
+	 * https://docs.google.com/presentation/d/1bUnrRFWCW3sDn9ngb9ftfQS-2JxNJaUZlh783hZMMEw/edit?usp=sharing
+	 *
+	 * @return the number of point to add after PointA to PointB to respect the chord error.
+	 * if PointNum = 0, PointB is not necessary
+	 * if PointNum = 1, PointB has to be add to respect the chord error and the segment AB is respecting the chord error
+	 * if PointNum = 2, Point at FirstIndex and PointB have to be add to respect the chord error and both segments are respecting the chord error
+	 *                  In this case, a complementary check is done at each extremity to verify that the curve do not have a fast modification
+	 * if PointNum > 2, Point at FirstIndex and PointB have to be add but at least (PointNum - 2) points have to be added in both segments to respect the chord error
+	 */
+	virtual int32 CheckSamplingError(int32 FirstIndex, int32 EndIndex) override
+	{
+		const PointType& APoint = this->Sampling.Points[this->StartSamplingSegmentIndex];
+		double ACoordinate = this->Sampling.Coordinates[this->StartSamplingSegmentIndex];
 
-			const PointType& BPoint = this->EndStudySegment.Polyline->Points[this->EndStudySegment.Index];
-			double BCoordinate = this->EndStudySegment.Polyline->Coordinates[this->EndStudySegment.Index];
+		const PointType& BPoint = this->EndStudySegment.Polyline->Points[this->EndStudySegment.Index];
+		double BCoordinate = this->EndStudySegment.Polyline->Coordinates[this->EndStudySegment.Index];
 
-			double MaxChord = 0;
-			for (int32 Index = FirstIndex; Index < EndIndex; ++Index)
+		double MaxChord = 0;
+		for (int32 Index = FirstIndex; Index < EndIndex; ++Index)
+		{
+			const PointType& Middle = this->CandidatePoints.GetPointAt(Index);
+			double Chord = DistanceOfPointToSegment(Middle, APoint, BPoint);
+			if (Chord > MaxChord)
 			{
-				const PointType& Middle = this->CandidatePoints.GetPointAt(Index);
-				double Chord = DistanceOfPointToSegment(Middle, APoint, BPoint);
-				if (Chord > MaxChord)
-				{
-					MaxChord = Chord;
-				}
+				MaxChord = Chord;
+			}
 
-				if (MaxChord > this->DesiredChordError)
-				{
-					break;
-				}
+			if (MaxChord > this->DesiredChordError)
+			{
+				break;
 			}
+		}
 
 #ifdef DEBUG_CHECKCHORDERROR
-			if (bDisplay)
+		if (bDisplay)
+		{
 			{
+				int32 Index = 0;
+				F3DDebugSession G(FString::Printf(TEXT("CheckChord %f %d"), MaxChord, FirstIndex));
+
+				CADKernel::Display(APoint, EVisuProperty::BluePoint);
+				CADKernel::Display(BPoint, EVisuProperty::BluePoint);
+
+				for (int32 Index = FirstIndex; Index < EndIndex; ++Index)
+				{
+					const PointType& Middle = CandidatePoints.GetPoints()[Index];
+					CADKernel::Display(Middle, EVisuProperty::YellowPoint);
+				}
+
+				DisplaySegment(APoint, BPoint, 0, EVisuProperty::BlueCurve);
+				DisplaySegment(APoint, CandidatePoints.GetPoints()[FirstIndex], 0, EVisuProperty::YellowCurve);
+				DisplaySegment(BPoint, CandidatePoints.GetPoints()[EndIndex - 1], 0, EVisuProperty::YellowCurve);
+
+				for (int32 Index = FirstIndex; Index < EndIndex - 1; ++Index)
 				{
-					int32 Index = 0;
-					F3DDebugSession G(FString::Printf(TEXT("CheckChord %f %d"), MaxChord, FirstIndex));
-
-					CADKernel::Display(APoint, EVisuProperty::BluePoint);
-					CADKernel::Display(BPoint, EVisuProperty::BluePoint);
-
-					for (int32 Index = FirstIndex; Index < EndIndex; ++Index)
-					{
-						const PointType& Middle = CandidatePoints.GetPoints()[Index];
-						CADKernel::Display(Middle, EVisuProperty::YellowPoint);
-					}
-
-					DisplaySegment(APoint, BPoint, 0, EVisuProperty::BlueCurve);
-					DisplaySegment(APoint, CandidatePoints.GetPoints()[FirstIndex], 0, EVisuProperty::YellowCurve);
-					DisplaySegment(BPoint, CandidatePoints.GetPoints()[EndIndex-1], 0, EVisuProperty::YellowCurve);
-
-					for (int32 Index = FirstIndex; Index < EndIndex-1; ++Index)
-					{
-						DisplaySegment(CandidatePoints.GetPoints()[Index], CandidatePoints.GetPoints()[Index + 1], 0, EVisuProperty::YellowCurve);
-					}
+					DisplaySegment(CandidatePoints.GetPoints()[Index], CandidatePoints.GetPoints()[Index + 1], 0, EVisuProperty::YellowCurve);
 				}
 			}
+		}
 #endif
 
-			if (MaxChord < this->DesiredChordError / 4)
-			{
-				// S2 = S1 * sqare(L2 / L1)
-				// if the length is doubled, the chord error is still smaller than the desired chord error.
-				// the end point is not needed
-				return 0;
-			}
-			else if (MaxChord < this->DesiredChordError)
-			{
-				// the end point is needed because the Chord is smaller but nearly equal to the desired chord error.
-				return 1;
-			}
-			else if (MaxChord < 4 * this->DesiredChordError)
-			{
-				// if the length is divided by two, the chord error is divided by 4 and becomes smaller than the desired chord error.
-				// Add First index and Index
-				return this->CheckTangentError(APoint, ACoordinate, BPoint, BCoordinate, FirstIndex, EndIndex, this->StartSamplingSegmentIndex);
-			}
-			else
-			{
-				return this->CountOfNeededPointsToRespectChordError(APoint, BPoint, MaxChord);
-			}
-		};
+		if (MaxChord < this->DesiredChordError / 4)
+		{
+			// S2 = S1 * sqare(L2 / L1)
+			// if the length is doubled, the chord error is still smaller than the desired chord error.
+			// the end point is not needed
+			return 0;
+		}
+		else if (MaxChord < this->DesiredChordError)
+		{
+			// the end point is needed because the Chord is smaller but nearly equal to the desired chord error.
+			return 1;
+		}
+		else if (MaxChord < 4 * this->DesiredChordError)
+		{
+			// if the length is divided by two, the chord error is divided by 4 and becomes smaller than the desired chord error.
+			// Add First index and Index
+			return this->CheckTangentError(APoint, ACoordinate, BPoint, BCoordinate, FirstIndex, EndIndex, this->StartSamplingSegmentIndex);
+		}
+		else
+		{
+			return this->CountOfNeededPointsToRespectChordError(APoint, BPoint, MaxChord);
+		}
+	};
 #ifdef DEBUG_CHECKCHORDERROR
-		bool bDisplay = true;
+	bool bDisplay = true;
 #endif
-	};
+};
 
-	/**
-	 * Sampler of curve based on parametrization and chord error control
-	 */
-	class FCurveSamplerOnChord : public TSamplerBasedOnChordError<FPoint>
+/**
+ * Sampler of curve based on parametrization and chord error control
+ */
+class FCurveSamplerOnChord : public TSamplerBasedOnChordError<FPoint>
+{
+public:
+	FCurveSamplerOnChord(const FCurve& InCurve, const FLinearBoundary& InBoundary, double InMaxAllowedError, FPolyline3D& OutPolyline)
+		: TSamplerBasedOnChordError<FPoint>(InBoundary, InMaxAllowedError, OutPolyline)
+		, Curve(InCurve)
 	{
-	public:
-		FCurveSamplerOnChord(const FCurve& InCurve, const FLinearBoundary& InBoundary, double InMaxAllowedError, FPolyline3D& OutPolyline)
-			: TSamplerBasedOnChordError<FPoint>(InBoundary, InMaxAllowedError, OutPolyline)
-			, Curve(InCurve)
-		{
-		};
-
-	protected:
-
-		const FCurve& Curve;
+	};
 
-		virtual void EvaluatesNewCandidatePoints() override
-		{
-			Curve.EvaluatePoints(CandidatePoints.Coordinates, CandidatePoints.Points);
-		};
+protected:
 
-		virtual void GetNotDerivableCoordinates(TArray<double>& OutNotDerivableCoordinates) override
-		{
-			Curve.FindNotDerivableCoordinates(Boundary, 2, OutNotDerivableCoordinates);
-		}
+	const FCurve& Curve;
 
+	virtual void EvaluatesNewCandidatePoints() override
+	{
+		Curve.EvaluatePoints(CandidatePoints.Coordinates, CandidatePoints.Points);
 	};
 
-	class FCurve2DSamplerOnChord : public TSamplerBasedOnChordError<FPoint2D>
+	virtual void GetNotDerivableCoordinates(TArray<double>& OutNotDerivableCoordinates) override
 	{
-	public:
-		FCurve2DSamplerOnChord(const FCurve& InCurve, const FLinearBoundary& InBoundary, double InMaxAllowedError, FPolyline2D& OutPolyline)
-			: TSamplerBasedOnChordError<FPoint2D>(InBoundary, InMaxAllowedError, OutPolyline)
-			, Curve(InCurve)
-		{
-		};
+		Curve.FindNotDerivableCoordinates(Boundary, 2, OutNotDerivableCoordinates);
+	}
 
-	protected:
+};
 
-		const FCurve& Curve;
+class FCurve2DSamplerOnChord : public TSamplerBasedOnChordError<FPoint2D>
+{
+public:
+	FCurve2DSamplerOnChord(const FCurve& InCurve, const FLinearBoundary& InBoundary, double InMaxAllowedError, FPolyline2D& OutPolyline)
+		: TSamplerBasedOnChordError<FPoint2D>(InBoundary, InMaxAllowedError, OutPolyline)
+		, Curve(InCurve)
+	{
+	};
 
-		virtual void EvaluatesNewCandidatePoints() override
-		{
-			Curve.Evaluate2DPoints(CandidatePoints.Coordinates, CandidatePoints.Points);
-		};
+protected:
 
-		virtual void GetNotDerivableCoordinates(TArray<double>& OutNotDerivableCoordinates) override
-		{
-			Curve.FindNotDerivableCoordinates(Boundary, 2, OutNotDerivableCoordinates);
-		}
+	const FCurve& Curve;
 
+	virtual void EvaluatesNewCandidatePoints() override
+	{
+		Curve.Evaluate2DPoints(CandidatePoints.Coordinates, CandidatePoints.Points);
 	};
 
-	class FIsoCurve3DSamplerOnChord : public TSamplerBasedOnChordError<FPoint>
+	virtual void GetNotDerivableCoordinates(TArray<double>& OutNotDerivableCoordinates) override
 	{
-	public:
-		FIsoCurve3DSamplerOnChord(const FSurface& InSurface, double InMaxAllowedError, FPolyline3D& OutPolyline)
-			: TSamplerBasedOnChordError<FPoint>(FLinearBoundary(), InMaxAllowedError, OutPolyline)
-			, Surface(InSurface)
-		{
-		};
+		Curve.FindNotDerivableCoordinates(Boundary, 2, OutNotDerivableCoordinates);
+	}
 
-		void Set(EIso InIsoType, double InIsoCoordinate, const FLinearBoundary& CurveBounds)
-		{
-			IsoType = InIsoType;
-			IsoCoordinate.Empty(1);
-			IsoCoordinate.Add(InIsoCoordinate);
-			Boundary = CurveBounds;
-		}
+};
 
-	protected:
+class FIsoCurve3DSamplerOnChord : public TSamplerBasedOnChordError<FPoint>
+{
+public:
+	FIsoCurve3DSamplerOnChord(const FSurface& InSurface, double InMaxAllowedError, FPolyline3D& OutPolyline)
+		: TSamplerBasedOnChordError<FPoint>(FLinearBoundary(), InMaxAllowedError, OutPolyline)
+		, Surface(InSurface)
+	{
+	};
 
-		const FSurface& Surface;
-		FCoordinateGrid NotDerivableCoordinates;
-		bool bNotDerivableFound = false;
+	void Set(EIso InIsoType, double InIsoCoordinate, const FLinearBoundary& CurveBounds)
+	{
+		IsoType = InIsoType;
+		IsoCoordinate.Empty(1);
+		IsoCoordinate.Add(InIsoCoordinate);
+		Boundary = CurveBounds;
+	}
 
-		mutable TArray<double> IsoCoordinate;
-		EIso IsoType;
+protected:
 
-		virtual void GetNotDerivableCoordinates(TArray<double>& OutNotDerivableCoordinates) override
-		{
-			if (!bNotDerivableFound)
-			{
-				Surface.LinesNotDerivables(Surface.GetBoundary(), 2, NotDerivableCoordinates);
-				bNotDerivableFound = true;
-			}
+	const FSurface& Surface;
+	FCoordinateGrid NotDerivableCoordinates;
+	bool bNotDerivableFound = false;
+
+	mutable TArray<double> IsoCoordinate;
+	EIso IsoType;
 
-			EIso CurveIso = IsoType == EIso::IsoU ? EIso::IsoV : EIso::IsoU;
-			const TArray<double>& SurfaceIsoNotDerivableCoordinates = NotDerivableCoordinates[CurveIso];
-			double Tolerance = Surface.GetIsoTolerance(CurveIso);
-			ArrayUtils::SubArrayWithoutBoundary(SurfaceIsoNotDerivableCoordinates, Boundary, Tolerance, OutNotDerivableCoordinates);
+	virtual void GetNotDerivableCoordinates(TArray<double>& OutNotDerivableCoordinates) override
+	{
+		if (!bNotDerivableFound)
+		{
+			Surface.LinesNotDerivables(Surface.GetBoundary(), 2, NotDerivableCoordinates);
+			bNotDerivableFound = true;
 		}
 
-		virtual void EvaluatesNewCandidatePoints() override
+		EIso CurveIso = IsoType == EIso::IsoU ? EIso::IsoV : EIso::IsoU;
+		const TArray<double>& SurfaceIsoNotDerivableCoordinates = NotDerivableCoordinates[CurveIso];
+		double Tolerance = Surface.GetIsoTolerance(CurveIso);
+		ArrayUtils::SubArrayWithoutBoundary(SurfaceIsoNotDerivableCoordinates, Boundary, Tolerance, OutNotDerivableCoordinates);
+	}
+
+	virtual void EvaluatesNewCandidatePoints() override
+	{
+		TFunction<void(TArray<double>&, TArray<double>&)> Evaluates = [&](TArray<double>& UCoordinates, TArray<double>& VCoordinates)
 		{
-			TFunction<void(TArray<double>&, TArray<double>&)> Evaluates = [&](TArray<double>& UCoordinates, TArray<double>& VCoordinates)
-			{
-				FSurfacicSampling PointGrid;
+			FSurfacicSampling PointGrid;
 
-				FCoordinateGrid Grid;
-				Grid.Swap(UCoordinates, VCoordinates);
-				Surface.EvaluatePointGrid(Grid, PointGrid);
-				Grid.Swap(UCoordinates, VCoordinates);
+			FCoordinateGrid Grid;
+			Grid.Swap(UCoordinates, VCoordinates);
+			Surface.EvaluatePointGrid(Grid, PointGrid);
+			Grid.Swap(UCoordinates, VCoordinates);
 
-				Swap(CandidatePoints.Points, PointGrid.Points3D);
-			};
+			Swap(CandidatePoints.Points, PointGrid.Points3D);
+		};
 
-			if (CandidatePoints.Coordinates.Num() > 0)
+		if (CandidatePoints.Coordinates.Num() > 0)
+		{
+			if (IsoType == EIso::IsoU)
 			{
-				if (IsoType == EIso::IsoU)
-				{
-					Evaluates(IsoCoordinate, CandidatePoints.Coordinates);
-				}
-				else
-				{
-					Evaluates(CandidatePoints.Coordinates, IsoCoordinate);
-				}
+				Evaluates(IsoCoordinate, CandidatePoints.Coordinates);
 			}
 			else
 			{
-				CandidatePoints.Points.SetNum(0);
+				Evaluates(CandidatePoints.Coordinates, IsoCoordinate);
 			}
 		}
+		else
+		{
+			CandidatePoints.Points.SetNum(0);
+		}
+	}
 
-	};
-}
+};
+
+} // ns CADKernel
