Update CADKernel : add oriention, fix one infinite loop

diff --git a/Engine/Source/Runtime/Datasmith/CADKernel/Private/CADKernel/Geo/Sampler/SamplerOnParam.h b/Engine/Source/Runtime/Datasmith/CADKernel/Private/CADKernel/Geo/Sampler/SamplerOnParam.h
--- a/Engine/Source/Runtime/Datasmith/CADKernel/Private/CADKernel/Geo/Sampler/SamplerOnParam.h
+++ b/Engine/Source/Runtime/Datasmith/CADKernel/Private/CADKernel/Geo/Sampler/SamplerOnParam.h
@@ -17,407 +17,407 @@
 namespace CADKernel
 {
 
-	template<typename PolylineType, typename PointType>
-	class TSamplerBasedOnParametrizationAndChordError : public TCurveSamplerAbstract<PolylineType, PointType>
+template<typename PolylineType, typename PointType>
+class TSamplerBasedOnParametrizationAndChordError : public TCurveSamplerAbstract<PolylineType, PointType>
+{
+public:
+	TSamplerBasedOnParametrizationAndChordError(const FLinearBoundary& InBoundary, double InMaxChordError, double InMaxParameterizationError, PolylineType& OutPolyline)
+		: TCurveSamplerAbstract<PolylineType, PointType>(InBoundary, OutPolyline, InMaxChordError)
+		, DesiredParameterizationError(InMaxParameterizationError)
 	{
-	public:
-		TSamplerBasedOnParametrizationAndChordError(const FLinearBoundary& InBoundary, double InMaxChordError, double InMaxParameterizationError, PolylineType& OutPolyline)
-			: TCurveSamplerAbstract<PolylineType, PointType>(InBoundary, OutPolyline, InMaxChordError)
-			,  DesiredParameterizationError(InMaxParameterizationError)
-		{
-		};
-
-	protected:
-
-		double  DesiredParameterizationError;
-
-		/**
-		 * This function check the chord error and the parametrization error
-		 *
-		 * Chord error:
-		 *   Check the chord error of the AB segment for each point of TmpPoints between FirstIndex and EndIndex
-		 *   If the MaxChord is biggest to the DesiredChordError, an estimation of the needed point number to insert is done.
-		 *   This estimation is bases on the approximation of local curve radius (R) :
-		 *   R = square(L) / 8Sag
-		 *   L2 = sqrt(square(L1) * S2 / S1)
-		 *   S2 = S1 * square(L2 / L1)
-		 *
-		 *   Sag & Angle criterion.pdf
-		 *   https://docs.google.com/presentation/d/1bUnrRFWCW3sDn9ngb9ftfQS-2JxNJaUZlh783hZMMEw/edit?usp=sharing
-		 *
-		 *   @return the number of point to add after PointA to PointB to respect the chord error.
-		 *   if PointNum = 0, PointB is not necessary
-		 *   if PointNum = 1, PointB has to be add to respect the chord error and the segment AB is respecting the chord error
-		 *   if PointNum = 2, Point at FirstIndex and PointB have to be add to respect the chord error and both segments are respecting the chord error
-		 *   if PointNum > 2, Point at FirstIndex and PointB have to be add but at least (PointNum - 2) points have to be added in both segments to respect the chord error
-		 *
-		 * Parametrization error:
-		 *   This function check that for inner point of the segment AB, point M(u) is near its approximation = A + AB * (u - uA) / (uB - uA) than the MaxAllowedError
-		 *   To Remove Chord error, M approximation is compare to the projection of M(u) on AB
-		 */
-
-		virtual int32 CheckSamplingError(int32 FirstIndex, int32 EndIndex) override
+	};
+
+protected:
+
+	double  DesiredParameterizationError;
+
+	/**
+	 * This function check the chord error and the parametrization error
+	 *
+	 * Chord error:
+	 *   Check the chord error of the AB segment for each point of TmpPoints between FirstIndex and EndIndex
+	 *   If the MaxChord is biggest to the DesiredChordError, an estimation of the needed point number to insert is done.
+	 *   This estimation is bases on the approximation of local curve radius (R) :
+	 *   R = square(L) / 8Sag
+	 *   L2 = sqrt(square(L1) * S2 / S1)
+	 *   S2 = S1 * square(L2 / L1)
+	 *
+	 *   Sag & Angle criterion.pdf
+	 *   https://docs.google.com/presentation/d/1bUnrRFWCW3sDn9ngb9ftfQS-2JxNJaUZlh783hZMMEw/edit?usp=sharing
+	 *
+	 *   @return the number of point to add after PointA to PointB to respect the chord error.
+	 *   if PointNum = 0, PointB is not necessary
+	 *   if PointNum = 1, PointB has to be add to respect the chord error and the segment AB is respecting the chord error
+	 *   if PointNum = 2, Point at FirstIndex and PointB have to be add to respect the chord error and both segments are respecting the chord error
+	 *   if PointNum > 2, Point at FirstIndex and PointB have to be add but at least (PointNum - 2) points have to be added in both segments to respect the chord error
+	 *
+	 * Parametrization error:
+	 *   This function check that for inner point of the segment AB, point M(u) is near its approximation = A + AB * (u - uA) / (uB - uA) than the MaxAllowedError
+	 *   To Remove Chord error, M approximation is compare to the projection of M(u) on AB
+	 */
+
+	virtual int32 CheckSamplingError(int32 FirstIndex, int32 EndIndex) override
+	{
+		const FPoint& APoint = this->Sampling.GetPointAt(this->StartSamplingSegmentIndex);
+		double ACoordinate = this->Sampling.Coordinates[this->StartSamplingSegmentIndex];
+
+		const FPoint& BPoint = this->EndStudySegment.Polyline->GetPointAt(this->EndStudySegment.Index);
+		double BCoordinate = this->EndStudySegment.Polyline->Coordinates[this->EndStudySegment.Index];
+
+		double ABCoordinate = BCoordinate - ACoordinate;
+		FPoint ABVector = BPoint - APoint;
+
+		double MaxChord = 0;
+		double MaxParametrizationError = 0;
+
+		for (int32 Index = FirstIndex; Index < EndIndex; ++Index)
 		{
-			const FPoint& APoint = this->Sampling.GetPointAt(this->StartSamplingSegmentIndex);
-			double ACoordinate = this->Sampling.Coordinates[this->StartSamplingSegmentIndex];
+			const FPoint& PointM = this->CandidatePoints.GetPointAt(Index);
 
-			const FPoint& BPoint = this->EndStudySegment.Polyline->GetPointAt(this->EndStudySegment.Index);
-			double BCoordinate = this->EndStudySegment.Polyline->Coordinates[this->EndStudySegment.Index];
+			double CurvilinearAbscise;
+			FPoint ProjectedPoint = ProjectPointOnSegment(PointM, APoint, BPoint, CurvilinearAbscise, /*bRestrictCoodinateToInside*/ true);
 
-			double ABCoordinate = BCoordinate - ACoordinate;
-			FPoint ABVector = BPoint - APoint;
+			double Chord = ProjectedPoint.Distance(PointM);
+			double AMCoordinate = this->CandidatePoints.Coordinates[Index] - ACoordinate;
+			FPoint MApproximate = APoint + (AMCoordinate / ABCoordinate) * ABVector;
+			double ParametrizationError = MApproximate.Distance(ProjectedPoint);
 
-			double MaxChord = 0;
-			double MaxParametrizationError = 0;
+			if (Chord > MaxChord)
+			{
+				MaxChord = Chord;
+			}
 
-			for (int32 Index = FirstIndex; Index < EndIndex; ++Index)
+			if (ParametrizationError > MaxParametrizationError)
 			{
-				const FPoint& PointM = this->CandidatePoints.GetPointAt(Index);
+				MaxParametrizationError = ParametrizationError;
+			}
+		}
 
-				double CurvilinearAbscise;
-				FPoint ProjectedPoint = ProjectPointOnSegment(PointM, APoint, BPoint, CurvilinearAbscise, /*bRestrictCoodinateToInside*/ true);
+#ifdef DEBUG_CHECKCHORDERROR
+		if (bDisplay)
+		{
+			{
+				int32 Index = 0;
+				F3DDebugSession G(FString::Printf(TEXT("CheckChord %f %d"), MaxChord, FirstIndex));
 
-				double Chord = ProjectedPoint.Distance(PointM);
-				double AMCoordinate = this->CandidatePoints.Coordinates[Index] - ACoordinate;
-				FPoint MApproximate = APoint + (AMCoordinate / ABCoordinate) * ABVector;
-				double ParametrizationError = MApproximate.Distance(ProjectedPoint);
+				CADKernel::Display(APoint, EVisuProperty::BluePoint);
+				CADKernel::Display(BPoint, EVisuProperty::BluePoint);
 
-				if (Chord > MaxChord)
+				for (int32 Index = FirstIndex; Index < EndIndex; ++Index)
 				{
-					MaxChord = Chord;
+					const FPoint& Middle = CandidatePoints.GetPoints()[Index];
+					CADKernel::Display(Middle, EVisuProperty::YellowPoint);
 				}
 
-				if (ParametrizationError > MaxParametrizationError)
-				{
-					MaxParametrizationError = ParametrizationError;
-				}
-			}
+				DisplaySegment(APoint, BPoint, 0, EVisuProperty::BlueCurve);
+				DisplaySegment(APoint, CandidatePoints.GetPoints()[FirstIndex], 0, EVisuProperty::YellowCurve);
+				DisplaySegment(BPoint, CandidatePoints.GetPoints()[EndIndex - 1], 0, EVisuProperty::YellowCurve);
 
-#ifdef DEBUG_CHECKCHORDERROR
-			if (bDisplay)
-			{
+				for (int32 Index = FirstIndex; Index < EndIndex - 1; ++Index)
 				{
-					int32 Index = 0;
-					F3DDebugSession G(FString::Printf(TEXT("CheckChord %f %d"), MaxChord, FirstIndex));
-
-					CADKernel::Display(APoint, EVisuProperty::BluePoint);
-					CADKernel::Display(BPoint, EVisuProperty::BluePoint);
-
-					for (int32 Index = FirstIndex; Index < EndIndex; ++Index)
-					{
-						const FPoint& Middle = CandidatePoints.GetPoints()[Index];
-						CADKernel::Display(Middle, EVisuProperty::YellowPoint);
-					}
-
-					DisplaySegment(APoint, BPoint, 0, EVisuProperty::BlueCurve);
-					DisplaySegment(APoint, CandidatePoints.GetPoints()[FirstIndex], 0, EVisuProperty::YellowCurve);
-					DisplaySegment(BPoint, CandidatePoints.GetPoints()[EndIndex - 1], 0, EVisuProperty::YellowCurve);
-
-					for (int32 Index = FirstIndex; Index < EndIndex - 1; ++Index)
-					{
-						DisplaySegment(CandidatePoints.GetPoints()[Index], CandidatePoints.GetPoints()[Index + 1], 0, EVisuProperty::YellowCurve);
-					}
+					DisplaySegment(CandidatePoints.GetPoints()[Index], CandidatePoints.GetPoints()[Index + 1], 0, EVisuProperty::YellowCurve);
 				}
 			}
+		}
 #endif
 
-			int32 SegmentNumToRespectParametrizationError = (int32)(MaxParametrizationError /  DesiredParameterizationError);
-			if (SegmentNumToRespectParametrizationError > 2)
-			{
-				SegmentNumToRespectParametrizationError = (int32)((SegmentNumToRespectParametrizationError - 1) / 2) + 2;
-			}
+		int32 SegmentNumToRespectParametrizationError = (int32)(MaxParametrizationError / DesiredParameterizationError);
+		if (SegmentNumToRespectParametrizationError > 2)
+		{
+			SegmentNumToRespectParametrizationError = (int32)((SegmentNumToRespectParametrizationError - 1) / 2) + 2;
+		}
 
-			int32 SegmentNumToRespectChordError = 0;
-			if (MaxChord < this->DesiredChordError / 4)
-			{
-				// S2 = S1 * sqare(L2 / L1)
-				// if the length is doubled, the chord error is still smaller than the desired chord error.
-				// the end point is not needed
-				return 0;
-			}
-			else if (MaxChord < this->DesiredChordError)
-			{
-				// the end point is needed because the Chord is smaller but nearly equal to the desired chord error.
-				SegmentNumToRespectChordError = 1;
-			}
-			else if (MaxChord < 4 * this->DesiredChordError)
-			{
-				// if the length is divided by two, the chord error is divided by 4 and becomes smaller than the desired chord error.
-				// Add First index and Index
-				SegmentNumToRespectChordError = this->CheckTangentError(APoint, ACoordinate, BPoint, BCoordinate, FirstIndex, EndIndex, this->StartSamplingSegmentIndex);
-			}
-			else
-			{
-				SegmentNumToRespectChordError = this->CountOfNeededPointsToRespectChordError(APoint, BPoint, MaxChord);
-			}
+		int32 SegmentNumToRespectChordError = 0;
+		if (MaxChord < this->DesiredChordError / 4)
+		{
+			// S2 = S1 * sqare(L2 / L1)
+			// if the length is doubled, the chord error is still smaller than the desired chord error.
+			// the end point is not needed
+			return 0;
+		}
+		else if (MaxChord < this->DesiredChordError)
+		{
+			// the end point is needed because the Chord is smaller but nearly equal to the desired chord error.
+			SegmentNumToRespectChordError = 1;
+		}
+		else if (MaxChord < 4 * this->DesiredChordError)
+		{
+			// if the length is divided by two, the chord error is divided by 4 and becomes smaller than the desired chord error.
+			// Add First index and Index
+			SegmentNumToRespectChordError = this->CheckTangentError(APoint, ACoordinate, BPoint, BCoordinate, FirstIndex, EndIndex, this->StartSamplingSegmentIndex);
+		}
+		else
+		{
+			SegmentNumToRespectChordError = this->CountOfNeededPointsToRespectChordError(APoint, BPoint, MaxChord);
+		}
 
-			return FMath::Max(SegmentNumToRespectChordError, SegmentNumToRespectParametrizationError);
-		};
+		return FMath::Max(SegmentNumToRespectChordError, SegmentNumToRespectParametrizationError);
 	};
+};
 
-	/**
-	 * Sampler of surfacic curve based on parametrization and chord error control
-	 */
-	class FSurfacicCurveSamplerOnParam : public TSamplerBasedOnParametrizationAndChordError<FSurfacicPolyline, FPoint>
+/**
+ * Sampler of surfacic curve based on parametrization and chord error control
+ */
+class FSurfacicCurveSamplerOnParam : public TSamplerBasedOnParametrizationAndChordError<FSurfacicPolyline, FPoint>
+{
+public:
+	FSurfacicCurveSamplerOnParam(const FRestrictionCurve& InCurve, const FLinearBoundary& InBoundary, double InMaxSagError, double InMaxParameterizationError, FSurfacicPolyline& OutPolyline)
+		: TSamplerBasedOnParametrizationAndChordError<FSurfacicPolyline, FPoint>(InBoundary, InMaxSagError, InMaxParameterizationError, OutPolyline)
+		, Surface(*InCurve.GetCarrierSurface())
+		, Curve(*InCurve.Get2DCurve())
 	{
-	public:
-		FSurfacicCurveSamplerOnParam(const FRestrictionCurve& InCurve, const FLinearBoundary& InBoundary, double InMaxSagError, double InMaxParameterizationError, FSurfacicPolyline& OutPolyline)
-			: TSamplerBasedOnParametrizationAndChordError<FSurfacicPolyline, FPoint>(InBoundary, InMaxSagError, InMaxParameterizationError, OutPolyline)
-			, Surface(*InCurve.GetCarrierSurface())
-			, Curve(*InCurve.Get2DCurve())
-		{
-		};
+	};
 
-		FSurfacicCurveSamplerOnParam(const FSurface& InSurface, const FCurve& InCurve, const FLinearBoundary& InBoundary, double InMaxSagError, double InMaxParameterizationError, FSurfacicPolyline& OutPolyline)
-			: TSamplerBasedOnParametrizationAndChordError<FSurfacicPolyline, FPoint>(InBoundary, InMaxSagError, InMaxParameterizationError, OutPolyline)
-			, Surface(InSurface)
-			, Curve(InCurve)
-		{
-		};
+	FSurfacicCurveSamplerOnParam(const FSurface& InSurface, const FCurve& InCurve, const FLinearBoundary& InBoundary, double InMaxSagError, double InMaxParameterizationError, FSurfacicPolyline& OutPolyline)
+		: TSamplerBasedOnParametrizationAndChordError<FSurfacicPolyline, FPoint>(InBoundary, InMaxSagError, InMaxParameterizationError, OutPolyline)
+		, Surface(InSurface)
+		, Curve(InCurve)
+	{
+	};
 
-	protected:
+protected:
 
-		const FSurface& Surface;
-		const FCurve& Curve;
-		double MaxParameterizationError;
+	const FSurface& Surface;
+	const FCurve& Curve;
+	double MaxParameterizationError;
 
-		virtual void EvaluatesNewCandidatePoints() override
+	virtual void EvaluatesNewCandidatePoints() override
+	{
+		CandidatePoints.bWithNormals = Sampling.bWithNormals;
+		CandidatePoints.bWithTangent = Sampling.bWithTangent;
+		if (CandidatePoints.bWithTangent)
 		{
-			CandidatePoints.bWithNormals = Sampling.bWithNormals;
-			CandidatePoints.bWithTangent = Sampling.bWithTangent;
-			if (CandidatePoints.bWithTangent)
-			{
-				TArray<FCurvePoint2D> OutPoints;
-				Curve.Evaluate2DPoints(CandidatePoints.Coordinates, OutPoints, 1);
-				Surface.EvaluatePoints(OutPoints, CandidatePoints);
-			}
-			else
-			{
-				Curve.Evaluate2DPoints(CandidatePoints.Coordinates, CandidatePoints.Points2D);
-				Surface.EvaluatePoints(CandidatePoints);
-			}
-		};
-
-		virtual void GetNotDerivableCoordinates(TArray<double>& OutNotDerivableCoordinates) override
+			TArray<FCurvePoint2D> OutPoints;
+			Curve.Evaluate2DPoints(CandidatePoints.Coordinates, OutPoints, 1);
+			Surface.EvaluatePoints(OutPoints, CandidatePoints);
+		}
+		else
 		{
-			Curve.FindNotDerivableCoordinates(Boundary, 2, OutNotDerivableCoordinates);
+			Curve.Evaluate2DPoints(CandidatePoints.Coordinates, CandidatePoints.Points2D);
+			Surface.EvaluatePoints(CandidatePoints);
 		}
-
 	};
 
-	/**
-	 * Sampler of curve based on parametrization and chord error control
-	 */
-	class FCurveSamplerOnParam : public TSamplerBasedOnParametrizationAndChordError<FPolyline3D, FPoint>
+	virtual void GetNotDerivableCoordinates(TArray<double>& OutNotDerivableCoordinates) override
 	{
-	public:
-		FCurveSamplerOnParam(const FCurve& InCurve, const FLinearBoundary& InBoundary, double InMaxSagError, double InMaxParameterizationError, FPolyline3D& OutPolyline)
-			: TSamplerBasedOnParametrizationAndChordError<FPolyline3D, FPoint>(InBoundary, InMaxSagError, InMaxParameterizationError, OutPolyline)
-			, Curve(InCurve)
-		{
-		};
+		Curve.FindNotDerivableCoordinates(Boundary, 2, OutNotDerivableCoordinates);
+	}
+
+};
 
-	protected:
+/**
+ * Sampler of curve based on parametrization and chord error control
+ */
+class FCurveSamplerOnParam : public TSamplerBasedOnParametrizationAndChordError<FPolyline3D, FPoint>
+{
+public:
+	FCurveSamplerOnParam(const FCurve& InCurve, const FLinearBoundary& InBoundary, double InMaxSagError, double InMaxParameterizationError, FPolyline3D& OutPolyline)
+		: TSamplerBasedOnParametrizationAndChordError<FPolyline3D, FPoint>(InBoundary, InMaxSagError, InMaxParameterizationError, OutPolyline)
+		, Curve(InCurve)
+	{
+	};
 
-		const FCurve& Curve;
-		double MaxParameterizationError;
+protected:
 
-		virtual void EvaluatesNewCandidatePoints() override
-		{
-			Curve.EvaluatePoints(CandidatePoints.Coordinates, CandidatePoints.Points);
-		};
+	const FCurve& Curve;
+	double MaxParameterizationError;
 
-		virtual void GetNotDerivableCoordinates(TArray<double>& OutNotDerivableCoordinates) override
-		{
-			Curve.FindNotDerivableCoordinates(Boundary, 2, OutNotDerivableCoordinates);
-		}
+	virtual void EvaluatesNewCandidatePoints() override
+	{
+		Curve.EvaluatePoints(CandidatePoints.Coordinates, CandidatePoints.Points);
 	};
 
-	/**
-	 * Sampler of surface based on parametrization and chord error control
-	 */
-	class FSurfaceSamplerOnParam : public TSamplerBasedOnParametrizationAndChordError<FPolyline3D, FPoint>
+	virtual void GetNotDerivableCoordinates(TArray<double>& OutNotDerivableCoordinates) override
 	{
+		Curve.FindNotDerivableCoordinates(Boundary, 2, OutNotDerivableCoordinates);
+	}
+};
+
+/**
+ * Sampler of surface based on parametrization and chord error control
+ */
+class FSurfaceSamplerOnParam : public TSamplerBasedOnParametrizationAndChordError<FPolyline3D, FPoint>
+{
 
-	public:
-		FSurfaceSamplerOnParam(const FSurface& InSurface, const FSurfacicBoundary& InBoundary, double InMaxSagError, double InMaxParameterizationError, FPolyline3D& TemporaryPolyline, FCoordinateGrid& OutSampling)
-			: TSamplerBasedOnParametrizationAndChordError<FPolyline3D, FPoint>(FLinearBoundary(), InMaxSagError, InMaxParameterizationError, TemporaryPolyline)
-			, Surface(InSurface)
-			, SurfaceBoundary(InBoundary)
-			, SurfaceSampling(OutSampling)
-		{
-		};
+public:
+	FSurfaceSamplerOnParam(const FSurface& InSurface, const FSurfacicBoundary& InBoundary, double InMaxSagError, double InMaxParameterizationError, FPolyline3D& TemporaryPolyline, FCoordinateGrid& OutSampling)
+		: TSamplerBasedOnParametrizationAndChordError<FPolyline3D, FPoint>(FLinearBoundary(), InMaxSagError, InMaxParameterizationError, TemporaryPolyline)
+		, Surface(InSurface)
+		, SurfaceBoundary(InBoundary)
+		, SurfaceSampling(OutSampling)
+	{
+	};
 
-		void Set(EIso InIsoType, double InIsoCoordinate, const FLinearBoundary& CurveBounds)
-		{
-			IsoType = InIsoType;
-			IsoCoordinate.Empty(1);
-			IsoCoordinate.Add(InIsoCoordinate);
-			Boundary = CurveBounds;
-		}
+	void Set(EIso InIsoType, double InIsoCoordinate, const FLinearBoundary& CurveBounds)
+	{
+		IsoType = InIsoType;
+		IsoCoordinate.Empty(1);
+		IsoCoordinate.Add(InIsoCoordinate);
+		Boundary = CurveBounds;
+	}
 
-		virtual void Sample() override
-		{
-			Sample(EIso::IsoU);
-			Sample(EIso::IsoV);
-		}
+	virtual void Sample() override
+	{
+		Sample(EIso::IsoU);
+		Sample(EIso::IsoV);
+	}
 
-		void Sample(EIso InIsoType)
-		{
+	void Sample(EIso InIsoType)
+	{
 #ifdef DEBUG_CURVE_SAMPLING
-			CurveIndex = CurveToDisplay - 1;
+		CurveIndex = CurveToDisplay - 1;
 #endif
 
-			double Middle = SurfaceBoundary[InIsoType].GetMiddle();
+		double Middle = SurfaceBoundary[InIsoType].GetMiddle();
 
-			Set(InIsoType, Middle, SurfaceBoundary[Other(InIsoType)]);
+		Set(InIsoType, Middle, SurfaceBoundary[Other(InIsoType)]);
 
-			SamplingInitalizing();
-			RunSampling();
+		SamplingInitalizing();
+		RunSampling();
 
-			double Step = SurfaceBoundary[InIsoType].Length() / 6;
-			IsoCoordinate[0] = SurfaceBoundary[InIsoType].GetMin();
-			for (int32 IsoIndex = 1; IsoIndex < 6; ++IsoIndex)
+		double Step = SurfaceBoundary[InIsoType].Length() / 6;
+		IsoCoordinate[0] = SurfaceBoundary[InIsoType].GetMin();
+		for (int32 IsoIndex = 1; IsoIndex < 6; ++IsoIndex)
+		{
+			IsoCoordinate[0] += Step;
+			if (IsoIndex == 3)
 			{
-				IsoCoordinate[0] += Step;
-				if (IsoIndex == 3)
-				{
-					continue;
-				}
+				continue;
+			}
 
-				Swap(CandidatePoints.Coordinates, Sampling.Coordinates);
-				EvaluatesNewCandidatePoints();
-				Swap(CandidatePoints.Coordinates, Sampling.Coordinates);
-				Swap(CandidatePoints.Points, Sampling.Points);
+			Swap(CandidatePoints.Coordinates, Sampling.Coordinates);
+			EvaluatesNewCandidatePoints();
+			Swap(CandidatePoints.Coordinates, Sampling.Coordinates);
+			Swap(CandidatePoints.Points, Sampling.Points);
 
-				// Check the result
-				CandidatePoints.Coordinates.Empty(Sampling.Coordinates.Num());
+			// Check the result
+			CandidatePoints.Coordinates.Empty(Sampling.Coordinates.Num());
 
-				for (int32 Index = 0; Index < Sampling.Coordinates.Num()-1; ++Index)
-				{
-					double UCoord = Sampling.Coordinates[Index] + (Sampling.Coordinates[Index + 1] - Sampling.Coordinates[Index]) * 0.5;
-					CandidatePoints.Coordinates.Add(UCoord);
-				}
+			for (int32 Index = 0; Index < Sampling.Coordinates.Num() - 1; ++Index)
+			{
+				double UCoord = Sampling.Coordinates[Index] + (Sampling.Coordinates[Index + 1] - Sampling.Coordinates[Index]) * 0.5;
+				CandidatePoints.Coordinates.Add(UCoord);
+			}
 
-				EvaluatesNewCandidatePoints();
+			EvaluatesNewCandidatePoints();
 
 #ifdef DEBUG_CURVE_SAMPLING
-				CurveToDisplay = CurveIndex;
-				DisplaySampling(CurveIndex == CurveToDisplay, 0);
+			CurveToDisplay = CurveIndex;
+			DisplaySampling(CurveIndex == CurveToDisplay, 0);
 #endif
 
-				RunSampling();
-			}
-
-			Swap(SurfaceSampling[Other(InIsoType)], Sampling.Coordinates);
+			RunSampling();
 		}
 
-		virtual void SamplingInitalizing() override
-		{
+		Swap(SurfaceSampling[Other(InIsoType)], Sampling.Coordinates);
+	}
+
+	virtual void SamplingInitalizing() override
+	{
 
 #ifdef DEBUG_CURVE_SAMPLING
-			++CurveIndex;
+		++CurveIndex;
 #endif
-			IsOptimalSegments.Empty(100);
-			Sampling.Empty(100);
+		IsOptimalSegments.Empty(100);
+		Sampling.Empty(100);
 
-			CandidatePoints.Empty(100);
+		CandidatePoints.Empty(100);
 
-			// Initialization of the algorithm with the not derivable coordinates of the curve and at least 5 temporary points
-			int32 ComplementatyPointCount = 0;
-			{
-				TArray<double> LocalNotDerivableCoordinates;
-				GetNotDerivableCoordinates(LocalNotDerivableCoordinates);
-				LocalNotDerivableCoordinates.Insert(Boundary.Min, 0);
-				LocalNotDerivableCoordinates.Add(Boundary.Max);
+		// Initialization of the algorithm with the not derivable coordinates of the curve and at least 5 temporary points
+		int32 ComplementatyPointCount = 0;
+		{
+			TArray<double> LocalNotDerivableCoordinates;
+			GetNotDerivableCoordinates(LocalNotDerivableCoordinates);
+			LocalNotDerivableCoordinates.Insert(Boundary.Min, 0);
+			LocalNotDerivableCoordinates.Add(Boundary.Max);
 
-				ComplementatyPointCount = LocalNotDerivableCoordinates.Num() < 5 ? 5 : 1;
+			ComplementatyPointCount = LocalNotDerivableCoordinates.Num() < 5 ? 5 : 1;
 
-				NextCoordinates.Empty(LocalNotDerivableCoordinates.Num() * (ComplementatyPointCount + 1));
-				NextCoordinates.Add(Boundary.Min);
+			NextCoordinates.Empty(LocalNotDerivableCoordinates.Num() * (ComplementatyPointCount + 1));
+			NextCoordinates.Add(Boundary.Min);
 
-				for (int32 Index = 0; Index < LocalNotDerivableCoordinates.Num() - 1; ++Index)
-				{
-					AddIntermediateCoordinates(LocalNotDerivableCoordinates[Index], LocalNotDerivableCoordinates[Index + 1], ComplementatyPointCount);
-					NextCoordinates.Add(LocalNotDerivableCoordinates[Index + 1]);
-				}
+			for (int32 Index = 0; Index < LocalNotDerivableCoordinates.Num() - 1; ++Index)
+			{
+				AddIntermediateCoordinates(LocalNotDerivableCoordinates[Index], LocalNotDerivableCoordinates[Index + 1], ComplementatyPointCount);
+				NextCoordinates.Add(LocalNotDerivableCoordinates[Index + 1]);
 			}
+		}
 
-			CandidatePoints.SwapCoordinates(NextCoordinates);
-			EvaluatesNewCandidatePoints();
+		CandidatePoints.SwapCoordinates(NextCoordinates);
+		EvaluatesNewCandidatePoints();
 
-			// Not Derivable point initialize the sampling
-			ComplementatyPointCount++;
-			for (int32 Index = 0, ISampling = 0; Index < CandidatePoints.Size(); Index += ComplementatyPointCount, ++ISampling)
-			{
-				Sampling.EmplaceAt(ISampling, CandidatePoints, Index);
-				IsOptimalSegments.Add(false);
-			}
-			IsOptimalSegments.Pop();
+		// Not Derivable point initialize the sampling
+		ComplementatyPointCount++;
+		for (int32 Index = 0, ISampling = 0; Index < CandidatePoints.Size(); Index += ComplementatyPointCount, ++ISampling)
+		{
+			Sampling.EmplaceAt(ISampling, CandidatePoints, Index);
+			IsOptimalSegments.Add(false);
+		}
+		IsOptimalSegments.Pop();
 
-			for (int32 Index = CandidatePoints.Size() - 1; Index >= 0; Index -= ComplementatyPointCount)
-			{
-				CandidatePoints.RemoveAt(Index);
-			}
+		for (int32 Index = CandidatePoints.Size() - 1; Index >= 0; Index -= ComplementatyPointCount)
+		{
+			CandidatePoints.RemoveAt(Index);
+		}
 
-			// first segment is not optimal
+		// first segment is not optimal
 
 #ifdef DEBUG_CURVE_SAMPLING
-			int32 Step = 0;
-			DisplaySampling(CurveIndex == CurveToDisplay, Step);
+		int32 Step = 0;
+		DisplaySampling(CurveIndex == CurveToDisplay, Step);
 #endif
-		}
+	}
 
-	protected:
+protected:
 
-		const FSurface& Surface;
-		const FSurfacicBoundary& SurfaceBoundary;
+	const FSurface& Surface;
+	const FSurfacicBoundary& SurfaceBoundary;
 
-		FCoordinateGrid& SurfaceSampling;
+	FCoordinateGrid& SurfaceSampling;
 
-		FCoordinateGrid NotDerivableCoordinates;
+	FCoordinateGrid NotDerivableCoordinates;
 
-		bool bNotDerivableFound = false;
+	bool bNotDerivableFound = false;
 
-		TArray<double> IsoCoordinate;
-		EIso IsoType;
+	TArray<double> IsoCoordinate;
+	EIso IsoType;
 
-		virtual void GetNotDerivableCoordinates(TArray<double>& OutNotDerivableCoordinates) override
+	virtual void GetNotDerivableCoordinates(TArray<double>& OutNotDerivableCoordinates) override
+	{
+		if (!bNotDerivableFound)
 		{
-			if (!bNotDerivableFound)
-			{
-				Surface.LinesNotDerivables(Surface.GetBoundary(), 2, NotDerivableCoordinates);
-				bNotDerivableFound = true;
-			}
-
-			EIso CurveIso = IsoType == EIso::IsoU ? EIso::IsoV : EIso::IsoU;
-			const TArray<double>& SurfaceIsoNotDerivableCoordinates = NotDerivableCoordinates[CurveIso];
-			double Tolerance = Surface.GetIsoTolerance(CurveIso);
-			ArrayUtils::SubArrayWithoutBoundary(SurfaceIsoNotDerivableCoordinates, Boundary, Tolerance, OutNotDerivableCoordinates);
+			Surface.LinesNotDerivables(Surface.GetBoundary(), 2, NotDerivableCoordinates);
+			bNotDerivableFound = true;
 		}
 
-		virtual void EvaluatesNewCandidatePoints() override
+		EIso CurveIso = IsoType == EIso::IsoU ? EIso::IsoV : EIso::IsoU;
+		const TArray<double>& SurfaceIsoNotDerivableCoordinates = NotDerivableCoordinates[CurveIso];
+		double Tolerance = Surface.GetIsoTolerance(CurveIso);
+		ArrayUtils::SubArrayWithoutBoundary(SurfaceIsoNotDerivableCoordinates, Boundary, Tolerance, OutNotDerivableCoordinates);
+	}
+
+	virtual void EvaluatesNewCandidatePoints() override
+	{
+		if (CandidatePoints.Coordinates.Num() > 0)
 		{
-			if (CandidatePoints.Coordinates.Num() > 0)
-			{
-				TArray<double>& UCoordinates = IsoType == EIso::IsoU ? IsoCoordinate : CandidatePoints.Coordinates;
-				TArray<double>& VCoordinates = IsoType == EIso::IsoU ? CandidatePoints.Coordinates : IsoCoordinate;
+			TArray<double>& UCoordinates = IsoType == EIso::IsoU ? IsoCoordinate : CandidatePoints.Coordinates;
+			TArray<double>& VCoordinates = IsoType == EIso::IsoU ? CandidatePoints.Coordinates : IsoCoordinate;
 
-				FSurfacicSampling PointGrid;
-				FCoordinateGrid Grid;
+			FSurfacicSampling PointGrid;
+			FCoordinateGrid Grid;
 
-				Grid.Swap(UCoordinates, VCoordinates);
-				Surface.EvaluatePointGrid(Grid, PointGrid);
-				Grid.Swap(UCoordinates, VCoordinates);
+			Grid.Swap(UCoordinates, VCoordinates);
+			Surface.EvaluatePointGrid(Grid, PointGrid);
+			Grid.Swap(UCoordinates, VCoordinates);
 
-				Swap(CandidatePoints.Points, PointGrid.Points3D);
-			}
-			else
-			{
-				CandidatePoints.Points.SetNum(0);
-			}
+			Swap(CandidatePoints.Points, PointGrid.Points3D);
+		}
+		else
+		{
+			CandidatePoints.Points.SetNum(0);
 		}
+	}
 
-	};
+};
 
-}
\ No newline at end of file
+} // ns CADKernel
\ No newline at end of file
diff --git a/Engine/Source/Runtime/Datasmith/CADKernel/Private/CADKernel/Geo/Sampling/SurfacicSampling.h b/Engine/Source/Runtime/Datasmith/CADKernel/Private/CADKernel/Geo/Sampling/SurfacicSampling.h
--- a/Engine/Source/Runtime/Datasmith/CADKernel/Private/CADKernel/Geo/Sampling/SurfacicSampling.h
+++ b/Engine/Source/Runtime/Datasmith/CADKernel/Private/CADKernel/Geo/Sampling/SurfacicSampling.h
@@ -7,56 +7,58 @@
 
 namespace CADKernel
 {
-	struct FSurfacicSampling
+
+struct FSurfacicSampling
+{
+	bool bWithNormals = false;
+	TArray<FPoint2D> Points2D;
+	TArray<FPoint> Points3D;
+	TArray<FVector> Normals;
+
+	int32 Count()
 	{
-		bool bWithNormals = false;
-		TArray<FPoint2D> Points2D;
-		TArray<FPoint> Points3D;
-		TArray<FVector> Normals;
+		return Points3D.Num();
+	}
 
-		int32 Count()
+	void SetNum(int32 Number)
+	{
+		Points2D.SetNum(Number);
+		Points3D.SetNum(Number);
+		if (bWithNormals)
 		{
-			return Points3D.Num();
+			Normals.SetNum(Number);
 		}
+	}
 
-		void SetNum(int32 Number)
+	void Reserve(int32 Number)
+	{
+		Points2D.Empty(Number);
+		Points3D.Empty(Number);
+		if (bWithNormals)
 		{
-			Points2D.SetNum(Number);
-			Points3D.SetNum(Number);
-			if (bWithNormals)
-			{
-				Normals.SetNum(Number);
-			}
+			Normals.Empty(Number);
 		}
+	}
 
-		void Reserve(int32 Number)
+	void NormalizeNormals()
+	{
+		for (FVector& Normal : Normals)
 		{
-			Points2D.Empty(Number);
-			Points3D.Empty(Number);
-			if (bWithNormals)
-			{
-				Normals.Empty(Number);
-			}
+			Normal.Normalize();
 		}
+	}
 
-		void NormalizeNormals()
+	void Set2DCoordinates(const FCoordinateGrid& Coordinates)
+	{
+		Reserve(Coordinates.Count());
+		for (double VCoord : Coordinates[EIso::IsoV])
 		{
-			for (FVector& Normal : Normals)
+			for (double UCoord : Coordinates[EIso::IsoU])
 			{
-				Normal.Normalize();
+				Points2D.Emplace(UCoord, VCoord);
 			}
 		}
+	}
+};
 
-		void Set2DCoordinates(const FCoordinateGrid& Coordinates)
-		{
-			Reserve(Coordinates.Count());
-			for (double VCoord : Coordinates[EIso::IsoV])
-			{
-				for (double UCoord : Coordinates[EIso::IsoU])
-				{
-					Points2D.Emplace(UCoord, VCoord);
-				}
-			}
-		}
-	};
-}
\ No newline at end of file
+} // ns CADKernel
\ No newline at end of file
diff --git a/Engine/Source/Runtime/Datasmith/CADKernel/Private/CADKernel/Geo/Surfaces/Surface.cpp b/Engine/Source/Runtime/Datasmith/CADKernel/Private/CADKernel/Geo/Surfaces/Surface.cpp
--- a/Engine/Source/Runtime/Datasmith/CADKernel/Private/CADKernel/Geo/Surfaces/Surface.cpp
+++ b/Engine/Source/Runtime/Datasmith/CADKernel/Private/CADKernel/Geo/Surfaces/Surface.cpp
@@ -15,196 +15,196 @@
 namespace CADKernel
 {
 #ifdef CADKERNEL_DEV
-	FInfoEntity& FSurface::GetInfo(FInfoEntity& Info) const
-	{
-		return FEntityGeom::GetInfo(Info)
-			.Add(TEXT("Surface type"), SurfacesTypesNames[(uint8)GetSurfaceType()])
-			.Add(TEXT("Boundary"), Boundary)
-			.Add(TEXT("MinToleranceIso"), (FPoint2D) *MinToleranceIso);
-	}
+FInfoEntity& FSurface::GetInfo(FInfoEntity& Info) const
+{
+	return FEntityGeom::GetInfo(Info)
+		.Add(TEXT("Surface type"), SurfacesTypesNames[(uint8)GetSurfaceType()])
+		.Add(TEXT("Boundary"), Boundary)
+		.Add(TEXT("MinToleranceIso"), (FPoint2D)*MinToleranceIso);
+}
 #endif
 
-	void FSurface::Sample(const FSurfacicBoundary& InBoundary, int32 NumberOfSubdivisions[2], FCoordinateGrid& OutCoordinateSampling) const
+void FSurface::Sample(const FSurfacicBoundary& InBoundary, int32 NumberOfSubdivisions[2], FCoordinateGrid& OutCoordinateSampling) const
+{
+	for (int32 Index = EIso::IsoU; Index <= EIso::IsoV; Index++)
 	{
-		for (int32 Index = EIso::IsoU; Index <= EIso::IsoV; Index++)
+		EIso Iso = (EIso)Index;
+		double Step = (InBoundary[Iso].Max - InBoundary[Iso].Min) / ((double)(NumberOfSubdivisions[Iso] - 1));
+
+		OutCoordinateSampling[Iso].Empty(NumberOfSubdivisions[Iso]);
+		double Value = InBoundary[Iso].Min;
+		for (int32 Kndex = 0; Kndex < NumberOfSubdivisions[Iso]; Kndex++)
 		{
-			EIso Iso = (EIso)Index;
-			double Step = (InBoundary[Iso].Max - InBoundary[Iso].Min) / ((double)(NumberOfSubdivisions[Iso] - 1));
-
-			OutCoordinateSampling[Iso].Empty(NumberOfSubdivisions[Iso]);
-			double Value = InBoundary[Iso].Min;
-			for (int32 Kndex = 0; Kndex < NumberOfSubdivisions[Iso]; Kndex++)
-			{
-				OutCoordinateSampling[Iso].Add(Value);
-				Value += Step;
-			}
+			OutCoordinateSampling[Iso].Add(Value);
+			Value += Step;
 		}
 	}
+}
 
-	void FSurface::Sample(const FSurfacicBoundary& Bounds, int32 NumberOfSubdivisions[2], FSurfacicSampling& OutPointSampling) const
-	{
-		FCoordinateGrid CoordinateSampling;
-		Sample(Bounds, NumberOfSubdivisions, CoordinateSampling);
-		EvaluatePointGrid(CoordinateSampling, OutPointSampling);
-	}
+void FSurface::Sample(const FSurfacicBoundary& Bounds, int32 NumberOfSubdivisions[2], FSurfacicSampling& OutPointSampling) const
+{
+	FCoordinateGrid CoordinateSampling;
+	Sample(Bounds, NumberOfSubdivisions, CoordinateSampling);
+	EvaluatePointGrid(CoordinateSampling, OutPointSampling);
+}
+
+void FSurface::Presample(const FSurfacicBoundary& InBoundaries, FCoordinateGrid& OutCoordinates)
+{
+	FPolyline3D TemporaryPolyline;
+	FSurfaceSamplerOnParam Sampler(*this, InBoundaries, Tolerance3D * 10, Tolerance3D, TemporaryPolyline, OutCoordinates);
+	Sampler.Sample();
+}
 
-	void FSurface::Presample(const FSurfacicBoundary& InBoundaries, FCoordinateGrid& OutCoordinates)
+void FSurface::EvaluatePoints(const TArray<FPoint2D>& InSurfacicCoordinates, TArray<FSurfacicPoint>& OutPoint3D, int32 InDerivativeOrder) const
+{
+	OutPoint3D.SetNum(InSurfacicCoordinates.Num());
+	for (int32 Index = 0; Index < InSurfacicCoordinates.Num(); ++Index)
 	{
-		FPolyline3D TemporaryPolyline;
-		FSurfaceSamplerOnParam Sampler(*this, InBoundaries, Tolerance3D * 10, Tolerance3D, TemporaryPolyline, OutCoordinates);
-		Sampler.Sample();
+		EvaluatePoint(InSurfacicCoordinates[Index], OutPoint3D[Index], InDerivativeOrder);
 	}
+}
 
-	void FSurface::EvaluatePoints(const TArray<FPoint2D>& InSurfacicCoordinates, TArray<FSurfacicPoint>& OutPoint3D, int32 InDerivativeOrder) const
+void FSurface::EvaluatePoints(const TArray<FCurvePoint2D>& InSurfacicCoordinates, TArray<FSurfacicPoint>& OutPoint3D, int32 InDerivativeOrder) const
+{
+	OutPoint3D.SetNum(InSurfacicCoordinates.Num());
+	for (int32 Index = 0; Index < InSurfacicCoordinates.Num(); Index++)
 	{
-		OutPoint3D.SetNum(InSurfacicCoordinates.Num());
-		for (int32 Index = 0; Index < InSurfacicCoordinates.Num(); ++Index)
-		{
-			EvaluatePoint(InSurfacicCoordinates[Index], OutPoint3D[Index], InDerivativeOrder);
-		}
+		EvaluatePoint(InSurfacicCoordinates[Index].Point, OutPoint3D[Index], InDerivativeOrder);
 	}
+}
 
-	void FSurface::EvaluatePoints(const TArray<FCurvePoint2D>& InSurfacicCoordinates, TArray<FSurfacicPoint>& OutPoint3D, int32 InDerivativeOrder) const
+void FSurface::EvaluatePoints(const TArray<FCurvePoint2D>& InSurfacicCoordinates, TArray<FCurvePoint>& OutPoints3D, int32 InDerivativeOrder) const
+{
+	OutPoints3D.SetNum(InSurfacicCoordinates.Num());
+	TArray<FSurfacicPoint> SurfacicPoints3D;
+	EvaluatePoints(InSurfacicCoordinates, SurfacicPoints3D, InDerivativeOrder);
+	for (int32 Index = 0; Index < InSurfacicCoordinates.Num(); ++Index)
 	{
-		OutPoint3D.SetNum(InSurfacicCoordinates.Num());
-		for (int32 Index = 0; Index < InSurfacicCoordinates.Num(); Index++)
-		{
-			EvaluatePoint(InSurfacicCoordinates[Index].Point, OutPoint3D[Index], InDerivativeOrder);
-		}
+		OutPoints3D[Index].Combine(InSurfacicCoordinates[Index], SurfacicPoints3D[Index]);
 	}
+}
 
-	void FSurface::EvaluatePoints(const TArray<FCurvePoint2D>& InSurfacicCoordinates, TArray<FCurvePoint>& OutPoints3D, int32 InDerivativeOrder) const
+void FSurface::EvaluatePoints(FSurfacicPolyline& Polyline) const
+{
+	int32 DerivativeOrder = Polyline.bWithNormals ? 1 : 0;
+
+	TArray<FSurfacicPoint> Point3D;
+	EvaluatePoints(Polyline.Points2D, Point3D, DerivativeOrder);
+
+	Polyline.Points3D.Empty(Polyline.Points2D.Num());
+	for (FSurfacicPoint Point : Point3D)
 	{
-		OutPoints3D.SetNum(InSurfacicCoordinates.Num());
-		TArray<FSurfacicPoint> SurfacicPoints3D;
-		EvaluatePoints(InSurfacicCoordinates, SurfacicPoints3D, InDerivativeOrder);
-		for(int32 Index = 0; Index < InSurfacicCoordinates.Num(); ++Index)
-		{
-			OutPoints3D[Index].Combine(InSurfacicCoordinates[Index], SurfacicPoints3D[Index]);
-		}
+		Polyline.Points3D.Emplace(Point.Point);
 	}
 
-	void FSurface::EvaluatePoints(FSurfacicPolyline& Polyline) const
+	if (Polyline.bWithNormals)
 	{
-		int32 DerivativeOrder = Polyline.bWithNormals ? 1 : 0;
-
-		TArray<FSurfacicPoint> Point3D;
-		EvaluatePoints(Polyline.Points2D, Point3D, DerivativeOrder);
-
-		Polyline.Points3D.Empty(Polyline.Points2D.Num());
+		Polyline.Normals.Empty(Polyline.Points2D.Num());
 		for (FSurfacicPoint Point : Point3D)
 		{
-			Polyline.Points3D.Emplace(Point.Point);
+			Polyline.Normals.Emplace(Point.GradientU ^ Point.GradientV);
 		}
-
-		if (Polyline.bWithNormals)
+		for (FVector Normal : Polyline.Normals)
 		{
-			Polyline.Normals.Empty(Polyline.Points2D.Num());
-			for (FSurfacicPoint Point : Point3D)
-			{
-				Polyline.Normals.Emplace(Point.GradientU ^ Point.GradientV);
-			}
-			for (FVector Normal : Polyline.Normals)
-			{
-				Normal.Normalize();
-			}
+			Normal.Normalize();
 		}
 	}
+}
 
-	void FSurface::EvaluatePoints(const TArray<FCurvePoint2D>& InPoints2D, FSurfacicPolyline& Polyline) const
-	{
-		int32 DerivativeOrder = 1;
-
-		TArray<FSurfacicPoint> Point3D;
-		EvaluatePoints(InPoints2D, Point3D, DerivativeOrder);
+void FSurface::EvaluatePoints(const TArray<FCurvePoint2D>& InPoints2D, FSurfacicPolyline& Polyline) const
+{
+	int32 DerivativeOrder = 1;
 
-		Polyline.Points2D.Empty(InPoints2D.Num());
-		Polyline.Points3D.Empty(InPoints2D.Num());
-		Polyline.Tangents.Empty(InPoints2D.Num());
-		Polyline.Normals.Empty(InPoints2D.Num());
+	TArray<FSurfacicPoint> Points3D;
+	EvaluatePoints(InPoints2D, Points3D, DerivativeOrder);
 
-		for (const FCurvePoint2D& Point : InPoints2D)
-		{
-			Polyline.Points2D.Emplace(Point.Point);
-		}
+	Polyline.Points2D.Empty(InPoints2D.Num());
+	Polyline.Points3D.Empty(InPoints2D.Num());
+	Polyline.Tangents.Empty(InPoints2D.Num());
+	Polyline.Normals.Empty(InPoints2D.Num());
 
-		for (const FSurfacicPoint& Point : Point3D)
-		{
-			Polyline.Points3D.Emplace(Point.Point);
-		}
+	for (const FCurvePoint2D& Point : InPoints2D)
+	{
+		Polyline.Points2D.Emplace(Point.Point);
+	}
 
-		for (int32 Index = 0; Index < InPoints2D.Num(); ++Index)
-		{
-			Polyline.Tangents.Emplace(Point3D[Index].GradientU * InPoints2D[Index].Gradient.U + Point3D[Index].GradientV * InPoints2D[Index].Gradient.V);
-		}
+	for (const FSurfacicPoint& Point : Points3D)
+	{
+		Polyline.Points3D.Emplace(Point.Point);
+	}
 
-		for (const FSurfacicPoint& Point : Point3D)
-		{
-			Polyline.Normals.Emplace(Point.GradientU ^ Point.GradientV);
-		}
+	for (int32 Index = 0; Index < InPoints2D.Num(); ++Index)
+	{
+		Polyline.Tangents.Emplace(Points3D[Index].GradientU * InPoints2D[Index].Gradient.U + Points3D[Index].GradientV * InPoints2D[Index].Gradient.V);
+	}
 
-		for (FVector& Normal : Polyline.Normals)
-		{
-			Normal.Normalize();
-		}
+	for (const FSurfacicPoint& Point : Points3D)
+	{
+		Polyline.Normals.Emplace(Point.GradientU ^ Point.GradientV);
 	}
 
-	void FSurface::EvaluateNormals(const TArray<FPoint2D>& InPoints2D, TArray<FVector>& Normals) const
+	for (FVector& Normal : Polyline.Normals)
 	{
-		int32 DerivativeOrder = 1;
-		TArray<FSurfacicPoint> Point3D;
-		EvaluatePoints(InPoints2D, Point3D, DerivativeOrder);
+		Normal.Normalize();
+	}
+}
 
-		Normals.Empty(InPoints2D.Num());
+void FSurface::EvaluateNormals(const TArray<FPoint2D>& InPoints2D, TArray<FVector>& Normals) const
+{
+	int32 DerivativeOrder = 1;
+	TArray<FSurfacicPoint> Points3D;
+	EvaluatePoints(InPoints2D, Points3D, DerivativeOrder);
 
-		for (const FSurfacicPoint& Point : Point3D)
-		{
-			Normals.Emplace(Point.GradientU ^ Point.GradientV);
-		}
+	Normals.Empty(InPoints2D.Num());
 
-		for (FVector& Normal : Normals)
-		{
-			Normal.Normalize();
-		}
+	for (const FSurfacicPoint& Point : Points3D)
+	{
+		Normals.Emplace(Point.GradientU ^ Point.GradientV);
 	}
 
-	void FSurface::EvaluatePointGrid(const FCoordinateGrid& Coordinates, FSurfacicSampling& OutPoints, bool bComputeNormals) const
+	for (FVector& Normal : Normals)
 	{
-		OutPoints.Reserve(Coordinates.Count());
+		Normal.Normalize();
+	}
+}
 
-		OutPoints.Set2DCoordinates(Coordinates);
+void FSurface::EvaluatePointGrid(const FCoordinateGrid& Coordinates, FSurfacicSampling& OutPoints, bool bComputeNormals) const
+{
+	OutPoints.Reserve(Coordinates.Count());
 
-		int32 DerivativeOrder = bComputeNormals ? 1 : 0;
-		TArray<FSurfacicPoint> Point3D;
-		EvaluatePoints(OutPoints.Points2D, Point3D, DerivativeOrder);
+	OutPoints.Set2DCoordinates(Coordinates);
 
-		OutPoints.bWithNormals = bComputeNormals;
+	int32 DerivativeOrder = bComputeNormals ? 1 : 0;
+	TArray<FSurfacicPoint> Point3D;
+	EvaluatePoints(OutPoints.Points2D, Point3D, DerivativeOrder);
 
-		for (FSurfacicPoint Point : Point3D)
-		{
-			OutPoints.Points3D.Emplace(Point.Point);
-		}
+	OutPoints.bWithNormals = bComputeNormals;
 
-		if (bComputeNormals)
+	for (FSurfacicPoint Point : Point3D)
+	{
+		OutPoints.Points3D.Emplace(Point.Point);
+	}
+
+	if (bComputeNormals)
+	{
+		for (FSurfacicPoint Point : Point3D)
 		{
-			for (FSurfacicPoint Point : Point3D)
-			{
-				OutPoints.Normals.Emplace(Point.GradientU ^ Point.GradientV);
-			}
-			OutPoints.NormalizeNormals();
+			OutPoints.Normals.Emplace(Point.GradientU ^ Point.GradientV);
 		}
+		OutPoints.NormalizeNormals();
 	}
+}
 
-	void FSurface::EvaluateGrid(FGrid& Grid) const
-	{
-		FSurfacicSampling OutPoints;
+void FSurface::EvaluateGrid(FGrid& Grid) const
+{
+	FSurfacicSampling OutPoints;
 
-		const FCoordinateGrid& CoordinateGrid = Grid.GetCuttingCoordinates();
-		EvaluatePointGrid(CoordinateGrid, OutPoints, true);
+	const FCoordinateGrid& CoordinateGrid = Grid.GetCuttingCoordinates();
+	EvaluatePointGrid(CoordinateGrid, OutPoints, true);
 
-		Grid.GetInner3DPoints() = MoveTemp(OutPoints.Points3D);
-		Grid.GetNormals() = MoveTemp(OutPoints.Normals);
-	}
+	Grid.GetInner3DPoints() = MoveTemp(OutPoints.Points3D);
+	Grid.GetNormals() = MoveTemp(OutPoints.Normals);
+}
 
 } // namespace CADKernel
 
diff --git a/Engine/Source/Runtime/Datasmith/CADKernel/Private/CADKernel/Mesh/Criteria/CriteriaGrid.cpp b/Engine/Source/Runtime/Datasmith/CADKernel/Private/CADKernel/Mesh/Criteria/CriteriaGrid.cpp
--- a/Engine/Source/Runtime/Datasmith/CADKernel/Private/CADKernel/Mesh/Criteria/CriteriaGrid.cpp
+++ b/Engine/Source/Runtime/Datasmith/CADKernel/Private/CADKernel/Mesh/Criteria/CriteriaGrid.cpp
@@ -37,19 +37,19 @@ void FCriteriaGrid::Init()
 	};
 
 	FCoordinateGrid CoordinateGrid2;
-	ComputeMiddlePointsCoordinates(Surface->GetCrossingPointCoordinates(EIso::IsoU), CoordinateGrid2[EIso::IsoU]);
-	ComputeMiddlePointsCoordinates(Surface->GetCrossingPointCoordinates(EIso::IsoV), CoordinateGrid2[EIso::IsoV]);
+	ComputeMiddlePointsCoordinates(Face.GetCrossingPointCoordinates(EIso::IsoU), CoordinateGrid2[EIso::IsoU]);
+	ComputeMiddlePointsCoordinates(Face.GetCrossingPointCoordinates(EIso::IsoV), CoordinateGrid2[EIso::IsoV]);
 
-	Surface->GetCarrierSurface()->EvaluatePointGrid(CoordinateGrid2, Grid, true);
+	Face.GetCarrierSurface()->EvaluatePointGrid(CoordinateGrid2, Grid, true);
 	TrueUcoorindateCount = CoordinateGrid2[EIso::IsoU].Num();
 }
 
-FCriteriaGrid::FCriteriaGrid(TSharedRef<FTopologicalFace> InSurface)
-	: Surface(InSurface)
-	, CoordinateGrid(InSurface->GetCrossingPointCoordinates())
+FCriteriaGrid::FCriteriaGrid(FTopologicalFace& InSurface)
+	: Face(InSurface)
+	, CoordinateGrid(InSurface.GetCrossingPointCoordinates())
 {
-	Surface->Presample();
-	Surface->InitDeltaUs();
+	Face.Presample();
+	Face.InitDeltaUs();
 	Init();
 #ifdef DISPLAY_CRITERIA_GRID
 	Display();
@@ -58,11 +58,11 @@ FCriteriaGrid::FCriteriaGrid(TSharedRef<FTopologicalFace> InSurface)
 
 void FCriteriaGrid::ApplyCriteria(const TArray<TSharedPtr<FCriterion>>& Criteria) const
 {
-	TArray<double>& DeltaUMaxArray = Surface->GetCrossingPointDeltaMaxs(EIso::IsoU);
-	TArray<double>& DeltaUMiniArray = Surface->GetCrossingPointDeltaMins(EIso::IsoU);
-	TArray<double>& DeltaVMaxArray = Surface->GetCrossingPointDeltaMaxs(EIso::IsoV);
-	TArray<double>& DeltaVMinArray = Surface->GetCrossingPointDeltaMins(EIso::IsoV);
-	FSurfaceCurvature& SurfaceCurvature = Surface->GetCurvatures();
+	TArray<double>& DeltaUMaxArray = Face.GetCrossingPointDeltaMaxs(EIso::IsoU);
+	TArray<double>& DeltaUMiniArray = Face.GetCrossingPointDeltaMins(EIso::IsoU);
+	TArray<double>& DeltaVMaxArray = Face.GetCrossingPointDeltaMaxs(EIso::IsoV);
+	TArray<double>& DeltaVMinArray = Face.GetCrossingPointDeltaMins(EIso::IsoV);
+	FSurfaceCurvature& SurfaceCurvature = Face.GetCurvatures();
 
 	for (int32 IndexV = 0; IndexV < GetCoordinateCount(EIso::IsoV) - 1; ++IndexV)
 	{
@@ -156,7 +156,7 @@ void FCriteriaGrid::Display() const
 
 	{
 		F3DDebugSession A(TEXT("Loop 3D"));
-		for (const TSharedPtr<FTopologicalLoop>& Loop : Surface->GetLoops())
+		for (const TSharedPtr<FTopologicalLoop>& Loop : Face.GetLoops())
 		{
 			CADKernel::Display(*Loop);
 		}
@@ -164,7 +164,7 @@ void FCriteriaGrid::Display() const
 
 	{
 		F3DDebugSession A(TEXT("Loop 2D"));
-		for (const TSharedPtr<FTopologicalLoop>& Loop : Surface->GetLoops())
+		for (const TSharedPtr<FTopologicalLoop>& Loop : Face.GetLoops())
 		{
 			CADKernel::Display2D(*Loop);
 		}
diff --git a/Engine/Source/Runtime/Datasmith/CADKernel/Private/CADKernel/Mesh/Criteria/CriteriaGrid.h b/Engine/Source/Runtime/Datasmith/CADKernel/Private/CADKernel/Mesh/Criteria/CriteriaGrid.h
--- a/Engine/Source/Runtime/Datasmith/CADKernel/Private/CADKernel/Mesh/Criteria/CriteriaGrid.h
+++ b/Engine/Source/Runtime/Datasmith/CADKernel/Private/CADKernel/Mesh/Criteria/CriteriaGrid.h
@@ -14,7 +14,7 @@ namespace CADKernel
 	class FCriteriaGrid
 	{
 	protected:
-		TSharedRef<FTopologicalFace> Surface;
+		FTopologicalFace& Face;
 
 		const FCoordinateGrid& CoordinateGrid;
 		int32 TrueUcoorindateCount; 
@@ -39,11 +39,11 @@ namespace CADKernel
 		void Init();
 
 	public:
-		FCriteriaGrid(TSharedRef<FTopologicalFace> Surface);
+		FCriteriaGrid(FTopologicalFace& Surface);
 
-		TSharedRef<FTopologicalFace> GetSurface()
+		FTopologicalFace& GetSurface()
 		{
-			return Surface;
+			return Face;
 		}
 
 		void ApplyCriteria(const TArray<TSharedPtr<FCriterion>>& GetCriteria) const;
diff --git a/Engine/Source/Runtime/Datasmith/CADKernel/Private/CADKernel/Mesh/Meshers/ParametricMesher.cpp b/Engine/Source/Runtime/Datasmith/CADKernel/Private/CADKernel/Mesh/Meshers/ParametricMesher.cpp
--- a/Engine/Source/Runtime/Datasmith/CADKernel/Private/CADKernel/Mesh/Meshers/ParametricMesher.cpp
+++ b/Engine/Source/Runtime/Datasmith/CADKernel/Private/CADKernel/Mesh/Meshers/ParametricMesher.cpp
@@ -26,26 +26,28 @@ FMesherParameters::FMesherParameters()
 	: InconsistencyAngle(TEXT("inconsistencyAngle"), 20., *this)
 {}
 
-FParametricMesher::FParametricMesher(TSharedRef<FModelMesh> InMeshModel)
+FParametricMesher::FParametricMesher(FModelMesh& InMeshModel)
 	: MeshModel(InMeshModel)
-	, Parameters(MakeShared<FMesherParameters>())
 {
 }
 
-void FParametricMesher::MeshEntities(TArray<TSharedPtr<FEntity>>& InEntities)
+void FParametricMesher::MeshEntities(TArray<FTopologicalShapeEntity*>& InEntities)
 {
 	int32 FaceCount = 0;
 
-	for (TSharedPtr<FTopologicalFace>& Face : Faces)
+	for (FTopologicalFace* Face : Faces)
 	{
+		if (Face == nullptr)
+		{
+			continue;
+		}
 		Face->SetMarker1();
 	}
 
 	// count faces
-	for (TSharedPtr<FEntity>& Entity : InEntities)
+	for (FTopologicalShapeEntity* TopologicalEntity : InEntities)
 	{
-		TSharedPtr<FTopologicalEntity> TopologicalEntity = StaticCastSharedPtr<FTopologicalEntity>(Entity);
-		if (!TopologicalEntity.IsValid())
+		if (TopologicalEntity == nullptr)
 		{
 			continue;
 		}
@@ -53,16 +55,19 @@ void FParametricMesher::MeshEntities(TArray<TSharedPtr<FEntity>>& InEntities)
 	}
 	Faces.Reserve(Faces.Num() + FaceCount);
 
-	for (TSharedPtr<FTopologicalFace>& Face : Faces)
+	for (FTopologicalFace* Face : Faces)
 	{
+		if (Face == nullptr)
+		{
+			continue;
+		}
 		Face->ResetMarkers();
 	}
 
 	// Get independent Faces and spread body's shells orientation
-	for (TSharedPtr<FEntity>& Entity : InEntities)
+	for (FTopologicalShapeEntity* TopologicalEntity : InEntities)
 	{
-		TSharedPtr<FTopologicalEntity> TopologicalEntity = StaticCastSharedPtr<FTopologicalEntity>(Entity);
-		if (!TopologicalEntity.IsValid())
+		if (TopologicalEntity == nullptr)
 		{
 			continue;
 		}
@@ -71,32 +76,13 @@ void FParametricMesher::MeshEntities(TArray<TSharedPtr<FEntity>>& InEntities)
 		TopologicalEntity->GetFaces(Faces);
 	}
 
-	for (TSharedPtr<FTopologicalFace>& Face : Faces)
-	{
-		Face->ResetMarkers();
-	}
-
-	// Get independent elementary entities (Edge, Vertex)
-	for (TSharedPtr<FEntity>& Entity : InEntities)
+	for (FTopologicalFace* Face : Faces)
 	{
-		if (!Entity.IsValid())
+		if (Face == nullptr)
 		{
 			continue;
 		}
-
-		switch (Entity->GetEntityType())
-		{
-		case EEntity::TopologicalEdge:
-			Edges.Add(StaticCastSharedPtr<FTopologicalEdge>(Entity));
-			break;
-
-		case EEntity::TopologicalVertex:
-			Vertices.Add(StaticCastSharedPtr<FTopologicalVertex>(Entity));
-			break;
-
-		default:
-			break;
-		}
+		Face->ResetMarkers();
 	}
 
 	MeshEntities();
@@ -115,14 +101,16 @@ void FParametricMesher::MeshEntities()
 	//      Apply Surface Criteria
 	// ============================================================================================================
 
-	for (TSharedPtr<FTopologicalFace> Face : Faces)
+	for (FTopologicalFace* Face : Faces)
 	{
 		FProgress _(1, TEXT("Meshing Entities : Apply Surface Criteria"));
 
-		ensureCADKernel(Face.IsValid());
-		ensureCADKernel(!Face->IsDeleted());
+		if(Face == nullptr || Face->IsDeleted())
+		{
+			continue;
+		}
 
-		ApplyFaceCriteria(Face.ToSharedRef());
+		ApplyFaceCriteria(*Face);
 	}
 
 	MesherReport.Chronos.ApplyCriteriaDuration = FChrono::Elapse(ApplyCriteriaStartTime);
@@ -160,18 +148,18 @@ void FParametricMesher::MeshEntities()
 	//Chronos.PrintTimeElapse();
 }
 
-void FParametricMesher::ApplyFaceCriteria(TSharedRef<FTopologicalFace> Face)
+void FParametricMesher::ApplyFaceCriteria(FTopologicalFace& Face)
 {
-	if (Face->IsApplyCriteria())
+	if (Face.IsApplyCriteria())
 	{
 		return;
 	}
 
 	FCriteriaGrid Grid(Face);
-	Grid.ApplyCriteria(GetMeshModel()->GetCriteria());
+	Grid.ApplyCriteria(GetMeshModel().GetCriteria());
 
-	Face->ChooseFinalDeltaUs();
-	Face->SetApplyCriteria();
+	Face.ChooseFinalDeltaUs();
+	Face.SetApplyCriteria();
 }
 
 void FParametricMesher::ApplyEdgeCriteria(FTopologicalEdge& Edge)
@@ -195,7 +183,7 @@ void FParametricMesher::ApplyEdgeCriteria(FTopologicalEdge& Edge)
 	TArray<FCurvePoint> Points3D;
 	Edge.EvaluatePoints(Coordinates, 0, Points3D);
 
-	const TArray<TSharedPtr<FCriterion>>& Criteria = GetMeshModel()->GetCriteria();
+	const TArray<TSharedPtr<FCriterion>>& Criteria = GetMeshModel().GetCriteria();
 	for (const TSharedPtr<FCriterion>& Criterion : Criteria)
 	{
 		Criterion->ApplyOnEdgeParameters(Edge, CrossingPointUs, Points3D);
@@ -240,7 +228,6 @@ void FParametricMesher::Mesh(FTopologicalFace& Face)
 		return;
 	}
 
-
 	FTimePoint IsoTriangulerStartTime = FChrono::Now();
 
 	TSharedRef<FFaceMesh> SurfaceMesh = StaticCastSharedRef<FFaceMesh>(Face.GetOrCreateMesh(MeshModel));
@@ -252,7 +239,7 @@ void FParametricMesher::Mesh(FTopologicalFace& Face)
 		{
 			SurfaceMesh->InverseOrientation();
 		}
-		MeshModel->AddMesh(SurfaceMesh);
+		MeshModel.AddMesh(SurfaceMesh);
 	}
 	Face.SetMeshed();
 
@@ -761,7 +748,7 @@ void FParametricMesher::GenerateEdgeElements(FTopologicalEdge& Edge)
 
 	EdgeMesh->RegisterCoordinates();
 	EdgeMesh->Mesh(StartVertexNodeIndex, EndVertexNodeIndex);
-	MeshModel->AddMesh(EdgeMesh);
+	MeshModel.AddMesh(EdgeMesh);
 	ActiveEdge.SetMeshed();
 }
 
@@ -772,9 +759,9 @@ void FParametricMesher::IsolateQuadFace(TArray<FCostToFace>& QuadSurfaces, TArra
 	QuadSurfaces.Reserve(Faces.Num() * 2);
 	OtherSurfaces.Reserve(Faces.Num());
 
-	for (const TSharedPtr<FTopologicalFace>& Face : Faces)
+	for (FTopologicalFace* Face : Faces)
 	{
-		if (Face->IsDeleted() || Face->IsMeshed())
+		if (Face == nullptr || Face->IsDeleted() || Face->IsMeshed())
 		{
 			continue;
 		}
@@ -782,8 +769,13 @@ void FParametricMesher::IsolateQuadFace(TArray<FCostToFace>& QuadSurfaces, TArra
 		Face->ComputeSurfaceSideProperties();
 	}
 
-	for (const TSharedPtr<FTopologicalFace>& FacePtr : Faces)
+	for (FTopologicalFace* FacePtr : Faces)
 	{
+		if (FacePtr == nullptr)
+		{
+			continue;
+		}
+
 		FTopologicalFace& Face = *FacePtr;
 		Face.DefineSurfaceType();
 		switch (Face.GetQuadType())
@@ -860,7 +852,7 @@ void FParametricMesher::IsolateQuadFace(TArray<FCostToFace>& QuadSurfaces, TArra
 
 void FParametricMesher::LinkQuadSurfaceForMesh(TArray<FCostToFace>& QuadTrimmedSurfaceSet, TArray<TArray<FTopologicalFace*>>& OutStrips)
 {
-	const double GeometricTolerance = 20. * MeshModel->GetGeometricTolerance();
+	const double GeometricTolerance = 20. * MeshModel.GetGeometricTolerance();
 
 	OutStrips.Reserve(QuadTrimmedSurfaceSet.Num());
 
@@ -1009,8 +1001,12 @@ void FParametricMesher::LinkQuadSurfaceForMesh(TArray<FCostToFace>& QuadTrimmedS
 		}
 	}
 
-	for (const TSharedPtr<FTopologicalFace>& Face : Faces)
+	for (FTopologicalFace* Face : Faces)
 	{
+		if (Face == nullptr)
+		{
+			continue;
+		}
 		Face->ResetMarkers();
 	}
 }
@@ -1023,12 +1019,16 @@ void FParametricMesher::MeshSurfaceByFront(TArray<FCostToFace>& QuadTrimmedSurfa
 
 	FMessage::Printf(EVerboseLevel::Debug, TEXT("Start MeshSurfaceByFront\n"));
 
-	for (TSharedPtr<FTopologicalFace> Face : Faces)
+	for (FTopologicalFace* Face : Faces)
 	{
+		if (Face == nullptr)
+		{
+			continue;
+		}
 		Face->SetMarker3();
 	}
 
-	const double GeometricTolerance = 20. * MeshModel->GetGeometricTolerance();
+	const double GeometricTolerance = 20. * MeshModel.GetGeometricTolerance();
 
 	TArray<FTopologicalFace*> CandidateFacesForMesh; // first in first out
 	CandidateFacesForMesh.Reserve(100);
@@ -1066,7 +1066,7 @@ void FParametricMesher::MeshSurfaceByFront(TArray<FCostToFace>& QuadTrimmedSurfa
 		{
 			const TSharedPtr<FTopologicalEdge> Edge = OrientedEdge.Entity;
 			Edge->SetMarker1(); // tmp for debug
-			for (FTopologicalEdge* NextEdge : Edge->GetTwinsEntities())
+			for (FTopologicalEdge* NextEdge : Edge->GetTwinEntities())
 			{
 				if (NextEdge->HasMarker1())
 				{
@@ -1217,9 +1217,9 @@ void FParametricMesher::MeshSurfaceByFront(TArray<FCostToFace>& QuadTrimmedSurfa
 	}
 
 	// the the other surface
-	for (TSharedPtr<FTopologicalFace> Face : Faces)
+	for (FTopologicalFace* Face : Faces)
 	{
-		if (!Face->IsMeshed())
+		if (Face != nullptr && !Face->IsMeshed())
 		{
 			MeshFacesByFront(*Face);
 		}
@@ -1595,12 +1595,6 @@ void FParametricMesher::MeshThinZoneSide(const FThinZoneSide& Side)
 
 }
 
-
-void FParametricMesher::InitParameters(const FString& ParametersString)
-{
-	Parameters->SetFromString(ParametersString);
-}
-
 void FParametricMesher::PrintReport()
 {
 	MesherReport.Print();
diff --git a/Engine/Source/Runtime/Datasmith/CADKernel/Private/CADKernel/Mesh/Structure/Grid.cpp b/Engine/Source/Runtime/Datasmith/CADKernel/Private/CADKernel/Mesh/Structure/Grid.cpp
--- a/Engine/Source/Runtime/Datasmith/CADKernel/Private/CADKernel/Mesh/Structure/Grid.cpp
+++ b/Engine/Source/Runtime/Datasmith/CADKernel/Private/CADKernel/Mesh/Structure/Grid.cpp
@@ -15,7 +15,7 @@
 //#define DEBUG_GETPREFERREDUVCOORDINATESFROMNEIGHBOURS
 namespace CADKernel
 {
-FGrid::FGrid(FTopologicalFace& InFace, TSharedRef<FModelMesh>& InMeshModel)
+FGrid::FGrid(FTopologicalFace& InFace, FModelMesh& InMeshModel)
 	: Face(InFace)
 	, FaceTolerance(InFace.GetIsoTolerances())
 	, Tolerance3D(InFace.GetCarrierSurface()->Get3DTolerance())
@@ -347,7 +347,10 @@ bool FGrid::GeneratePointCloud()
 
 	ComputeMaxElementSize();
 
-	ScaleGrid();
+	if (!ScaleGrid())
+	{
+		return false;
+	}
 
 	Chronos.GeneratePointCloudDuration += FChrono::Elapse(StartTime);
 	return true;
@@ -993,7 +996,7 @@ bool FGrid::GetMeshOfLoops()
 				if (EdgeCuttingPointCoordinates.IsEmpty())
 				{
 					int32 CuttingPointCount = 2;
-					for (const FTopologicalEdge* TwinEdge : Edge->GetTwinsEntities())
+					for (const FTopologicalEdge* TwinEdge : Edge->GetTwinEntities())
 					{
 						int32 TwinCuttingPointCount = TwinEdge->GetCuttingPoints().Num();
 						if (TwinCuttingPointCount > CuttingPointCount)
@@ -1181,7 +1184,7 @@ bool FGrid::GetMeshOfLoops()
 	return true;
 }
 
-void FGrid::ScaleGrid()
+bool FGrid::ScaleGrid()
 {
 	FTimePoint StartTime = FChrono::Now();
 
@@ -1260,9 +1263,14 @@ void FGrid::ScaleGrid()
 	}
 
 	double MeanLengthV = GetMean(LengthsV);
-	double FactorV = MeanLengthV / (CuttingCoordinates[EIso::IsoV].Last() - CuttingCoordinates[EIso::IsoV][0]);
-
 	double MeanLengthU = GetMean(LengthsU);
+	if(MeanLengthV < Tolerance3D || MeanLengthU < Tolerance3D)
+	{
+		SetAsDegenerated();
+		return false;
+	}
+
+	double FactorV = MeanLengthV / (CuttingCoordinates[EIso::IsoV].Last() - CuttingCoordinates[EIso::IsoV][0]);
 	double FactorU = MeanLengthU / (CuttingCoordinates[EIso::IsoU].Last() - CuttingCoordinates[EIso::IsoU][0]);
 
 	//TArray<double> ScaledCuttingCoordinates;
@@ -1360,6 +1368,7 @@ void FGrid::ScaleGrid()
 		}
 	}
 	Chronos.ScaleGridDuration = FChrono::Elapse(StartTime);
+	return true;
 }
 
 void FGrid::TransformPoints(EGridSpace DestinationSpace, const TArray<FPoint2D>& InPointsToScale, TArray<FPoint2D>& OutTransformedPoints) const
diff --git a/Engine/Source/Runtime/Datasmith/CADKernel/Private/CADKernel/Mesh/Structure/Grid.h b/Engine/Source/Runtime/Datasmith/CADKernel/Private/CADKernel/Mesh/Structure/Grid.h
--- a/Engine/Source/Runtime/Datasmith/CADKernel/Private/CADKernel/Mesh/Structure/Grid.h
+++ b/Engine/Source/Runtime/Datasmith/CADKernel/Private/CADKernel/Mesh/Structure/Grid.h
@@ -76,7 +76,7 @@ protected:
 	const double Tolerance3D;
 	const double MinimumElementSize;
 
-	TSharedRef<FModelMesh> MeshModel;
+	FModelMesh& MeshModel;
 
 	FThinZone2DFinder ThinZoneFinder;
 
@@ -157,7 +157,7 @@ public:
 	FGridChronos Chronos;
 
 public:
-	FGrid(FTopologicalFace& InFace, TSharedRef<FModelMesh>& InShellMesh);
+	FGrid(FTopologicalFace& InFace, FModelMesh& InShellMesh);
 
 #ifndef CADKERNEL_DEV
 	virtual ~FGrid() = default;
@@ -210,8 +210,9 @@ protected:
 	 * Builds the scaled parametric spaces
 	 * @see Points2D
 	 * @see GeneratePointCloud (ended GeneratePointCloud process)
+	 * @return false if the scaled grid is degenerated
 	 */
-	void ScaleGrid();
+	bool ScaleGrid();
 
 	/**
 	 * Projects loop's points in the scaled parametric spaces
diff --git a/Engine/Source/Runtime/Datasmith/CADKernel/Private/CADKernel/Topo/Body.cpp b/Engine/Source/Runtime/Datasmith/CADKernel/Private/CADKernel/Topo/Body.cpp
--- a/Engine/Source/Runtime/Datasmith/CADKernel/Private/CADKernel/Topo/Body.cpp
+++ b/Engine/Source/Runtime/Datasmith/CADKernel/Private/CADKernel/Topo/Body.cpp
@@ -1,6 +1,7 @@
 // Copyright Epic Games, Inc. All Rights Reserved.
 #include "CADKernel/Topo/Body.h"
 
+#include "CADKernel/Topo/Model.h"
 #include "CADKernel/Topo/Shell.h"
 #include "CADKernel/Topo/TopologyReport.h"
 
@@ -11,10 +12,10 @@ namespace CADKernel
 void FBody::AddShell(TSharedRef<FShell> Shell)
 {
 	Shells.Add(Shell);
-	Shell->HostedBy = this;
+	Shell->SetHost(this);
 }
 
-void FBody::RemoveEmptyShell()
+void FBody::RemoveEmptyShell(FModel& Model)
 {
 	TArray<TSharedPtr<FShell>> NewShells;
 	NewShells.Reserve(Shells.Num());
@@ -29,23 +30,21 @@ void FBody::RemoveEmptyShell()
 			Shell->Delete();
 		}
 	}
-	Swap(NewShells, Shells);
-}
-
-TSharedPtr<FEntityGeom> FBody::ApplyMatrix(const FMatrixH& InMatrix) const 
-{
-	TArray<TSharedPtr<FShell>> NewShells;
-	for (TSharedPtr<FShell> Shell : Shells)
+	if (NewShells.IsEmpty())
+	{
+		Delete();
+		Model.RemoveBody(this);
+	}
+	else
 	{
-		NewShells.Add(StaticCastSharedPtr<FShell>(Shell->ApplyMatrix(InMatrix)));
+		Swap(NewShells, Shells);
 	}
-	return FEntity::MakeShared<FBody>(NewShells);
 }
 
 #ifdef CADKERNEL_DEV
 FInfoEntity& FBody::GetInfo(FInfoEntity& Info) const
 {
-	return FTopologicalEntity::GetInfo(Info).Add(TEXT("shells"), Shells).Add(*this);
+	return FTopologicalShapeEntity::GetInfo(Info).Add(TEXT("Shells"), Shells);
 }
 #endif
 
diff --git a/Engine/Source/Runtime/Datasmith/CADKernel/Private/CADKernel/Topo/Model.cpp b/Engine/Source/Runtime/Datasmith/CADKernel/Private/CADKernel/Topo/Model.cpp
--- a/Engine/Source/Runtime/Datasmith/CADKernel/Private/CADKernel/Topo/Model.cpp
+++ b/Engine/Source/Runtime/Datasmith/CADKernel/Private/CADKernel/Topo/Model.cpp
@@ -3,7 +3,6 @@
 
 #include "CADKernel/Core/EntityGeom.h"
 #include "CADKernel/Topo/Body.h"
-#include "CADKernel/Topo/Topomaker.h"
 #include "CADKernel/Topo/TopologicalEdge.h"
 #include "CADKernel/Topo/TopologicalFace.h"
 #include "CADKernel/Topo/TopologicalVertex.h"
@@ -70,62 +69,6 @@ void FModel::PrintBodyAndShellCount()
 	FMessage::Printf(Log, TEXT("Body count %d shell count %d \n"), NbBody, NbShell);
 }
 
-void FModel::RemoveEmptyBodies()
-{
-	int32 NbBody = 0;
-	int32 NbShell = 0;
-
-	TArray<TSharedPtr<FBody>> NewBodies;
-	NewBodies.Reserve(Bodies.Num());
-	for (TSharedPtr<FBody> Body : Bodies)
-	{
-		Body->RemoveEmptyShell();
-		if (Body->GetShells().Num())
-		{
-			NbShell += Body->GetShells().Num();
-			NewBodies.Emplace(Body);
-			NbBody++;
-		}
-		else
-		{
-			Body->Delete();
-		}
-	}
-	Swap(NewBodies, Bodies);
-	FMessage::Printf(Log, TEXT("After RemoveEmptyBodies, Body count %d shell count %d \n"), NbBody, NbShell);
-}
-
-
-
-/*void FModel::RemoveDomain(TSharedPtr<FTopologicalFace> Domain)
-{
-	Faces.Remove(Domain);
-}
-
-void FModel::RemoveBody(TSharedPtr<FBody> Body)
-{
-	Bodies.Remove(Body);
-}*/
-
-TSharedPtr<FEntityGeom> FModel::ApplyMatrix(const FMatrixH& InMatrix) const
-{
-	TSharedPtr<FModel> Model = FEntity::MakeShared<FModel>();
-
-	for (TSharedPtr<FBody> Body : Bodies)
-	{
-		TSharedPtr<FBody> TransformedBody = StaticCastSharedPtr<FBody>(Body->ApplyMatrix(InMatrix));
-		Model->Add(TransformedBody);
-	}
-
-	for (TSharedPtr<FTopologicalFace> Domain : Faces)
-	{
-		TSharedPtr<FTopologicalFace> TransformedSurface = StaticCastSharedPtr<FTopologicalFace>(Domain->ApplyMatrix(InMatrix));
-		Model->Add(TransformedSurface);
-	}
-
-	return Model;
-}
-
 int32 FModel::FaceCount() const
 {
 	int32 FaceCount = 0;
@@ -137,7 +80,7 @@ int32 FModel::FaceCount() const
 	return FaceCount;
 }
 
-void FModel::GetFaces(TArray<TSharedPtr<FTopologicalFace>>& OutFaces) 
+void FModel::GetFaces(TArray<FTopologicalFace*>& OutFaces) 
 {
 	for (const TSharedPtr<FBody>& Body : Bodies)
 	{
@@ -148,7 +91,7 @@ void FModel::GetFaces(TArray<TSharedPtr<FTopologicalFace>>& OutFaces)
 	{
 		if (!Face->HasMarker1())
 		{
-			OutFaces.Emplace(Face);
+			OutFaces.Add(Face.Get());
 			Face->SetMarker1();
 		}
 	}
@@ -165,10 +108,9 @@ void FModel::SpreadBodyOrientation()
 #ifdef CADKERNEL_DEV
 FInfoEntity& FModel::GetInfo(FInfoEntity& Info) const
 {
-	return FEntityGeom::GetInfo(Info)
+	return FTopologicalShapeEntity::GetInfo(Info)
 		.Add(TEXT("Bodies"), Bodies)
-		.Add(TEXT("Domains"), Faces)
-		.Add(*this);
+		.Add(TEXT("Faces"), Faces);
 }
 #endif
 
@@ -264,9 +206,12 @@ void FModel::FillTopologyReport(FTopologyReport& Report) const
 	}
 }
 
-void FModel::Split(TSharedPtr<FBody> Body, TArray<TSharedPtr<FBody>>& OutNewBody)
+void FModel::Orient()
 {
-
+	for (TSharedPtr<FBody> Body : Bodies)
+	{
+		Body->Orient();
+	}
 }
 
 /**
diff --git a/Engine/Source/Runtime/Datasmith/CADKernel/Private/CADKernel/Topo/Shell.cpp b/Engine/Source/Runtime/Datasmith/CADKernel/Private/CADKernel/Topo/Shell.cpp
--- a/Engine/Source/Runtime/Datasmith/CADKernel/Private/CADKernel/Topo/Shell.cpp
+++ b/Engine/Source/Runtime/Datasmith/CADKernel/Private/CADKernel/Topo/Shell.cpp
@@ -10,7 +10,7 @@ namespace CADKernel
 {
 
 FShell::FShell(const TArray<TSharedPtr<FTopologicalFace>>& InTopologicalFaces, bool bIsInnerShell)
-	: FTopologicalEntity()
+	: FTopologicalShapeEntity()
 	, TopologicalFaces()
 {
 	TopologicalFaces.Reserve(InTopologicalFaces.Num());
@@ -29,9 +29,8 @@ FShell::FShell(const TArray<TSharedPtr<FTopologicalFace>>& InTopologicalFaces, b
 	}
 }
 
-
 FShell::FShell(const TArray<TSharedPtr<FTopologicalFace>>& InTopologicalFaces, const TArray<EOrientation>& InOrientations, bool bIsInnerShell)
-	: FTopologicalEntity()
+	: FTopologicalShapeEntity()
 {
 	ensureCADKernel(InTopologicalFaces.Num() == InOrientations.Num());
 	for (int32 Index = 0; Index < InTopologicalFaces.Num(); ++Index)
@@ -47,12 +46,6 @@ FShell::FShell(const TArray<TSharedPtr<FTopologicalFace>>& InTopologicalFaces, c
 	}
 }
 
-TSharedPtr<FEntityGeom> FShell::ApplyMatrix(const FMatrixH& InMatrix) const
-{
-	ensureCADKernel(false);
-	return TSharedPtr<FEntityGeom>();
-}
-
 void FShell::Empty(int32 NewSize)
 {
 	for(FOrientedFace& Face : TopologicalFaces)
@@ -86,14 +79,12 @@ void FShell::Add(TSharedRef<FTopologicalFace> InTopologicalFace, EOrientation Or
 #ifdef CADKERNEL_DEV
 FInfoEntity& FShell::GetInfo(FInfoEntity& Info) const
 {
-	return FEntity::GetInfo(Info)
-		.Add(TEXT("Hosted by"), (FEntity*) HostedBy)
-		.Add(TEXT("TopologicalFaces"), (TArray<TOrientedEntity<FEntity>>&) TopologicalFaces)
-		.Add(*this);
+	return FTopologicalShapeEntity::GetInfo(Info)
+		.Add(TEXT("TopologicalFaces"), (TArray<TOrientedEntity<FEntity>>&) TopologicalFaces);
 }
 #endif
 
-void FShell::GetFaces(TArray<TSharedPtr<FTopologicalFace>>& Faces)
+void FShell::GetFaces(TArray<FTopologicalFace*>& Faces)
 {
 	for (FOrientedFace& Face : TopologicalFaces)
 	{
@@ -102,7 +93,7 @@ void FShell::GetFaces(TArray<TSharedPtr<FTopologicalFace>>& Faces)
 			continue;
 		}
 
-		Faces.Emplace(Face.Entity);
+		Faces.Add(Face.Entity.Get());
 		Face.Entity->SetMarker1();
 	}
 }
@@ -119,17 +110,26 @@ void FShell::SpreadBodyOrientation()
 	bool bIsOutter = IsOutter();
 	for (FOrientedFace& Face : TopologicalFaces)
 	{
-		if (Face.Entity->HasMarker2())
+		if (bIsOutter != (Face.Direction == EOrientation::Front))
 		{
-			// the face is already processed, this should not append...
-			continue;
+			Face.Entity->SetBackOriented();
 		}
+	}
+}
 
-		if (bIsOutter != (Face.Direction == EOrientation::Front))
+void FShell::UpdateShellOrientation()
+{
+	bool bIsOutter = IsOutter();
+	for (FOrientedFace& Face : TopologicalFaces)
+	{
+		if (Face.Entity->IsBackOriented())
 		{
-			Face.Entity->SetBackOriented();
+			Face.Direction = EOrientation::Back;
+		}
+		else
+		{
+			Face.Direction = EOrientation::Front;
 		}
-		Face.Entity->SetMarker2();
 	}
 }
 
@@ -143,7 +143,7 @@ bool FShell::IsOpenShell()
 			for (const FOrientedEdge& OrientedEdge : Loop->GetEdges())
 			{
 				const TSharedPtr<FTopologicalEdge>& Edge = OrientedEdge.Entity;
-				if (Edge->GetTwinsEntityCount() == 1)
+				if (Edge->GetTwinEntityCount() == 1)
 				{
 					return true;
 				}
@@ -153,8 +153,7 @@ bool FShell::IsOpenShell()
 	return false;
 }
 
-
-void FShell::CheckTopology(TArray<FFaceSubset>& SubShells)
+void FShell::CheckTopology(TArray<FFaceSubset>& Subshells)
 {
 	// Processed1 : Surfaces added in CandidateSurfacesForMesh
 
@@ -178,7 +177,7 @@ void FShell::CheckTopology(TArray<FFaceSubset>& SubShells)
 				}
 				Edge->SetMarker1();
 
-				if (Edge->GetTwinsEntityCount() == 1)
+				if (Edge->GetTwinEntityCount() == 1)
 				{
 					if (!Edge->IsDegenerated())
 					{
@@ -187,12 +186,12 @@ void FShell::CheckTopology(TArray<FFaceSubset>& SubShells)
 					continue;
 				}
 
-				if (Edge->GetTwinsEntityCount() > 2)
+				if (Edge->GetTwinEntityCount() > 2)
 				{
 					Shell.NonManifoldEdgeCount++;
 				}
 
-				for (FTopologicalEdge* NextEdge : Edge->GetTwinsEntities())
+				for (FTopologicalEdge* NextEdge : Edge->GetTwinEntities())
 				{
 					if (NextEdge->HasMarker1())
 					{
@@ -229,7 +228,7 @@ void FShell::CheckTopology(TArray<FFaceSubset>& SubShells)
 			continue;
 		}
 	
-		FFaceSubset& Shell = SubShells.Emplace_GetRef();
+		FFaceSubset& Shell = Subshells.Emplace_GetRef();
 		Shell.Faces.Reserve(TopologicalFaceCount - ProcessFaceCount);
 		Front.Empty(TopologicalFaceCount);
 
@@ -247,28 +246,279 @@ void FShell::CheckTopology(TArray<FFaceSubset>& SubShells)
 		}
 	}
 
-	// reset Marker
+	ResetMarkersRecursively();
+}
+
+void FShell::FillTopologyReport(FTopologyReport& Report) const 
+{
+	Report.Add(this);
 	for (const FOrientedFace& OrientedFace : GetFaces())
 	{
-		const TSharedPtr<FTopologicalFace>& Face = OrientedFace.Entity;
-		Face->ResetMarkers();
-		for (const TSharedPtr<FTopologicalLoop>& Loop : Face->GetLoops())
+		OrientedFace.Entity->FillTopologyReport(Report);
+	}
+}
+
+int32 FShell::Orient()
+{
+	int32 TopologicalFaceCount = GetFaces().Num();
+
+	if (TopologicalFaceCount < 2)
+	{
+		return 0;
+	}
+
+	TArray<FTopologicalFace*> Subshell;
+	TArray<FTopologicalFace*> Front;
+
+	int32 ShellSwappedFaceCount = 0;
+
+	int32 SubshellFaceCount = 0;
+	int32 SubshellSwappedFaceCount = 0;
+	double BorderEdgesLength = 0;
+
+	TFunction<void(const FTopologicalFace&)> GetAndOrientNeighboringFaces = [&](const FTopologicalFace& Face)
+	{
+		for (const TSharedPtr<FTopologicalLoop>& Loop : Face.GetLoops())
 		{
 			for (const FOrientedEdge& OrientedEdge : Loop->GetEdges())
 			{
-				OrientedEdge.Entity->ResetMarker1();
+				const TSharedPtr<FTopologicalEdge>& Edge = OrientedEdge.Entity;
+				if (Edge->HasMarker1())
+				{
+					continue;
+				}
+				Edge->SetMarker1();
+
+				if ((Edge->GetTwinEntityCount() != 2) || Edge->IsDegenerated())
+				{
+					if (Edge->GetTwinEntityCount() == 1)
+					{
+						BorderEdgesLength += Edge->Length();
+					}
+					continue;
+				}
+
+				FTopologicalEdge* NeighboringEdge = nullptr;
+				for (FTopologicalEdge* NextEdge : Edge->GetTwinEntities())
+				{
+					if (NextEdge == Edge.Get())
+					{
+						continue;
+					}
+					NeighboringEdge = NextEdge;
+					break;
+				}
+
+				if (NeighboringEdge == nullptr)
+				{
+					continue;
+				}
+
+				const FTopologicalLoop* NeighboringLoop = NeighboringEdge->GetLoop();
+				if (NeighboringLoop == nullptr)
+				{
+					continue;
+				}
+
+				FTopologicalFace* NeighboringFace = NeighboringLoop->GetFace();
+				if (NeighboringFace == nullptr)
+				{
+					continue;
+				}
+
+				if (NeighboringFace->HasMarker1())
+				{
+					continue;
+				}
+
+				NeighboringFace->SetMarker1();
+				Front.Add(NeighboringFace);
+
+				const FOrientedEdge* OrientedNeighboringEdge = NeighboringLoop->GetOrientedEdge(NeighboringEdge);
+				if (OrientedNeighboringEdge == nullptr)
+				{
+					continue;
+				}
+
+				bool EdgesAreWellOriented = (Edge->IsSameDirection(*NeighboringEdge) == (OrientedNeighboringEdge->Direction != OrientedEdge.Direction));
+				bool FacesHaveSameOrientation = Face.IsBackOriented() == NeighboringFace->IsBackOriented();
+
+				if (EdgesAreWellOriented != FacesHaveSameOrientation)
+				{
+					NeighboringFace->SwapOrientation();
+					SubshellSwappedFaceCount++;
+				}
+			}
+		}
+	};
+
+	TFunction<void()> SpreadFront = [&]()
+	{
+		while (Front.Num())
+		{
+			FTopologicalFace* Face = Front.Pop();
+			if (Face == nullptr)
+			{
+				continue;
+			}
+
+			Subshell.Add(Face);
+			SubshellFaceCount++;
+
+			GetAndOrientNeighboringFaces(*Face);
+		}
+	};
+
+	TFunction<FBBoxWithNormal(const double)> GetSubshellBoundingBox = [&](const double ApproximationFactor) -> FBBoxWithNormal
+	{
+		F3DDebugSession _(TEXT("GetShellBoundingBox"));
+
+		FBBoxWithNormal BBox;
+
+		for (FTopologicalFace* Face : Subshell)
+		{
+			if (Face == nullptr)
+			{
+				continue;
+			}
+			Face->UpdateBBox(3, ApproximationFactor, BBox);
+		}
+
+#ifdef DEBUG_GET_SHELL_BBOX
+		{
+			F3DDebugSession _(TEXT("BBox Face"));
+
+			FAABB AABB(BBox.Min, BBox.Max);
+			CADKernel::DisplayAABB(AABB);
+
+			for (int32 Index = 0; Index < 3; ++Index)
+			{
+				CADKernel::DisplayPoint(BBox.MaxPoints[Index], EVisuProperty::YellowPoint);
+				CADKernel::DisplayPoint(BBox.MinPoints[Index], EVisuProperty::YellowPoint);
+				CADKernel::DisplaySegment(BBox.MaxPoints[Index], BBox.MaxPoints[Index] + BBox.MaxPointNormals[Index], 0, EVisuProperty::YellowCurve);
+				CADKernel::DisplaySegment(BBox.MinPoints[Index], BBox.MinPoints[Index] + BBox.MinPointNormals[Index], 0, EVisuProperty::YellowCurve);
+			}
+			Wait();
+		}
+#endif
+		return BBox;
+	};
+
+	TFunction<void()> SwapSubshellOrientation = [&]()
+	{
+		for (FTopologicalFace* Face : Subshell)
+		{
+			if (Face == nullptr)
+			{
+				continue;
+			}
+			Face->SwapOrientation();
+		}
+	};
+
+	TFunction<void()> FixOpenShellOrientation = [&]()
+	{
+		if (SubshellSwappedFaceCount > (SubshellFaceCount / 2))
+		{
+			FMessage::Printf(Log, TEXT("the open shell %s (Id %d) is badly oriented. Its orientation is swapped"), *GetName(), GetId());
+			SwapSubshellOrientation();
+			ShellSwappedFaceCount += (SubshellFaceCount - SubshellSwappedFaceCount);
+		}
+		else
+		{
+			ShellSwappedFaceCount += SubshellSwappedFaceCount;
+		}
+	};
+
+	TFunction<void()> CheckAndFixSubshellMainOrientation = [&]()
+	{
+		if (Subshell.Num() == 1)
+		{
+			return;
+		}
+
+		FBBoxWithNormal BBox = GetSubshellBoundingBox(100.);
+		double BBoxLength = BBox.Length();
+
+		bool bIsOpenShell = (BBoxLength < BorderEdgesLength);
+		if(bIsOpenShell)
+		{
+			FixOpenShellOrientation();
+			return;
+		}
+
+		bool bHasWrongOrientation = false;
+		if (!BBox.CheckOrientation(bHasWrongOrientation))
+		{
+			FBBoxWithNormal BBox2 = GetSubshellBoundingBox(10.);
+			if (!BBox.CheckOrientation(bHasWrongOrientation))
+			{
+				FixOpenShellOrientation();
+				return;
+			}
+		}
+
+		if (bHasWrongOrientation)
+		{
+			FMessage::Printf(Log, TEXT("The closed shell %s (Id %d) is badly oriented. Its orientation is swapped"), *GetName(), GetId());
+			SwapSubshellOrientation();
+			ShellSwappedFaceCount += (SubshellFaceCount - SubshellSwappedFaceCount);
+		}
+		else
+		{
+			ShellSwappedFaceCount += SubshellSwappedFaceCount;
+		}
+	};
+
+	SpreadBodyOrientation();
+
+	for (FOrientedFace& Face : GetFaces())
+	{
+		Face.Entity->ResetMarkers();
+		for (const TSharedPtr<FTopologicalLoop>& Loop : Face.Entity->GetLoops())
+		{
+			for (const FOrientedEdge& OrientedEdge : Loop->GetEdges())
+			{
+				OrientedEdge.Entity->ResetMarkers();
 			}
 		}
 	}
-}
 
-void FShell::FillTopologyReport(FTopologyReport& Report) const 
-{
-	Report.Add(this);
-	for (const FOrientedFace& OrientedFace : GetFaces())
+	int32 ProcessFaceCount = 0;
+	for (FOrientedFace& Face : GetFaces())
 	{
-		OrientedFace.Entity->FillTopologyReport(Report);
+		if (Face.Entity->HasMarker1())
+		{
+			continue;
+		}
+
+		// Init data for the new subshell
+		SubshellFaceCount = 0;
+		SubshellSwappedFaceCount = 0;
+		BorderEdgesLength = 0;
+		Subshell.Empty(TopologicalFaceCount);
+		Front.Empty(TopologicalFaceCount);
+
+		Face.Entity->SetMarker1();
+		Front.Add(Face.Entity.Get());
+
+		SpreadFront();
+
+		CheckAndFixSubshellMainOrientation();
+
+		ProcessFaceCount += Subshell.Num();
+
+		if (ProcessFaceCount == TopologicalFaceCount)
+		{
+			break;
+		}
 	}
+
+	UpdateShellOrientation();
+
+	ResetMarkersRecursively();
+
+	return ShellSwappedFaceCount;
 }
 
 }
\ No newline at end of file
diff --git a/Engine/Source/Runtime/Datasmith/CADKernel/Private/CADKernel/Topo/TopologicalEdge.cpp b/Engine/Source/Runtime/Datasmith/CADKernel/Private/CADKernel/Topo/TopologicalEdge.cpp
--- a/Engine/Source/Runtime/Datasmith/CADKernel/Private/CADKernel/Topo/TopologicalEdge.cpp
+++ b/Engine/Source/Runtime/Datasmith/CADKernel/Private/CADKernel/Topo/TopologicalEdge.cpp
@@ -531,29 +531,6 @@ bool FTopologicalEdge::IsSameDirection(const FTopologicalEdge& Edge) const
 	return vertex1Edge == Edge.GetStartVertex()->GetLink();
 }
 
-TSharedPtr<FEntityGeom> FTopologicalEdge::ApplyMatrix(const FMatrixH& InMatrix) const
-{
-	TSharedPtr<FTopologicalVertex> v1Transformed = StaticCastSharedPtr<FTopologicalVertex>(StartVertex->ApplyMatrix(InMatrix));
-	if (!v1Transformed.IsValid())
-	{
-		return TSharedPtr<FEntityGeom>();
-	}
-
-	TSharedPtr<FTopologicalVertex> v2Transformed = StaticCastSharedPtr<FTopologicalVertex>(EndVertex->ApplyMatrix(InMatrix));
-	if (!v2Transformed.IsValid())
-	{
-		return TSharedPtr<FEntityGeom>();
-	}
-
-	TSharedPtr<FRestrictionCurve> TransformedCurve = StaticCastSharedPtr<FRestrictionCurve>(Curve->ApplyMatrix(InMatrix));
-	if (!TransformedCurve.IsValid())
-	{
-		return TSharedPtr<FEntityGeom>();
-	}
-
-	return FTopologicalEdge::Make(TransformedCurve.ToSharedRef(), v1Transformed.ToSharedRef(), v2Transformed.ToSharedRef(), Boundary);
-}
-
 #ifdef CADKERNEL_DEV
 FInfoEntity& FTopologicalEdge::GetInfo(FInfoEntity& Info) const
 {
@@ -569,7 +546,7 @@ FInfoEntity& FTopologicalEdge::GetInfo(FInfoEntity& Info) const
 }
 #endif
 
-TSharedRef<FEdgeMesh> FTopologicalEdge::GetOrCreateMesh(const TSharedRef<FModelMesh>& ShellMesh)
+TSharedRef<FEdgeMesh> FTopologicalEdge::GetOrCreateMesh(FModelMesh& ShellMesh)
 {
 	if (!IsActiveEntity())
 	{
@@ -578,7 +555,7 @@ TSharedRef<FEdgeMesh> FTopologicalEdge::GetOrCreateMesh(const TSharedRef<FModelM
 
 	if (!Mesh)
 	{
-		Mesh = FEntity::MakeShared<FEdgeMesh>(ShellMesh, StaticCastSharedRef<FTopologicalEdge>(AsShared()));
+		Mesh = FEntity::MakeShared<FEdgeMesh>(ShellMesh, *this);
 	}
 	return Mesh.ToSharedRef();
 }
@@ -939,7 +916,7 @@ void FTopologicalEdge::SpawnIdent(FDatabase& Database)
 
 TSharedPtr<FTopologicalVertex> FTopologicalEdge::SplitAt(double SplittingCoordinate, const FPoint& NewVertexCoordinate, bool bKeepStartVertexConnectivity, TSharedPtr<FTopologicalEdge>& NewEdge)
 {
-	if (GetTwinsEntityCount() > 1)
+	if (GetTwinEntityCount() > 1)
 	{
 		return TSharedPtr<FTopologicalVertex>();
 		ensureCADKernel(false);
@@ -1025,7 +1002,7 @@ FInfoEntity& FEdgeLink::GetInfo(FInfoEntity& Info) const
 {
 	return FEntity::GetInfo(Info)
 		.Add(TEXT("Active Entity"), ActiveEntity)
-		.Add(TEXT("Twin Entities"), TwinsEntities);
+		.Add(TEXT("Twin Entities"), TwinEntities);
 }
 #endif
 
diff --git a/Engine/Source/Runtime/Datasmith/CADKernel/Private/CADKernel/Topo/TopologicalEntity.cpp b/Engine/Source/Runtime/Datasmith/CADKernel/Private/CADKernel/Topo/TopologicalEntity.cpp
--- a/Engine/Source/Runtime/Datasmith/CADKernel/Private/CADKernel/Topo/TopologicalEntity.cpp
+++ b/Engine/Source/Runtime/Datasmith/CADKernel/Private/CADKernel/Topo/TopologicalEntity.cpp
@@ -7,7 +7,7 @@ namespace CADKernel
 #ifdef CADKERNEL_DEV
 FInfoEntity& FTopologicalEntity::GetInfo(FInfoEntity& Info) const
 {
-	return FEntityGeom::GetInfo(Info)
+	return FEntity::GetInfo(Info)
 		.Add(TEXT("IsMesh"), IsMeshed());
 }
 #endif
diff --git a/Engine/Source/Runtime/Datasmith/CADKernel/Private/CADKernel/Topo/TopologicalFace.cpp b/Engine/Source/Runtime/Datasmith/CADKernel/Private/CADKernel/Topo/TopologicalFace.cpp
--- a/Engine/Source/Runtime/Datasmith/CADKernel/Private/CADKernel/Topo/TopologicalFace.cpp
+++ b/Engine/Source/Runtime/Datasmith/CADKernel/Private/CADKernel/Topo/TopologicalFace.cpp
@@ -7,6 +7,8 @@
 #include "CADKernel/Geo/Curves/RestrictionCurve.h"
 #include "CADKernel/Geo/Curves/SegmentCurve.h"
 #include "CADKernel/Geo/GeoPoint.h"
+#include "CADKernel/Geo/Sampler/SamplerOnChord.h"
+#include "CADKernel/Geo/Sampling/Polyline.h"
 #include "CADKernel/Geo/Surfaces/Surface.h"
 #include "CADKernel/Mesh/Structure/FaceMesh.h"
 #include "CADKernel/Mesh/Structure/Grid.h"
@@ -44,6 +46,122 @@ void FTopologicalFace::Presample()
 	CarrierSurface->Presample(FaceBoundaries, CrossingCoordinates);
 }
 
+#ifdef DEBUG_GET_BBOX
+#include "CADKernel/Math/Aabb.h"
+#endif
+
+void FTopologicalFace::UpdateBBox(int32 IsoCount, const double ApproximationFactor, FBBoxWithNormal& BBox)
+{
+	const double SAG = GetCarrierSurface()->Get3DTolerance() * ApproximationFactor;
+
+	TArray<TArray<FPoint2D>> BoundaryApproximation;
+	Get2DLoopSampling(BoundaryApproximation);
+
+	FPolyline3D Polyline;
+	const FSurface& Surface = GetCarrierSurface().Get();
+	FIsoCurve3DSamplerOnChord Sampler(Surface, SAG, Polyline);
+
+	TFunction <void(const EIso)> UpdateBBoxWithIsos = [&](const EIso IsoType)
+	{
+		const FLinearBoundary& Bounds = GetBoundary().Get(IsoType);
+
+		EIso Other = IsoType == EIso::IsoU ? EIso::IsoV : EIso::IsoU;
+
+		double Coordinate = Bounds.Min;
+		IsoCount++;
+		const double Step = (Bounds.Max - Bounds.Min) / IsoCount;
+
+		for (int32 iIso = 1; iIso < IsoCount; iIso++)
+		{
+			FPolylineBBox IsoBBox;
+
+			Coordinate += Step;
+
+			TArray<double> Intersections;
+			FindLoopIntersectionsWithIso(IsoType, Coordinate, BoundaryApproximation, Intersections);
+			int32 IntersectionCount = Intersections.Num();
+			if (IntersectionCount == 0)
+			{
+				return;
+			}
+
+			FLinearBoundary CurveBounds(Intersections[0], Intersections.Last());
+
+			Polyline.Empty();
+			Sampler.Set(IsoType, Coordinate, CurveBounds);
+			Sampler.Sample();
+
+			if (IntersectionCount == 0)
+			{
+				continue;
+			}
+
+			if (IntersectionCount % 2 != 0)
+			{
+				TArray<FPoint> SubPolyline;
+				FLinearBoundary IntersectionBoundary(Intersections[IntersectionCount - 1], CurveBounds.GetMax());
+
+#ifdef DEBUG_GET_BBOX2
+				Polyline.GetSubPolyline(Boundary, EOrientation::Front, SubPolyline);
+				Draw(SubPolyline, EVisuProperty::Iso);
+#endif
+
+				Polyline.UpdateSubPolylineBBox(IntersectionBoundary, IsoBBox);
+
+				Intersections.Pop();
+				IntersectionCount--;
+			}
+
+			for (int32 ISection = 0; ISection < IntersectionCount; ISection += 2)
+			{
+				TArray<FPoint> SubPolyline;
+				FLinearBoundary IntersectionBoundary(Intersections[ISection], Intersections[ISection + 1]);
+
+#ifdef DEBUG_GET_BBOX2
+				Polyline.GetSubPolyline(Boundary, EOrientation::Front, SubPolyline);
+				Draw(SubPolyline, EVisuProperty::Iso);
+#endif
+				Polyline.UpdateSubPolylineBBox(IntersectionBoundary, IsoBBox);
+			}
+
+#ifdef DEBUG_GET_BBOX2
+			for (int32 Index = 0; Index < 3; ++Index)
+			{
+				CADKernel::DisplayPoint(IsoBBox.MaxPoints[Index], EVisuProperty::YellowPoint);
+				CADKernel::DisplayPoint(IsoBBox.MinPoints[Index], EVisuProperty::YellowPoint);
+			}
+#endif
+
+			BBox.Update(IsoBBox, IsoType, Coordinate);
+		}
+	};
+
+	UpdateBBoxWithIsos(EIso::IsoV);
+	UpdateBBoxWithIsos(EIso::IsoU);
+
+	BBox.UpdateNormal(*this);
+
+#ifdef DEBUG_GET_BBOX
+	{
+		F3DDebugSession _(TEXT("BBox Face"));
+
+		FAABB AABB(BBox.Min, BBox.Max);
+		CADKernel::DisplayAABB(AABB);
+
+		for (int32 Index = 0; Index < 3; ++Index)
+		{
+			CADKernel::DisplayPoint(BBox.MaxPoints[Index], EVisuProperty::YellowPoint);
+			CADKernel::DisplayPoint(BBox.MinPoints[Index], EVisuProperty::YellowPoint);
+			CADKernel::DisplaySegment(BBox.MaxPoints[Index], BBox.MaxPoints[Index] + BBox.MaxPointNormals[Index], 0, EVisuProperty::YellowCurve);
+			CADKernel::DisplaySegment(BBox.MinPoints[Index], BBox.MinPoints[Index] + BBox.MinPointNormals[Index], 0, EVisuProperty::YellowCurve);
+		}
+		Wait();
+	}
+#endif
+
+}
+
+
 void FTopologicalFace::ApplyNaturalLoops()
 {
 	const FSurfacicBoundary& Boundaries = CarrierSurface->GetBoundary();
@@ -210,7 +328,7 @@ bool FTopologicalFace::HasSameBoundariesAs(const TSharedPtr<FTopologicalFace>& O
 
 const FTopologicalEdge* FTopologicalFace::GetLinkedEdge(const FTopologicalEdge& LinkedEdge) const
 {
-	for (FTopologicalEdge* TwinEdge : LinkedEdge.GetTwinsEntities())
+	for (FTopologicalEdge* TwinEdge : LinkedEdge.GetTwinEntities())
 	{
 		if (&*TwinEdge->GetLoop()->GetFace() == this)
 		{
@@ -282,22 +400,20 @@ void FTopologicalFace::SpawnIdent(FDatabase& Database)
 #ifdef CADKERNEL_DEV
 FInfoEntity& FTopologicalFace::GetInfo(FInfoEntity& Info) const
 {
-	return FTopologicalEntity::GetInfo(Info)
-		.Add(TEXT("Hosted by"), HostedBy)
+	return FTopologicalShapeEntity::GetInfo(Info)
 		.Add(TEXT("Carrier Surface"), CarrierSurface)
 		.Add(TEXT("Boundary"), (FSurfacicBoundary&) Boundary)
 		.Add(TEXT("Loops"), Loops)
 		.Add(TEXT("QuadCriteria"), QuadCriteria)
-		.Add(TEXT("mesh"), Mesh)
-		.Add(*this);
+		.Add(TEXT("Mesh"), Mesh);
 }
 #endif
 
-TSharedRef<FFaceMesh> FTopologicalFace::GetOrCreateMesh(const TSharedRef<FModelMesh>& MeshModel)
+TSharedRef<FFaceMesh> FTopologicalFace::GetOrCreateMesh(FModelMesh& MeshModel)
 {
 	if (!Mesh.IsValid())
 	{
-		Mesh = FEntity::MakeShared<FFaceMesh>(MeshModel, StaticCastSharedRef<FTopologicalFace>(AsShared()));
+		Mesh = FEntity::MakeShared<FFaceMesh>(MeshModel, *this);
 	}
 	return Mesh.ToSharedRef();
 }
@@ -341,12 +457,6 @@ void FTopologicalFace::ChooseFinalDeltaUs()
 // =========================================================================================================================================================================================================
 // =========================================================================================================================================================================================================
 
-TSharedPtr<FEntityGeom> FTopologicalFace::ApplyMatrix(const FMatrixH& InMatrix) const
-{
-	NOT_IMPLEMENTED;
-	return TSharedPtr<FEntityGeom>();
-}
-
 // Quad ==============================================================================================================================================================================================================================
 
 double FTopologicalFace::GetQuadCriteria()
@@ -403,8 +513,10 @@ void FTopologicalFace::ComputeSurfaceSideProperties()
 
 void FTopologicalFace::DefineSurfaceType()
 {
-	if (CarrierSurface == nullptr)
+	if (!CarrierSurface.IsValid())
+	{
 		return;
+	}
 
 	const double Tolerance3D = CarrierSurface->Get3DTolerance();
 	const double GeometricTolerance = 20.0 * Tolerance3D;
@@ -423,7 +535,7 @@ void FTopologicalFace::DefineSurfaceType()
 			if (SideProperties[Index].IsoType == EIso::UndefinedIso)
 			{
 				TSharedPtr<FTopologicalEdge> Edge = Loops[0]->GetEdge(StartSideIndices[Index]);
-				int32 NeighborsNum = Edge->GetTwinsEntityCount();
+				int32 NeighborsNum = Edge->GetTwinEntityCount();
 				// If non manifold Edge => Stop
 				if (NeighborsNum != 2)
 				{
@@ -445,7 +557,7 @@ void FTopologicalFace::DefineSurfaceType()
 
 				FTopologicalFace* Neighbor = nullptr;
 				{
-					for(FTopologicalEdge* NeighborEdge : Edge->GetTwinsEntities() )
+					for(FTopologicalEdge* NeighborEdge : Edge->GetTwinEntities() )
 					{
 						if (NeighborEdge == Edge.Get())
 						{
@@ -455,10 +567,8 @@ void FTopologicalFace::DefineSurfaceType()
 					}
 				}
 
-				ensure(Neighbor != nullptr);
-
 				// it's not a quad surface
-				if (Neighbor==nullptr || Neighbor->SurfaceCorners.Num() == 0)
+				if (Neighbor == nullptr || Neighbor->SurfaceCorners.Num() == 0)
 				{
 					return;
 				}
@@ -520,7 +630,7 @@ void FTopologicalFace::DefineSurfaceType()
 }
 
 
-void FFaceSubset::SetMainShell(TMap<FShell*, int32>& ShellToFaceCount)
+void FFaceSubset::SetMainShell(TMap<FTopologicalShapeEntity*, int32>& ShellToFaceCount)
 {
 	if (ShellToFaceCount.Num() == 0)
 	{
@@ -528,9 +638,9 @@ void FFaceSubset::SetMainShell(TMap<FShell*, int32>& ShellToFaceCount)
 	}
 
 	int32 MaxFaceCount = 0;
-	FShell* CandidateShell = nullptr;
+	FTopologicalShapeEntity* CandidateShell = nullptr;
 
-	for (TPair<FShell*, int32>& Pair : ShellToFaceCount)
+	for (TPair<FTopologicalShapeEntity*, int32>& Pair : ShellToFaceCount)
 	{
 		if (Pair.Value > MaxFaceCount)
 		{
@@ -545,16 +655,16 @@ void FFaceSubset::SetMainShell(TMap<FShell*, int32>& ShellToFaceCount)
 	}
 }
 
-void FFaceSubset::SetMainBody(TMap<FBody*, int32>& BodyToFaceCount)
+void FFaceSubset::SetMainBody(TMap<FTopologicalShapeEntity*, int32>& BodyToFaceCount)
 {
 	if (BodyToFaceCount.Num() == 0)
 	{
 		return;
 	}
 
-	FBody* CandidateBody = nullptr;
+	FTopologicalShapeEntity* CandidateBody = nullptr;
 	int32 MaxFaceCount = 0;
-	for (TPair<FBody*, int32>& Pair : BodyToFaceCount)
+	for (TPair<FTopologicalShapeEntity*, int32>& Pair : BodyToFaceCount)
 	{
 		if (Pair.Value > MaxFaceCount)
 		{
diff --git a/Engine/Source/Runtime/Datasmith/CADKernel/Private/CADKernel/Topo/TopologicalLoop.cpp b/Engine/Source/Runtime/Datasmith/CADKernel/Private/CADKernel/Topo/TopologicalLoop.cpp
--- a/Engine/Source/Runtime/Datasmith/CADKernel/Private/CADKernel/Topo/TopologicalLoop.cpp
+++ b/Engine/Source/Runtime/Datasmith/CADKernel/Private/CADKernel/Topo/TopologicalLoop.cpp
@@ -979,13 +979,13 @@ void FTopologicalLoop::FindBreaks(TArray<TSharedPtr<FTopologicalVertex>>& OutBre
 	FPoint StartTangentEdge;
 	FPoint EndTangentPreviousEdge;
 	Edges[EdgeNum - 1].Entity->GetTangentsAtExtremities(StartTangentEdge, EndTangentPreviousEdge, Edges[EdgeNum - 1].Direction == EOrientation::Front);
-	bool bPreviousIsSurface = (Edges[EdgeNum - 1].Entity->GetTwinsEntityCount() > 1);
+	bool bPreviousIsSurface = (Edges[EdgeNum - 1].Entity->GetTwinEntityCount() > 1);
 
 	for (int32 Index = 0; Index < EdgeNum; Index++)
 	{
 		FPoint EndTangentEdge;
 		Edges[Index].Entity->GetTangentsAtExtremities(StartTangentEdge, EndTangentEdge, Edges[Index].Direction == EOrientation::Front);
-		bool bIsSurface = (Edges[Index].Entity->GetTwinsEntityCount() > 1);
+		bool bIsSurface = (Edges[Index].Entity->GetTwinEntityCount() > 1);
 
 		// if both edge are border, the rupture is not evaluate. 
 		if (bIsSurface || bPreviousIsSurface)
diff --git a/Engine/Source/Runtime/Datasmith/CADKernel/Private/CADKernel/Topo/TopologicalShapeEntity.cpp b/Engine/Source/Runtime/Datasmith/CADKernel/Private/CADKernel/Topo/TopologicalShapeEntity.cpp
new file mode 100644
--- /dev/null
+++ b/Engine/Source/Runtime/Datasmith/CADKernel/Private/CADKernel/Topo/TopologicalShapeEntity.cpp
@@ -0,0 +1,25 @@
+// Copyright Epic Games, Inc. All Rights Reserved.
+#include "CADKernel/Topo/TopologicalShapeEntity.h"
+
+namespace CADKernel
+{
+
+void FTopologicalShapeEntity::CompleteMetadata()
+{
+	if (HostedBy != nullptr)
+	{
+		Dictionary.CompleteDictionary(HostedBy->GetMetadataDictionary());
+	}
+}
+
+
+#ifdef CADKERNEL_DEV
+FInfoEntity& FTopologicalShapeEntity::GetInfo(FInfoEntity& Info) const
+{
+	return FTopologicalEntity::GetInfo(Info)
+		.Add(TEXT("Hosted by"), (FEntity*) HostedBy)
+		.Add(Dictionary);
+}
+#endif
+
+}
diff --git a/Engine/Source/Runtime/Datasmith/CADKernel/Private/CADKernel/Topo/TopologicalVertex.cpp b/Engine/Source/Runtime/Datasmith/CADKernel/Private/CADKernel/Topo/TopologicalVertex.cpp
--- a/Engine/Source/Runtime/Datasmith/CADKernel/Private/CADKernel/Topo/TopologicalVertex.cpp
+++ b/Engine/Source/Runtime/Datasmith/CADKernel/Private/CADKernel/Topo/TopologicalVertex.cpp
@@ -45,7 +45,7 @@ void FTopologicalVertex::RemoveConnectedEdge(FTopologicalEdge& Edge)
 
 bool FTopologicalVertex::IsBorderVertex()
 {
-	for (FTopologicalVertex* Vertex : GetTwinsEntities())
+	for (FTopologicalVertex* Vertex : GetTwinEntities())
 	{
 		if (Vertex != nullptr)
 		{
@@ -53,7 +53,7 @@ bool FTopologicalVertex::IsBorderVertex()
 			{
 				if (Edge != nullptr)
 				{
-					if (Edge->GetTwinsEntityCount() == 1)
+					if (Edge->GetTwinEntityCount() == 1)
 					{
 						return true;
 					}
@@ -64,18 +64,12 @@ bool FTopologicalVertex::IsBorderVertex()
 	return false;
 }
 
-TSharedPtr<FEntityGeom> FTopologicalVertex::ApplyMatrix(const FMatrixH& InMatrix) const
-{
-	FPoint transformedPoint = InMatrix.Multiply(Coordinates);
-	return FTopologicalVertex::Make(transformedPoint);
-}
-
 void FTopologicalVertex::GetConnectedEdges(const FTopologicalVertex& OtherVertex, TArray<FTopologicalEdge*>& OutEdges) const
 {
-	OutEdges.Reserve(GetTwinsEntityCount());
+	OutEdges.Reserve(GetTwinEntityCount());
 
 	TSharedPtr<TTopologicalLink<FTopologicalVertex>> OtherVertexLink = OtherVertex.GetLink();
-	for (const FTopologicalVertex* Vertex : GetTwinsEntities())
+	for (const FTopologicalVertex* Vertex : GetTwinEntities())
 	{
 		if (Vertex != nullptr)
 		{
@@ -104,12 +98,12 @@ void FTopologicalVertex::Link(FTopologicalVertex& Twin)
 		}
 	}
 
-	FPoint Barycenter = GetBarycenter() * (double)GetTwinsEntityCount()
-		+ Twin.GetBarycenter() * (double)Twin.GetTwinsEntityCount();
+	FPoint Barycenter = GetBarycenter() * (double)GetTwinEntityCount()
+		+ Twin.GetBarycenter() * (double)Twin.GetTwinEntityCount();
 
 	MakeLink(Twin);
 
-	Barycenter /= (double)GetTwinsEntityCount();
+	Barycenter /= (double)GetTwinEntityCount();
 	GetLink()->SetBarycenter(Barycenter);
 
 	// Find the closest vertex of the Barycenter
@@ -122,7 +116,7 @@ void FTopologicalVertex::UnlinkTo(FTopologicalVertex& OtherVertex)
 	ResetTopologicalLink();
 	OtherVertex.ResetTopologicalLink();
 
-	for (FTopologicalVertex* Vertex : OldLink->GetTwinsEntities())
+	for (FTopologicalVertex* Vertex : OldLink->GetTwinEntities())
 	{
 		if (!Vertex || Vertex == this || Vertex == &OtherVertex)
 		{
@@ -146,24 +140,24 @@ void FTopologicalVertex::UnlinkTo(FTopologicalVertex& OtherVertex)
 void FVertexLink::ComputeBarycenter()
 {
 	Barycenter = FPoint::ZeroPoint;
-	for (const FTopologicalVertex* Vertex : TwinsEntities)
+	for (const FTopologicalVertex* Vertex : TwinEntities)
 	{
 		Barycenter += Vertex->Coordinates;
 	}
-	Barycenter /= TwinsEntities.Num();
+	Barycenter /= TwinEntities.Num();
 }
 
 void FVertexLink::DefineActiveEntity()
 {
-	if (TwinsEntities.Num() == 0)
+	if (TwinEntities.Num() == 0)
 	{
 		ActiveEntity = nullptr;
 		return;
 	}
 
 	double DistanceSquare = HUGE_VALUE;
-	FTopologicalVertex* ClosedVertex = TwinsEntities.HeapTop();
-	for (FTopologicalVertex* Vertex : TwinsEntities)
+	FTopologicalVertex* ClosedVertex = TwinEntities.HeapTop();
+	for (FTopologicalVertex* Vertex : TwinEntities)
 	{
 		ensureCADKernel(!Vertex->IsDeleted());
 
@@ -181,7 +175,7 @@ void FVertexLink::DefineActiveEntity()
 	ActiveEntity = ClosedVertex;
 }
 
-TSharedRef<FVertexMesh> FTopologicalVertex::GetOrCreateMesh(TSharedRef<FModelMesh>& MeshModel)
+TSharedRef<FVertexMesh> FTopologicalVertex::GetOrCreateMesh(FModelMesh& MeshModel)
 {
 	if (!IsActiveEntity())
 	{
@@ -190,11 +184,11 @@ TSharedRef<FVertexMesh> FTopologicalVertex::GetOrCreateMesh(TSharedRef<FModelMes
 
 	if (!Mesh.IsValid())
 	{
-		Mesh = FEntity::MakeShared<FVertexMesh>(MeshModel, StaticCastSharedRef<FTopologicalVertex>(AsShared()));
+		Mesh = FEntity::MakeShared<FVertexMesh>(MeshModel, *this);
 
 		Mesh->GetNodeCoordinates().Emplace(GetBarycenter());
 		Mesh->RegisterCoordinates();
-		MeshModel->AddMesh(Mesh.ToSharedRef());
+		MeshModel.AddMesh(Mesh.ToSharedRef());
 		SetMeshed();
 	}
 	return Mesh.ToSharedRef();
@@ -224,7 +218,7 @@ FInfoEntity& TTopologicalLink<FTopologicalVertex>::GetInfo(FInfoEntity& Info) co
 {
 	return FEntity::GetInfo(Info)
 		.Add(TEXT("active Entity"), ActiveEntity)
-		.Add(TEXT("twin Entities"), TwinsEntities);
+		.Add(TEXT("twin Entities"), TwinEntities);
 }
 
 FInfoEntity& FVertexLink::GetInfo(FInfoEntity& Info) const
diff --git a/Engine/Source/Runtime/Datasmith/CADKernel/Private/CADKernel/Topo/TopologyReport.cpp b/Engine/Source/Runtime/Datasmith/CADKernel/Private/CADKernel/Topo/TopologyReport.cpp
--- a/Engine/Source/Runtime/Datasmith/CADKernel/Private/CADKernel/Topo/TopologyReport.cpp
+++ b/Engine/Source/Runtime/Datasmith/CADKernel/Private/CADKernel/Topo/TopologyReport.cpp
@@ -58,7 +58,7 @@ void FTopologyReport::Add(const FTopologicalEdge* Edge)
 
 		CoedgeCount++;
 
-		switch (ActiveEdge->GetTwinsEntityCount())
+		switch (ActiveEdge->GetTwinEntityCount())
 		{
 		case 1:
 			BorderEdgeCount++;
@@ -93,7 +93,7 @@ void FTopologyReport::CountLoops()
 		}
 		Edge->SetMarker1();
 
-		if (Edge->GetTwinsEntityCount() == 2)
+		if (Edge->GetTwinEntityCount() == 2)
 		{
 			continue;
 		}
