Update CADKernel : add oriention, fix one infinite loop

diff --git a/Engine/Source/Runtime/Datasmith/CADKernel/Private/CADKernel/Topo/Topomaker.cpp b/Engine/Source/Runtime/Datasmith/CADKernel/Private/CADKernel/Topo/Topomaker.cpp
--- a/Engine/Source/Runtime/Datasmith/CADKernel/Private/CADKernel/Topo/Topomaker.cpp
+++ b/Engine/Source/Runtime/Datasmith/CADKernel/Private/CADKernel/Topo/Topomaker.cpp
@@ -13,6 +13,7 @@
 #include "CADKernel/Topo/TopologicalFace.h"
 #include "CADKernel/Topo/TopologicalLoop.h"
 #include "CADKernel/Topo/TopologicalVertex.h"
+#include "CADKernel/Topo/TopologicalShapeEntity.h"
 #include "CADKernel/UI/Display.h"
 #include "CADKernel/UI/Message.h"
 #include "CADKernel/Utils/Util.h"
@@ -20,54 +21,61 @@
 namespace CADKernel
 {
 
-FJoiner::FJoiner(FSession& InSession, const TSharedRef<FModel>& InModel, double InTolerance)
+FTopomaker::FTopomaker(FSession& InSession, double InTolerance)
 	: Session(InSession)
-	, JoiningTolerance(InTolerance* UE_SQRT_2)
-	, JoiningToleranceSquare(FMath::Square(JoiningTolerance))
+	, SewTolerance(InTolerance* UE_SQRT_2)
+	, SewToleranceSquare(FMath::Square(SewTolerance))
 {
 	int32 ShellCount = 0;
-	for (const TSharedPtr<FBody>& Body : InModel->GetBodies())
+	for (const TSharedPtr<FBody>& Body : Session.GetModel().GetBodies())
 	{
 		ShellCount += Body->GetShells().Num();
 	}
 	Shells.Reserve(ShellCount);
 
-	for (const TSharedPtr<FBody>& Body : InModel->GetBodies())
+	for (const TSharedPtr<FBody>& Body : Session.GetModel().GetBodies())
 	{
-		Shells.Append(Body->GetShells());
+		for (const TSharedPtr<FShell>& Shell : Body->GetShells())
+		{
+			Shells.Add(Shell.Get());
+		}
 	}
 
 	InitFaces();
 }
 
-FJoiner::FJoiner(FSession& InSession, const TArray<TSharedPtr<FTopologicalFace>>& InFaces, double InTolerance)
+FTopomaker::FTopomaker(FSession& InSession, const TArray<TSharedPtr<FTopologicalFace>>& InFaces, double InTolerance)
 	: Session(InSession)
 	, Faces(InFaces)
-	, JoiningTolerance(InTolerance * UE_SQRT_2)
-	, JoiningToleranceSquare(FMath::Square(JoiningTolerance))
+	, SewTolerance(InTolerance * UE_SQRT_2)
+	, SewToleranceSquare(FMath::Square(SewTolerance))
 {
 }
 
-FJoiner::FJoiner(FSession& InSession, const TArray<TSharedPtr<FShell>>& InShells, double InTolerance)
+FTopomaker::FTopomaker(FSession& InSession, const TArray<TSharedPtr<FShell>>& InShells, double InTolerance)
 	: Session(InSession)
-	, Shells(InShells)
-	, JoiningTolerance(InTolerance)
-	, JoiningToleranceSquare(FMath::Square(JoiningTolerance))
+	, SewTolerance(InTolerance)
+	, SewToleranceSquare(FMath::Square(SewTolerance))
 {
+	Shells.Reserve(InShells.Num());
+	for(const TSharedPtr<FShell>& Shell : InShells)
+	{
+		Shells.Add(Shell.Get());
+	}
 	InitFaces();
 }
 
-void FJoiner::InitFaces()
+void FTopomaker::InitFaces()
 {
 	int32 FaceCount = 0;
-	for (const TSharedPtr<FShell>& Shell : Shells)
+	for (FShell* Shell : Shells)
 	{
 		Shell->CompleteMetadata();
 		FaceCount += Shell->FaceCount();
 	}
 	Faces.Reserve(FaceCount);
 
-	for (const TSharedPtr<FShell>& Shell : Shells)
+	for (FShell* Shell : Shells)
 	{
 		Shell->SpreadBodyOrientation();
 		for (const FOrientedFace& Face : Shell->GetFaces())
@@ -86,7 +94,7 @@ void FJoiner::InitFaces()
 	}
 }
 
-void FJoiner::RemoveFacesFromShell()
+void FTopomaker::RemoveFacesFromShell()
 {
 	// remove faces from their shells
 	TSet<FShell*> ShellSet;
@@ -94,7 +102,7 @@ void FJoiner::RemoveFacesFromShell()
 	{
 		if (Face->GetHost() != nullptr)
 		{
-			ShellSet.Add(Face->GetHost());
+			ShellSet.Add((FShell*) Face->GetHost());
 			Face->ResetHost();
 		}	
 	}
@@ -123,42 +131,33 @@ void FJoiner::RemoveFacesFromShell()
 	}
 }
 
-void FJoiner::EmptyShells()
+void FTopomaker::EmptyShells()
 {
-	for (const TSharedPtr<FShell>& Shell : Shells)
+	for (FShell* Shell : Shells)
 	{
 		Shell->Empty();
 	}
 }
 
-void FJoiner::JoinFaces()
+void FTopomaker::Sew()
 {
 	FTimePoint StartJoinTime = FChrono::Now();
 
-	//RemoveIsolatedEdges();
-
 	TArray<FTopologicalVertex*> BorderVertices;
 	GetBorderVertices(BorderVertices);
 	MergeCoincidentVertices(BorderVertices);
-	//RemoveIsolatedEdges();
 
 	CheckSelfConnectedEdge();
 
-	//RemoveIsolatedEdges();
-
-	//Wait();
 	MergeCoincidentEdges(BorderVertices);
 
-	//Wait();
 	MergeUnconnectedAdjacentEdges();
 
-	//StitchParallelEdges(BorderVertices);
-
-	FDuration InitKernelIODuration = FChrono::Elapse(StartJoinTime);
-	FChrono::PrintClockElapse(EVerboseLevel::Log, TEXT(""), TEXT("Join"), InitKernelIODuration);
+	Report.SewDuration = FChrono::Elapse(StartJoinTime);
+	FChrono::PrintClockElapse(EVerboseLevel::Log, TEXT(""), TEXT("Sew"), Report.SewDuration);
 }
 
-void FJoiner::GetVertices(TArray<FTopologicalVertex*>& Vertices)
+void FTopomaker::GetVertices(TArray<FTopologicalVertex*>& Vertices)
 {
 	Vertices.Empty(10 * Faces.Num());
 
@@ -188,7 +187,7 @@ void FJoiner::GetVertices(TArray<FTopologicalVertex*>& Vertices)
 	}
 }
 
-void FJoiner::GetBorderVertices(TArray<FTopologicalVertex*>& BorderVertices)
+void FTopomaker::GetBorderVertices(TArray<FTopologicalVertex*>& BorderVertices)
 {
 	TArray<FTopologicalVertex*> Vertices;
 	GetVertices(Vertices);
@@ -204,12 +203,12 @@ void FJoiner::GetBorderVertices(TArray<FTopologicalVertex*>& BorderVertices)
 	}
 }
 
-void FJoiner::MergeCoincidentVertices(TArray<FTopologicalVertex*>& VerticesToMerge)
+void FTopomaker::MergeCoincidentVertices(TArray<FTopologicalVertex*>& VerticesToMerge)
 {
 	FTimePoint StartTime = FChrono::Now();
 
-	const double JoiningVerticesToleranceSquare = 2 * JoiningToleranceSquare;
-	const double WeigthTolerance = 3 * JoiningTolerance;
+	const double JoiningVerticesToleranceSquare = 2 * SewToleranceSquare;
+	const double WeigthTolerance = 3 * SewTolerance;
 
 	int32 VertexNum = (int32)VerticesToMerge.Num();
 
@@ -295,12 +294,12 @@ void FJoiner::MergeCoincidentVertices(TArray<FTopologicalVertex*>& VerticesToMer
 
 	Swap(ActiveVertices, VerticesToMerge);
 
-	FDuration Duration = FChrono::Elapse(StartTime);
-	FChrono::PrintClockElapse(EVerboseLevel::Log, TEXT("    "), TEXT("Merge Coincident vertices"), Duration);
+	FDuration StepDuration = FChrono::Elapse(StartTime);
+	FChrono::PrintClockElapse(EVerboseLevel::Log, TEXT("    "), TEXT("Merge Coincident vertices"), StepDuration);
 
 }
 
-void FJoiner::MergeBorderVerticesWithCoincidentOtherVertices(TArray<FTopologicalVertex*>& Vertices)
+void FTopomaker::MergeBorderVerticesWithCoincidentOtherVertices(TArray<FTopologicalVertex*>& Vertices)
 {
 	int32 VertexNum = (int32)Vertices.Num();
 
@@ -341,19 +340,19 @@ void FJoiner::MergeBorderVerticesWithCoincidentOtherVertices(TArray<FTopological
 
 			double OtherVertexWeigth = VerticesWeight[SortedVertexIndices[IndexJ]];
 
-			if (OtherVertexWeigth + JoiningTolerance - VertexWeigth < 0)
+			if (OtherVertexWeigth + SewTolerance - VertexWeigth < 0)
 			{
 				StartIndexJ = IndexJ;
 				continue;
 			}
 
-			if (OtherVertexWeigth - JoiningTolerance - VertexWeigth > 0)
+			if (OtherVertexWeigth - SewTolerance - VertexWeigth > 0)
 			{
 				break;
 			}
 
 			double DistanceSqr = OtherVertex.GetLinkActiveEntity()->SquareDistance(Barycenter);
-			if (DistanceSqr < JoiningToleranceSquare)
+			if (DistanceSqr < SewToleranceSquare)
 			{
 				TArray<FTopologicalEdge*> CommonEdges;
 				Vertex->GetConnectedEdges(OtherVertex, CommonEdges);
@@ -391,7 +390,7 @@ void FJoiner::MergeBorderVerticesWithCoincidentOtherVertices(TArray<FTopological
 	Swap(NewVertices, Vertices);
 }
 
-void FJoiner::MergeCoincidentEdges(TArray<FTopologicalVertex*>& VerticesToProcess)
+void FTopomaker::MergeCoincidentEdges(TArray<FTopologicalVertex*>& VerticesToProcess)
 {
 	FTimePoint StartTime = FChrono::Now();
 
@@ -445,20 +444,20 @@ void FJoiner::MergeCoincidentEdges(TArray<FTopologicalVertex*>& VerticesToProces
 
 					if (Edge->GetFace() != SecondEdge->GetFace())
 					{
-						Edge->Link(*SecondEdge, JoiningTolerance);
+						Edge->Link(*SecondEdge, SewTolerance);
 					}
 				}
 			}
 		}
 	}
 
-	FDuration Duration = FChrono::Elapse(StartTime);
+	FDuration StepDuration = FChrono::Elapse(StartTime);
 
-	FChrono::PrintClockElapse(EVerboseLevel::Log, TEXT("    "), TEXT("Merge coincident edges"), Duration);
+	FChrono::PrintClockElapse(EVerboseLevel::Log, TEXT("    "), TEXT("Merge coincident edges"), StepDuration);
 
 }
 
-TSharedPtr<FTopologicalVertex> FJoiner::SplitAndLink(FTopologicalVertex& StartVertex, FTopologicalEdge& EdgeToLink, FTopologicalEdge& EdgeToSplit)
+TSharedPtr<FTopologicalVertex> FTopomaker::SplitAndLink(FTopologicalVertex& StartVertex, FTopologicalEdge& EdgeToLink, FTopologicalEdge& EdgeToSplit)
 {
 	FTopologicalVertex* VertexToLink = EdgeToLink.GetOtherVertex(StartVertex);
 
@@ -466,17 +465,17 @@ TSharedPtr<FTopologicalVertex> FJoiner::SplitAndLink(FTopologicalVertex& StartVe
 	double UProjectedPoint = EdgeToSplit.ProjectPoint(VertexToLink->GetBarycenter(), ProjectedPoint);
 
 	double SquareDistanceToProjectedPoint = ProjectedPoint.SquareDistance(VertexToLink->GetBarycenter());
-	if (SquareDistanceToProjectedPoint > JoiningToleranceSquare)
+	if (SquareDistanceToProjectedPoint > SewToleranceSquare)
 	{
 		return TSharedPtr<FTopologicalVertex>();
 	}
 
 	// Check if the ProjectedPoint is not nearly equal edge boundary
 	FTopologicalVertex* EndVertex = EdgeToSplit.GetOtherVertex(StartVertex);
-	if (EndVertex->SquareDistance(ProjectedPoint) < JoiningToleranceSquare)
+	if (EndVertex->SquareDistance(ProjectedPoint) < SewToleranceSquare)
 	{
 		VertexToLink->Link(*EndVertex);
-		EdgeToLink.Link(EdgeToSplit, JoiningTolerance);
+		EdgeToLink.Link(EdgeToSplit, SewTolerance);
 		// TSharedPtr<FTopologicalVertex>() is returned as EndVertex is not new
 		return TSharedPtr<FTopologicalVertex>();
 	}
@@ -494,12 +493,12 @@ TSharedPtr<FTopologicalVertex> FJoiner::SplitAndLink(FTopologicalVertex& StartVe
 	}
 
 	VertexToLink->Link(*NewVertex);
-	EdgeToLink.Link(EdgeToSplit, JoiningTolerance);
+	EdgeToLink.Link(EdgeToSplit, SewTolerance);
 
 	return NewVertex;
 }
 
-void FJoiner::StitchParallelEdges(TArray<FTopologicalVertex*>& VerticesToProcess)
+void FTopomaker::StitchParallelEdges(TArray<FTopologicalVertex*>& VerticesToProcess)
 {
 	FTimePoint StartTime = FChrono::Now();
 
@@ -568,7 +567,7 @@ void FJoiner::StitchParallelEdges(TArray<FTopologicalVertex*>& VerticesToProcess
 				FTopologicalVertex& OtherEdgeEndVertex = *SecondEdge->GetOtherVertex(*Vertex)->GetLinkActiveEntity();
 				if (&EndVertex == &OtherEdgeEndVertex)
 				{
-					Edge->Link(*SecondEdge, JoiningTolerance);
+					Edge->Link(*SecondEdge, SewTolerance);
 				}
 				else
 				{
@@ -594,11 +593,11 @@ void FJoiner::StitchParallelEdges(TArray<FTopologicalVertex*>& VerticesToProcess
 		}
 	}
 
-	FDuration Duration = FChrono::Elapse(StartTime);
-	FChrono::PrintClockElapse(EVerboseLevel::Log, TEXT("    "), TEXT("Stitch Parallel Edges"), Duration);
+	FDuration StepDuration = FChrono::Elapse(StartTime);
+	FChrono::PrintClockElapse(EVerboseLevel::Log, TEXT("    "), TEXT("Stitch Parallel Edges"), StepDuration);
 }
 
-void FJoiner::MergeUnconnectedAdjacentEdges()
+void FTopomaker::MergeUnconnectedAdjacentEdges()
 {
 	FTimePoint StartTime = FChrono::Now();
 
@@ -648,7 +647,7 @@ void FJoiner::MergeUnconnectedAdjacentEdges()
 				bCanStop = true;
 
 				FOrientedEdge& Edge = Edges[Index];
-				if (Edge.Entity->GetTwinsEntityCount() == 1)
+				if (Edge.Entity->GetTwinEntityCount() == 1)
 				{
 					TSharedPtr<FTopologicalVertex> EndVertex = Edge.Direction == EOrientation::Front ? Edge.Entity->GetEndVertex() : Edge.Entity->GetStartVertex();
 
@@ -745,7 +744,7 @@ void FJoiner::MergeUnconnectedAdjacentEdges()
 					{
 						if (Edge)
 						{
-							Edge->Link(*NewEdge, JoiningTolerance);
+							Edge->Link(*NewEdge, SewTolerance);
 							break;
 						}
 					}
@@ -754,13 +753,12 @@ void FJoiner::MergeUnconnectedAdjacentEdges()
 		}
 	}
 
-	FDuration Duration = FChrono::Elapse(StartTime);
-
-	FChrono::PrintClockElapse(EVerboseLevel::Log, TEXT("    "), TEXT("Merge unconnected adjacent edges"), Duration);
+	FDuration StepDuration = FChrono::Elapse(StartTime);
 
+	FChrono::PrintClockElapse(EVerboseLevel::Log, TEXT("    "), TEXT("Merge unconnected adjacent edges"), StepDuration);
 }
 
-void FJoiner::RemoveIsolatedEdges()
+void FTopomaker::RemoveIsolatedEdges()
 {
 	FTimePoint StartTime = FChrono::Now();
 
@@ -771,7 +769,7 @@ void FJoiner::RemoveIsolatedEdges()
 
 	for (const FTopologicalVertex* Vertex : Vertices)
 	{
-		for (const FTopologicalVertex* TwinVertex : Vertex->GetTwinsEntities())
+		for (const FTopologicalVertex* TwinVertex : Vertex->GetTwinEntities())
 		{
 			TArray<FTopologicalEdge*> Edges = Vertex->GetDirectConnectedEdges();
 			for (FTopologicalEdge* Edge : Edges)
@@ -783,14 +781,14 @@ void FJoiner::RemoveIsolatedEdges()
 			}
 		}
 	}
-	FDuration Duration = FChrono::Elapse(StartTime);
+	FDuration StepDuration = FChrono::Elapse(StartTime);
 
-	FChrono::PrintClockElapse(EVerboseLevel::Log, TEXT("    "), TEXT("Remove Isolated Edges"), Duration);
+	FChrono::PrintClockElapse(EVerboseLevel::Log, TEXT("    "), TEXT("Remove Isolated Edges"), StepDuration);
 
 	FMessage::Printf(EVerboseLevel::Log, TEXT("\n\nIsolatedEdges count %d\n\n\n"), IsolatedEdges.Num());
 }
 
-void FJoiner::UnlinkNonManifoldVertex()
+void FTopomaker::UnlinkNonManifoldVertex()
 {
 	TArray<FTopologicalVertex*> Vertices;
 	GetVertices(Vertices);
@@ -798,12 +796,12 @@ void FJoiner::UnlinkNonManifoldVertex()
 	TMap<FShell*, int32> ShellToVertexCount;
 	for (FTopologicalVertex* Vertex : Vertices)
 	{
-		ShellToVertexCount.Empty(Vertex->GetTwinsEntityCount());
-		for (const FTopologicalVertex* TwinVertex : Vertex->GetTwinsEntities())
+		ShellToVertexCount.Empty(Vertex->GetTwinEntityCount());
+		for (const FTopologicalVertex* TwinVertex : Vertex->GetTwinEntities())
 		{
 			ensureCADKernel(!TwinVertex->GetDirectConnectedEdges().IsEmpty());
 			FTopologicalEdge* Edge = TwinVertex->GetDirectConnectedEdges()[0];
-			FShell* Shell = Edge->GetLoop()->GetFace()->GetHost();
+			FShell* Shell = (FShell*) Edge->GetLoop()->GetFace()->GetHost();
 			if (Shell != nullptr)
 			{
 				ShellToVertexCount.FindOrAdd(Shell)++;
@@ -819,11 +817,11 @@ void FJoiner::UnlinkNonManifoldVertex()
 				ShellToVertices.FindOrAdd(Pair.Key).Reserve(Pair.Value);
 			}
 
-			for (FTopologicalVertex* TwinVertex : Vertex->GetTwinsEntities())
+			for (FTopologicalVertex* TwinVertex : Vertex->GetTwinEntities())
 			{
 				ensureCADKernel(!TwinVertex->GetDirectConnectedEdges().IsEmpty());
 				FTopologicalEdge* Edge = TwinVertex->GetDirectConnectedEdges()[0];
-				FShell* Shell = Edge->GetLoop()->GetFace()->GetHost();
+				FShell* Shell = (FShell*) Edge->GetLoop()->GetFace()->GetHost();
 				if (Shell != nullptr)
 				{
 					ShellToVertices.FindOrAdd(Shell).Add(TwinVertex);
@@ -859,7 +857,7 @@ void FJoiner::UnlinkNonManifoldVertex()
 // =========================================================================================================================================================================================================
 // =========================================================================================================================================================================================================
 
-void FJoiner::CheckSelfConnectedEdge()
+void FTopomaker::CheckSelfConnectedEdge()
 {
 	FTimePoint StartTime = FChrono::Now();
 
@@ -873,7 +871,7 @@ void FJoiner::CheckSelfConnectedEdge()
 				TSharedPtr<FTopologicalEdge> Edge = OrientedEdge.Entity;
 				if (Edge->GetStartVertex()->IsLinkedTo(Edge->GetEndVertex()))
 				{
-					if (!Edge->IsDegenerated() && Edge->Length() < 2 * JoiningTolerance)
+					if (!Edge->IsDegenerated() && Edge->Length() < 2 * SewTolerance)
 					{
 						FMessage::Printf(Debug, TEXT("Face %d Edge %d was self connected\n"), Face->GetId(), Edge->GetId());
 						Edge->GetStartVertex()->UnlinkTo(*Edge->GetEndVertex());
@@ -882,13 +880,15 @@ void FJoiner::CheckSelfConnectedEdge()
 			}
 		}
 	}
-	FDuration Duration = FChrono::Elapse(StartTime);
+	FDuration StepDuration = FChrono::Elapse(StartTime);
 
-	FChrono::PrintClockElapse(EVerboseLevel::Log, TEXT("    "), TEXT("Unconnect Self connected edges"), Duration);
+	FChrono::PrintClockElapse(EVerboseLevel::Log, TEXT("    "), TEXT("Unconnect Self connected edges"), StepDuration);
 }
 
-void FJoiner::SplitIntoConnectedShell()
+void FTopomaker::SplitIntoConnectedShells()
 {
+	FTimePoint StartTime = FChrono::Now();
+
 	// Processed1 : Surfaces added in CandidateSurfacesForMesh
 
 	int32 TopologicalFaceCount = Faces.Num();
@@ -913,7 +913,7 @@ void FJoiner::SplitIntoConnectedShell()
 				}
 				Edge->SetMarker1();
 
-				if (Edge->GetTwinsEntityCount() == 1)
+				if (Edge->GetTwinEntityCount() == 1)
 				{
 					if (!Edge->IsDegenerated())
 					{
@@ -922,12 +922,12 @@ void FJoiner::SplitIntoConnectedShell()
 					continue;
 				}
 
-				if (Edge->GetTwinsEntityCount() > 2)
+				if (Edge->GetTwinEntityCount() > 2)
 				{
 					Shell.NonManifoldEdgeCount++;
 				}
 
-				for (FTopologicalEdge* NextEdge : Edge->GetTwinsEntities())
+				for (FTopologicalEdge* NextEdge : Edge->GetTwinEntities())
 				{
 					if (NextEdge->HasMarker1())
 					{
@@ -1001,15 +1001,15 @@ void FJoiner::SplitIntoConnectedShell()
 	// for each FaceSubset, find the main shell
 	for (FFaceSubset& FaceSubset : SubShells)
 	{
-		TMap<FBody*, int32> BodyToFaceCount;
-		TMap<FShell*, int32> ShellToFaceCount;
+		TMap<FTopologicalShapeEntity*, int32> BodyToFaceCount;
+		TMap<FTopologicalShapeEntity*, int32> ShellToFaceCount;
 		TMap<uint32, int32> ColorToFaceCount;
 		TMap<FString, int32> NameToFaceCount;
 
 		for (FTopologicalFace* Face : FaceSubset.Faces)
 		{
-			FShell* Shell = Face->GetHost();
-			FBody* Body = Shell->GetHost();
+			FTopologicalShapeEntity* Shell = Face->GetHost();
+			FTopologicalShapeEntity* Body = Shell->GetHost();
 
 			ShellToFaceCount.FindOrAdd(Shell)++;
 			BodyToFaceCount.FindOrAdd(Body)++;
@@ -1037,142 +1037,92 @@ void FJoiner::SplitIntoConnectedShell()
 	{
 		if (FaceSubset.MainShell != nullptr)
 		{
-			FShell* Shell = FaceSubset.MainShell;
+			FShell* Shell = (FShell*) FaceSubset.MainShell;
 			Shell->Empty(FaceSubset.Faces.Num());
 			Shell->Add(FaceSubset.Faces);
 		}
 		else
 		{
-			FBody* Body = FaceSubset.MainBody;
+			FBody* Body = (FBody*) FaceSubset.MainBody;
 			if (Body == nullptr)
 			{
 				TSharedRef<FBody> SharedBody = FEntity::MakeShared<FBody>();
-				Session.GetModel()->Add(SharedBody);
+				Session.GetModel().Add(SharedBody);
 				Body = &SharedBody.Get();
+
+				Session.SpawnEntityIdent(*Body);
 				Body->SetName(FaceSubset.MainName);
 				Body->SetColorId(FaceSubset.MainColor);
+				Body->SetHostId(Body->GetId());
 			}
 
 			TSharedRef<FShell> Shell = FEntity::MakeShared<FShell>();
+
+			Shells.Add(&*Shell);
 			Body->AddShell(Shell);
+			Session.SpawnEntityIdent(*Shell);
 
 			Shell->Add(FaceSubset.Faces);
 			Shell->SetName(FaceSubset.MainName);
 			Shell->SetColorId(FaceSubset.MainColor);
+			Shell->SetHostId(Shell->GetId());
 		}
 	}
 
 	UnlinkNonManifoldVertex();
 
-	Session.GetModel()->RemoveEmptyBodies();
-}
-
-
-
-//void FJoiner::FixCollapsedEdges()
-//{
-//	TArray<TSharedPtr<FTopologicalEdge>> CollapsedEdges;
-//
-//	int32 NbDegeneratedEdges = 0;
-//	for (const TSharedPtr<FTopologicalFace>& Face : Faces)
-//	{
-//		for (const TSharedPtr<FTopologicalLoop>& Loop : Face->GetLoops())
-//		{
-//			for (const FOrientedEdge& OrientedEdge : Loop->GetEdges())
-//			{
-//				const TSharedPtr<FTopologicalEdge>& Edge = OrientedEdge.Entity;
-//				for (const TWeakPtr<FTopologicalEdge>& TwinEdge : Edge->GetTwinsEntities())
-//				{
-//					TSharedPtr<FTopologicalFace> Face = TwinEdge.Pin()->GetFace();
-//					if (!Face.IsValid())
-//					{
-//						continue;
-//					}
-//				}
-//
-//				if (Edge->GetStartVertex()->GetLink() == Edge->GetEndVertex()->GetLink())
-//				{
-//					if (!Edge->IsDegenerated())
-//					{
-//						for (const TWeakPtr<FTopologicalEdge>& TwinEdge : Edge->GetTwinsEntities())
-//						{
-//							TwinEdge.Pin()->SetAsDegenerated();
-//						}
-//					}
-//				}
-//			}
-//		}
-//	}
-//
-//	for (TSharedPtr<FTopologicalEdge>& Edge : CollapsedEdges)
-//	{
-//		TSharedPtr<FTopologicalLoop> Loop = Edge->GetLoop();
-//		Loop->RemoveEdge(Edge);
-//
-//		TSharedPtr<FTopologicalVertex> StartVertex = Edge->GetStartVertex();
-//		TSharedPtr<FTopologicalVertex> EndVertex = Edge->GetEndVertex();
-//
-//		Edge->Delete();
-//
-//		StartVertex->DeleteIfIsolated();
-//		EndVertex->DeleteIfIsolated();
-//
-//		if (Loop->GetEdges().Num() == 0)
-//		{
-//			TSharedPtr<FTopologicalFace> Face = Loop->GetFace();
-//			Face->RemoveLoop(Loop);
-//		}
-//	}
-//}
-
-
-
-
+	RemoveEmptyShells();
 
+	FDuration StepDuration = FChrono::Elapse(StartTime);
+	FChrono::PrintClockElapse(EVerboseLevel::Log, TEXT(""), TEXT("Split"), StepDuration);
+}
 
-/*
-NewEdgeVertices.Reserve(AllVertices.Num());
-for (TSharedPtr<FTopologicalVertex> Vertex : AllVertices)
+void FTopomaker::RemoveEmptyShells()
 {
-	TArray<TWeakPtr<FTopologicalEdge>> Edges = Vertex->GetConnectedEdges();
-	if (Edges.Num() == 2)
-	{
-		// TODO Work in 2D
-		FPoint Edge0Tangent;
-		Edges[0].Pin()->GetTangentAt(Vertex, Edge0Tangent);
+	TArray<FBody*> Bodies;
+	Bodies.Reserve(Shells.Num());
 
-		FPoint Edge1Tangent;
-		Edges[1].Pin()->GetTangentAt(Vertex, Edge1Tangent);
+	TArray<FShell*> NewShells;
 
-		double CosAngle = Edge0Tangent.ComputeCosinus(Edge1Tangent);
-		if (CosAngle < -0.75)
+	for (FShell* Shell : Shells)
+	{
+		if (Shell->FaceCount() == 0)
 		{
-			TArray<EOrientation> EdgeDirections;
-			EdgeDirections.Reserve(2);
-			EdgeDirections.Add(Edges[0].Pin()->GetEndVertex() == Vertex ? EOrientation::Front : EOrientation::Back);
-			EdgeDirections.Add(Edges[1].Pin()->GetStartVertex() == Vertex ? EOrientation::Front : EOrientation::Back);
-			TSharedPtr<FTopologicalEdge> ExtendedEdge = nullptr;
-
-			// TODO Check better if the both curve are similar (2DCurve as to be similar)
-			if (Edges[0].Pin()->GetCurve() != Edges[1].Pin()->GetCurve())
+			FBody* Body = (FBody*) Shell->GetHost();
+			if (Body != nullptr && !Body->HasMarker1())
 			{
-				//ExtendedEdge = FEdge::ReplaceEdgesByOne(Edges, EdgeDirections, GetTolerance());
+				Body->SetMarker1();
+				Bodies.Add(Body);
 			}
-			else
-			{
-				//ExtendedEdge = FEdge::CreateEdgeToMerge2Edges(Edges, EdgeDirections);
-			}
-
-			if (ExtendedEdge)
-			{
-				NewEdgeVertices.Add(ExtendedEdge->GetStartVertex());
-				NewEdgeVertices.Add(ExtendedEdge->GetEndVertex());
-			}
-			continue;
+			Shell->Delete();
 		}
+		else
+		{
+			NewShells.Add(Shell);
+		}
+	}
+	Swap(Shells, NewShells);
+
+	FModel& Model = Session.GetModel();
+	for (FBody* Body : Bodies)
+	{
+		Body->RemoveEmptyShell(Model);
+		Body->ResetMarkers();
 	}
 }
-*/
 
+void FTopomaker::OrientShells()
+{
+	FTimePoint StartTime = FChrono::Now();
+
+	for (FShell* Shell : Shells)
+	{
+		int32 FaceSwapCount = Shell->Orient();
+		Report.AddSwappedFaceCount(FaceSwapCount);
+	}
+
+	Report.OrientationDuration = FChrono::Elapse(StartTime);
+	FChrono::PrintClockElapse(EVerboseLevel::Log, TEXT(""), TEXT("Orient"), Report.OrientationDuration);
+}
 
 }
diff --git a/Engine/Source/Runtime/Datasmith/CADKernel/Private/CADKernel/Topo/TopomakerReport.cpp b/Engine/Source/Runtime/Datasmith/CADKernel/Private/CADKernel/Topo/TopomakerReport.cpp
new file mode 100644
--- /dev/null
+++ b/Engine/Source/Runtime/Datasmith/CADKernel/Private/CADKernel/Topo/TopomakerReport.cpp
@@ -0,0 +1,25 @@
+// Copyright Epic Games, Inc. All Rights Reserved.
+
+#include "CADKernel/Topo/TopomakerReport.h"
+
+#include "CADKernel/Topo/Body.h"
+#include "CADKernel/Topo/Shell.h"
+#include "CADKernel/Topo/TopologicalFace.h"
+#include "CADKernel/Topo/TopologicalEdge.h"
+
+#include "CADKernel/UI/Message.h"
+
+namespace CADKernel
+{
+
+void FTopomakerReport::Print()
+{
+	FMessage::FillReportFile(TEXT(""), TEXT(""));
+	FMessage::FillReportFile(TEXT("Sew"), SewDuration);
+	FMessage::FillReportFile(TEXT("Orient"), OrientationDuration);
+	FMessage::FillReportFile(TEXT(""), TEXT(""));
+	FMessage::FillReportFile(TEXT("ShellOrient"), OrientationFixedCount);
+	FMessage::FillReportFile(TEXT("FaceOrient"), SwappedFaceCount);
+}
+
+}
diff --git a/Engine/Source/Runtime/Datasmith/CADKernel/Private/CADKernel/UI/Display.cpp b/Engine/Source/Runtime/Datasmith/CADKernel/Private/CADKernel/UI/Display.cpp
--- a/Engine/Source/Runtime/Datasmith/CADKernel/Private/CADKernel/UI/Display.cpp
+++ b/Engine/Source/Runtime/Datasmith/CADKernel/Private/CADKernel/UI/Display.cpp
@@ -721,7 +721,7 @@ void Draw(const FTopologicalFace& Face)
 			for (const FOrientedEdge& Edge : Loop->GetEdges())
 			{
 				EVisuProperty Property = EVisuProperty::BlueCurve;
-				switch (Edge.Entity->GetTwinsEntityCount())
+				switch (Edge.Entity->GetTwinEntityCount())
 				{
 				case 1:
 					Property = EVisuProperty::BorderEdge;
diff --git a/Engine/Source/Runtime/Datasmith/CADKernel/Public/CADKernel/Core/Database.h b/Engine/Source/Runtime/Datasmith/CADKernel/Public/CADKernel/Core/Database.h
--- a/Engine/Source/Runtime/Datasmith/CADKernel/Public/CADKernel/Core/Database.h
+++ b/Engine/Source/Runtime/Datasmith/CADKernel/Public/CADKernel/Core/Database.h
@@ -10,6 +10,7 @@ namespace CADKernel
 {
 	class FModel;
 	class FTopologicalEntity;
+	class FTopologicalShapeEntity;
 
 	class CADKERNEL_API FDatabase
 	{
@@ -76,7 +77,7 @@ namespace CADKernel
 
 		FDatabase();
 
-		TSharedRef<FModel> GetModel();
+		FModel& GetModel();
 
 		/**
 		 * Remove from database
@@ -254,8 +255,9 @@ namespace CADKernel
 			}
 		}
 
-		uint32 SpawnEntityIdent(const TArray<TSharedPtr<FEntity>>& SelectedEntities, bool bForceSpawning = false);
-		uint32 SpawnEntityIdent(TSharedPtr<FEntity>& SelectedEntity, bool bForceSpawning = false);
+		uint32 SpawnEntityIdents(const TArray<TSharedPtr<FEntity>>& SelectedEntities, bool bForceSpawning = false);
+		uint32 SpawnEntityIdents(const TArray<FEntity*>& SelectedEntities, bool bForceSpawning = false);
+		uint32 SpawnEntityIdent(FEntity& SelectedEntity, bool bForceSpawning = false);
 
 	protected:
 		/**
@@ -289,6 +291,8 @@ namespace CADKernel
 
 		TSharedPtr<FEntity> GetEntity(FIdent id) const;
 		void GetEntities(const TArray<FIdent>& ids, TArray<TSharedPtr<FEntity>>& Entities) const;
+		void GetTopologicalEntities(const TArray<FIdent>& ids, TArray<FTopologicalEntity*>& Entities) const;
+		void GetTopologicalShapeEntities(const TArray<FIdent>& ids, TArray<FTopologicalShapeEntity*>& Entities) const;
 
 		void GetEntitiesOfType(EEntity Type, TArray<TSharedPtr<FEntity>>& OutEntities) const;
 		void GetEntitiesOfTypes(const TSet<EEntity>& Types, TArray<TSharedPtr<FEntity>>& OutEntities) const;
diff --git a/Engine/Source/Runtime/Datasmith/CADKernel/Public/CADKernel/Core/EntityGeom.h b/Engine/Source/Runtime/Datasmith/CADKernel/Public/CADKernel/Core/EntityGeom.h
--- a/Engine/Source/Runtime/Datasmith/CADKernel/Public/CADKernel/Core/EntityGeom.h
+++ b/Engine/Source/Runtime/Datasmith/CADKernel/Public/CADKernel/Core/EntityGeom.h
@@ -10,39 +10,38 @@ class FString;
 
 namespace CADKernel
 {
-	class CADKERNEL_API FEntityGeom : public FEntity
-	{
+
+class CADKERNEL_API FEntityGeom : public FEntity
+{
 	friend class FCoreTechBridge;
 
-	protected:
-		FIdent CtKioId = 0;
+protected:
+	FIdent CtKioId = 0;
 
-	public:
-		FEntityGeom() = default;
+public:
 
-		virtual TSharedPtr<FEntityGeom> ApplyMatrix(const FMatrixH& InMatrix) const
-		{
-			return TSharedPtr<FEntityGeom>();
-		}
+	FEntityGeom() = default;
 
-		virtual void Display(const FString& Name) const
-		{
-		}
+	virtual TSharedPtr<FEntityGeom> ApplyMatrix(const FMatrixH& InMatrix) const
+	{
+		return TSharedPtr<FEntityGeom>();
+	}
 
-		virtual void Serialize(FCADKernelArchive& Ar) override
-		{
-			FEntity::Serialize(Ar);
-			Ar << CtKioId;
-		}
+	virtual void Serialize(FCADKernelArchive& Ar) override
+	{
+		FEntity::Serialize(Ar);
+		Ar << CtKioId;
+	}
 
-		FIdent GetKioId() const
-		{
-			return CtKioId;
-		}
+	FIdent GetKioId() const
+	{
+		return CtKioId;
+	}
 
 #ifdef CADKERNEL_DEV
-		virtual FInfoEntity& GetInfo(FInfoEntity& Info) const override;
+	virtual FInfoEntity& GetInfo(FInfoEntity& Info) const override;
 #endif
-	};
+};
+
 }
 
diff --git a/Engine/Source/Runtime/Datasmith/CADKernel/Public/CADKernel/Core/MetadataDictionary.h b/Engine/Source/Runtime/Datasmith/CADKernel/Public/CADKernel/Core/MetadataDictionary.h
--- a/Engine/Source/Runtime/Datasmith/CADKernel/Public/CADKernel/Core/MetadataDictionary.h
+++ b/Engine/Source/Runtime/Datasmith/CADKernel/Public/CADKernel/Core/MetadataDictionary.h
@@ -7,153 +7,153 @@
 
 namespace CADKernel
 {
-	class CADKERNEL_API FMetadataDictionary
-	{
-
-	private:
-
-		FString Name;
-		int32   HostId = 0;
-		int32   PatchId = 0;
-		uint32  ColorId = 0;
-		uint32  MaterialId = 0;
-		int32   LayerId = -1;
-
-	public:
 
-		FMetadataDictionary() = default;
-
-		void Serialize(FCADKernelArchive& Ar)
+class FMetadataDictionary
+{
+private:
+	FString Name;
+	int32   HostId = 0;
+	int32   PatchId = 0;
+	uint32  ColorId = 0;
+	uint32  MaterialId = 0;
+	int32   LayerId = -1;
+
+public:
+	FMetadataDictionary() = default;
+
+	void Serialize(FCADKernelArchive& Ar)
+	{
+		Ar.Archive << Name;
+		Ar.Archive << HostId;
+		Ar.Archive << PatchId;
+		Ar.Archive << ColorId;
+		Ar.Archive << MaterialId;
+		Ar.Archive << LayerId;
+	}
+
+	void CompleteDictionary(const FMetadataDictionary& MetaData)
+	{
+		if (Name.IsEmpty())
 		{
-			Ar.Archive << Name;
-			Ar.Archive << HostId;
-			Ar.Archive << PatchId;
-			Ar.Archive << ColorId;
-			Ar.Archive << MaterialId;
-			Ar.Archive << LayerId;
+			Name = MetaData.Name;
 		}
-
-		void CompleteDictionary(const FMetadataDictionary& MetaData)
+		if (HostId == 0)
 		{
-			if (Name.IsEmpty())
-			{
-				Name = MetaData.Name;
-			}
-			if (HostId == 0)
-			{
-				HostId = MetaData.HostId;
-			}
-			if (PatchId == 0)
-			{
-				PatchId = MetaData.PatchId;
-			}
-			if (ColorId == 0)
-			{
-				ColorId = MetaData.ColorId;
-			}
-			if (MaterialId == 0)
-			{
-				MaterialId = MetaData.MaterialId;
-			}
-			if (LayerId == 0)
-			{
-				LayerId = MetaData.LayerId;
-			}
+			HostId = MetaData.HostId;
 		}
-
-		void SetHostId(const int32 InHostId)
+		if (PatchId == 0)
 		{
-			HostId = InHostId;
+			PatchId = MetaData.PatchId;
 		}
-
-		const int32 GetHostId() const
+		if (ColorId == 0)
 		{
-			return HostId;
+			ColorId = MetaData.ColorId;
 		}
-
-		void SetLayer(const int32 InLayerId)
+		if (MaterialId == 0)
 		{
-			LayerId = InLayerId;
+			MaterialId = MetaData.MaterialId;
 		}
-
-		int32 GetLayer() const
+		if (LayerId == 0)
 		{
-			return LayerId;
+			LayerId = MetaData.LayerId;
 		}
+	}
 
-		void SetName(const FString& InName)
+	void ExtractMetaData(TMap<FString, FString>& OutMetaData) const 
+	{
+		if (!Name.IsEmpty())
 		{
-			Name = InName;
+			OutMetaData.Add(TEXT("Name"), Name);
 		}
-
-		bool HasName() const
+		if (HostId != 0)
 		{
-			return !Name.IsEmpty();
+			OutMetaData.Add(TEXT("HostId"), FString::FromInt(HostId));
 		}
-
-		const TCHAR* GetName() const
+		if (PatchId != 0)
 		{
-			return *Name;
+			OutMetaData.Add(TEXT("PatchId"), FString::FromInt(PatchId));
 		}
-
-		void SetColorId(const uint32& InColorId)
+		if (ColorId != 0)
 		{
-			ColorId = InColorId;
+			OutMetaData.Add(TEXT("ColorName"), FString::Printf(TEXT("%u"), ColorId));
 		}
-
-		uint32 GetColorId() const
+		if (MaterialId != 0)
 		{
-			return ColorId;
+			OutMetaData.Add(TEXT("MaterialName"), FString::Printf(TEXT("%u"), MaterialId));
 		}
-
-		void SetMaterialId(const uint32& InMaterialId)
+		if (LayerId != 0)
 		{
-			MaterialId = InMaterialId;
+			OutMetaData.Add(TEXT("LayerId"), FString::FromInt(LayerId));
 		}
+	}
 
-		uint32 GetMaterialId() const
-		{
-			return MaterialId;
-		}
+	void SetHostId(const int32 InHostId)
+	{
+		HostId = InHostId;
+	}
 
-		void SetPatchId(int32 InPatchId)
-		{
-			PatchId = InPatchId;
-		}
+	const int32 GetHostId() const
+	{
+		return HostId;
+	}
+
+	void SetLayer(const int32 InLayerId)
+	{
+		LayerId = InLayerId;
+	}
+
+	int32 GetLayer() const
+	{
+		return LayerId;
+	}
+
+	void SetName(const FString& InName)
+	{
+		Name = InName;
+	}
+
+	bool HasName() const
+	{
+		return !Name.IsEmpty();
+	}
+
+	const TCHAR* GetName() const
+	{
+		return *Name;
+	}
+
+	void SetColorId(const uint32& InColorId)
+	{
+		ColorId = InColorId;
+	}
+
+	uint32 GetColorId() const
+	{
+		return ColorId;
+	}
+
+	void SetMaterialId(const uint32& InMaterialId)
+	{
+		MaterialId = InMaterialId;
+	}
+
+	uint32 GetMaterialId() const
+	{
+		return MaterialId;
+	}
+
+	void SetPatchId(int32 InPatchId)
+	{
+		PatchId = InPatchId;
+	}
+
+	int32 GetPatchId() const
+	{
+		return PatchId;
+	}
+};
 
-		int32 GetPatchId() const
-		{
-			return PatchId;
-		}
 
-		void GetMetaData(TMap<FString, FString>& OutMetaData)
-		{
-			if (!Name.IsEmpty())
-			{
-				OutMetaData.Add(TEXT("Name"), Name);
-			}
-			if (HostId != 0)
-			{
-				OutMetaData.Add(TEXT("HostId"), FString::FromInt(HostId));
-			}
-			if (PatchId != 0)
-			{
-				OutMetaData.Add(TEXT("PatchId"), FString::FromInt(PatchId));
-			}
-			if (ColorId != 0)
-			{
-				OutMetaData.Add(TEXT("ColorName"), FString::Printf(TEXT("%u"), ColorId));
-			}
-			if (MaterialId != 0)
-			{
-				OutMetaData.Add(TEXT("MaterialName"), FString::Printf(TEXT("%u"), MaterialId));
-			}
-			if (LayerId != 0)
-			{
-				OutMetaData.Add(TEXT("LayerId"), FString::FromInt(LayerId));
-			}
-		}
-	};
 }
 
 
diff --git a/Engine/Source/Runtime/Datasmith/CADKernel/Public/CADKernel/Core/Session.h b/Engine/Source/Runtime/Datasmith/CADKernel/Public/CADKernel/Core/Session.h
--- a/Engine/Source/Runtime/Datasmith/CADKernel/Public/CADKernel/Core/Session.h
+++ b/Engine/Source/Runtime/Datasmith/CADKernel/Public/CADKernel/Core/Session.h
@@ -31,7 +31,7 @@ namespace CADKernel
 		{
 		}
 
-		TSharedRef<FModel> GetModel();
+		FModel& GetModel();
 
 		void Serialize(FCADKernelArchive& Ar)
 		{
@@ -101,20 +101,19 @@ namespace CADKernel
 		 * This method browses all sub entities and set their Id if needed
 		 * @param bForceSpawning If false, the process does not iterate through the children of entities with a defined ID
 		 */
-		uint32 SpawnEntityIdent(const TArray<TSharedPtr<FEntity>>& SelectedEntities, bool bForceSpawning = false)
+		uint32 SpawnEntityIdent(FEntity& SelectedEntity, bool bForceSpawning = false)
 		{
-			return Database.SpawnEntityIdent(SelectedEntities, bForceSpawning);
+			return Database.SpawnEntityIdent(SelectedEntity, bForceSpawning);
 		}
 
-		uint32 SpawnEntityIdent(TSharedPtr<FEntity> SelectedEntity, bool bForceSpawning)
+		uint32 SpawnEntityIdents(const TArray<TSharedPtr<FEntity>>& SelectedEntities, bool bForceSpawning = false)
 		{
-			return Database.SpawnEntityIdent(SelectedEntity, bForceSpawning);
+			return Database.SpawnEntityIdents(SelectedEntities, bForceSpawning);
 		}
 
-		template<typename PointType>
-		uint32 SpawnEntityIdent(TSharedPtr<PointType>& SelectedEntity, bool bForceSpawning = false)
+		uint32 SpawnEntityIdents(const TArray<FEntity*>& SelectedEntities, bool bForceSpawning = false)
 		{
-			return Database.SpawnEntityIdent((TSharedPtr<FEntity>&) SelectedEntity, bForceSpawning);
+			return Database.SpawnEntityIdents(SelectedEntities, bForceSpawning);
 		}
 
 		int32 GetLastHostId() const 
diff --git a/Engine/Source/Runtime/Datasmith/CADKernel/Public/CADKernel/Geo/Sampling/Polyline.h b/Engine/Source/Runtime/Datasmith/CADKernel/Public/CADKernel/Geo/Sampling/Polyline.h
--- a/Engine/Source/Runtime/Datasmith/CADKernel/Public/CADKernel/Geo/Sampling/Polyline.h
+++ b/Engine/Source/Runtime/Datasmith/CADKernel/Public/CADKernel/Geo/Sampling/Polyline.h
@@ -11,198 +11,208 @@
 
 namespace CADKernel
 {
-	template<class PointType>
-	class TPolyline
-	{
-	public:
 
-		TArray<double> Coordinates;
-		TArray<PointType> Points;
-		TArray<PointType> Tangent;
-		TPolylineApproximator<PointType> Approximator;
-		bool bWithTangent = false;
+template<class PointType>
+class TPolyline
+{
+public:
 
+	TArray<double> Coordinates;
+	TArray<PointType> Points;
+	TArray<PointType> Tangent;
+	TPolylineApproximator<PointType> Approximator;
+	bool bWithTangent = false;
 
-		TPolyline(const TArray<PointType>& InPoints, const TArray<double>& InCoordinates)
-			: Coordinates(InCoordinates)
-			, Points(InPoints)
-			, Approximator(Coordinates, Points)
-		{
-		}
 
-		TPolyline(const TArray<PointType>& InPoints)
-			: Points(InPoints)
-			, Approximator(Coordinates, Points)
-		{
-		}
+	TPolyline(const TArray<PointType>& InPoints, const TArray<double>& InCoordinates)
+		: Coordinates(InCoordinates)
+		, Points(InPoints)
+		, Approximator(Coordinates, Points)
+	{
+	}
 
-		TPolyline()
-			: Approximator(Coordinates, Points)
-		{
-		}
+	TPolyline(const TArray<PointType>& InPoints)
+		: Points(InPoints)
+		, Approximator(Coordinates, Points)
+	{
+	}
 
-		void Serialize(FCADKernelArchive& Ar)
-		{
-			Ar << Points;
-			Ar << Coordinates;
-		}
+	TPolyline()
+		: Approximator(Coordinates, Points)
+	{
+	}
 
-		PointType ApproximatePoint(double InCoordinate) const
-		{
-			return Approximator.ApproximatePoint(InCoordinate);
-		}
+	void Serialize(FCADKernelArchive& Ar)
+	{
+		Ar << Points;
+		Ar << Coordinates;
+	}
 
-		void ApproximatePoints(const TArray<double>& InCoordinates, TArray<PointType>& OutPoints) const
-		{
-			Approximator.ApproximatePoints(InCoordinates, OutPoints);
-		}
+	PointType ApproximatePoint(double InCoordinate) const
+	{
+		return Approximator.ApproximatePoint(InCoordinate);
+	}
 
-		void GetSubPolyline(const FLinearBoundary& InBoundary, const EOrientation InOrientation, TArray<PointType>& OutPoints) const
-		{
-			Approximator.GetSubPolyline(InBoundary, InOrientation, OutPoints);
-		}
+	void ApproximatePoints(const TArray<double>& InCoordinates, TArray<PointType>& OutPoints) const
+	{
+		Approximator.ApproximatePoints(InCoordinates, OutPoints);
+	}
 
-		void Sample(const FLinearBoundary& Boundary, const double DesiredSegmentLength, TArray<double>& OutCoordinates) const
-		{
-			Approximator.SamplePolyline(Boundary, DesiredSegmentLength, OutCoordinates);
-		}
+	void GetSubPolyline(const FLinearBoundary& InBoundary, const EOrientation InOrientation, TArray<PointType>& OutPoints) const
+	{
+		Approximator.GetSubPolyline(InBoundary, InOrientation, OutPoints);
+	}
 
-		double GetCoordinateOfProjectedPoint(const FLinearBoundary& Boundary, const PointType& PointOnEdge, PointType& ProjectedPoint) const
-		{
-			return Approximator.ProjectPointToPolyline(Boundary, PointOnEdge, ProjectedPoint);
-		}
+	/**
+	 * The FPolylineBBox is updated by adding all the points of the polyline included in the boundary
+	 */
+	void UpdateSubPolylineBBox(const FLinearBoundary& InBoundary, FPolylineBBox& OutBBox) const
+	{
+		Approximator.UpdateSubPolylineBBox(InBoundary, OutBBox);
+	}
 
-		void ProjectPoints(const FLinearBoundary& InBoundary, const TArray<PointType>& InPointsToProject, TArray<double>& ProjectedPointCoordinates, TArray<PointType>& ProjectedPoints) const
-		{
-			Approximator.ProjectPointsToPolyline(InBoundary, InPointsToProject, ProjectedPointCoordinates, ProjectedPoints);
-		}
+	void Sample(const FLinearBoundary& Boundary, const double DesiredSegmentLength, TArray<double>& OutCoordinates) const
+	{
+		Approximator.SamplePolyline(Boundary, DesiredSegmentLength, OutCoordinates);
+	}
 
-		const PointType& GetPointAt(int32 Index) const
-		{
-			return Points[Index];
-		}
+	double GetCoordinateOfProjectedPoint(const FLinearBoundary& Boundary, const PointType& PointOnEdge, PointType& ProjectedPoint) const
+	{
+		return Approximator.ProjectPointToPolyline(Boundary, PointOnEdge, ProjectedPoint);
+	}
 
-		const TArray<PointType>& GetPoints() const
-		{
-			return Points;
-		}
+	void ProjectPoints(const FLinearBoundary& InBoundary, const TArray<PointType>& InPointsToProject, TArray<double>& ProjectedPointCoordinates, TArray<PointType>& ProjectedPoints) const
+	{
+		Approximator.ProjectPointsToPolyline(InBoundary, InPointsToProject, ProjectedPointCoordinates, ProjectedPoints);
+	}
 
-		const TArray<double>& GetCoordinates() const
-		{
-			return Coordinates;
-		}
+	const PointType& GetPointAt(int32 Index) const
+	{
+		return Points[Index];
+	}
 
-		TArray<double>& GetCoordinates()
-		{
-			return Coordinates;
-		}
+	const TArray<PointType>& GetPoints() const
+	{
+		return Points;
+	}
 
-		void SwapCoordinates(TArray<double>& NewCoordinates)
-		{
-			Swap(NewCoordinates, Coordinates);
-		}
+	const TArray<double>& GetCoordinates() const
+	{
+		return Coordinates;
+	}
 
-		void GetAt(int32 Index, double& Coordinate, PointType& Point)
-		{
-			ensureCADKernel(Coordinates.IsValidIndex(Index));
-			Coordinate = Coordinates[Index];
-			Point = Points[Index];
-		}
+	TArray<double>& GetCoordinates()
+	{
+		return Coordinates;
+	}
 
-		/**
-		 * @return the size of the polyline i.e. the count of points.
-		 */
-		int32 Size() const
-		{
-			return Points.Num();
-		}
+	void SwapCoordinates(TArray<double>& NewCoordinates)
+	{
+		Swap(NewCoordinates, Coordinates);
+	}
 
-		/**
-		 * Reserves memory such that the polyline can contain at least Number elements.
-		 *
-		 * @param Number The number of elements that the polyline should be able to contain after allocation.
-		 */
-		void Reserve(int32 Number)
-		{
-			Points.Reserve(Number);
-			Coordinates.Reserve(Number);
-		}
+	void GetAt(int32 Index, double& Coordinate, PointType& Point)
+	{
+		ensureCADKernel(Coordinates.IsValidIndex(Index));
+		Coordinate = Coordinates[Index];
+		Point = Points[Index];
+	}
+
+	/**
+	 * @return the size of the polyline i.e. the count of points.
+	 */
+	int32 Size() const
+	{
+		return Points.Num();
+	}
+
+	/**
+	 * Reserves memory such that the polyline can contain at least Number elements.
+	 *
+	 * @param Number The number of elements that the polyline should be able to contain after allocation.
+	 */
+	void Reserve(int32 Number)
+	{
+		Points.Reserve(Number);
+		Coordinates.Reserve(Number);
+	}
+
+	/**
+	 * Empties the polyline.
+	 *
+	 * @param Slack (Optional) The expected usage size after empty operation. Default is 0.
+	 */
+	void Empty(int32 Slack = 0)
+	{
+		Points.Empty(Slack);
+		Coordinates.Empty(Slack);
+	}
 
-		/**
-		 * Empties the polyline.
-		 *
-		 * @param Slack (Optional) The expected usage size after empty operation. Default is 0.
-		 */
-		void Empty(int32 Slack = 0)
-		{
-			Points.Empty(Slack);
-			Coordinates.Empty(Slack);
-		}
+	void EmplaceAt(int32 Index, TPolyline<PointType>& Polyline, int32 PointIndex)
+	{
+		Coordinates.EmplaceAt(Index, Polyline.Coordinates[PointIndex]);
+		Points.EmplaceAt(Index, Polyline.Points[PointIndex]);
+	}
 
-		void EmplaceAt(int32 Index, TPolyline<PointType>& Polyline, int32 PointIndex)
-		{
-			Coordinates.EmplaceAt(Index, Polyline.Coordinates[PointIndex]);
-			Points.EmplaceAt(Index, Polyline.Points[PointIndex]);
-		}
+	void RemoveAt(int32 Index)
+	{
+		Coordinates.RemoveAt(Index);
+		Points.RemoveAt(Index);
+	}
 
-		void RemoveAt(int32 Index)
-		{
-			Coordinates.RemoveAt(Index);
-			Points.RemoveAt(Index);
-		}
+	void Pop()
+	{
+		Coordinates.Pop(false);
+		Points.Pop(false);
+	}
 
-		void Pop()
-		{
-			Coordinates.Pop(false);
-			Points.Pop(false);
-		}
+	double GetLength(const FLinearBoundary& InBoundary) const
+	{
+		return Approximator.ComputeLengthOfSubPolyline(InBoundary);
+	}
+};
 
-		double GetLength(const FLinearBoundary& InBoundary) const
-		{
-			return Approximator.ComputeLengthOfSubPolyline(InBoundary);
-		}
-	};
+typedef TPolyline<FPoint> FPolyline3D;
 
-	typedef TPolyline<FPoint> FPolyline3D;
+class FPolyline2D : public TPolyline<FPoint2D>
+{
+public:
 
-	class FPolyline2D : public TPolyline<FPoint2D>
+	/**
+	 * Add to OutIntersection the curve coordinate of the intersection with the Iso curve
+	 */
+	void FindIntersectionsWithIso(const EIso Iso, double IsoParameter, TArray<double>& OutIntersections)
 	{
-	public:
+		OutIntersections.Empty(8);
 
-		/**
-		 * Add to OutIntersection the curve coordinate of the intersection with the Iso curve
-		 */
-		void FindIntersectionsWithIso(const EIso Iso, double IsoParameter, TArray<double>& OutIntersections)
-		{
-			OutIntersections.Empty(8);
+		int32 UIndex = Iso == EIso::IsoU ? 0 : 1;
+		int32 VIndex = Iso == EIso::IsoU ? 1 : 0;
 
-			int32 UIndex = Iso == EIso::IsoU ? 0 : 1;
-			int32 VIndex = Iso == EIso::IsoU ? 1 : 0;
+		TFunction<void(const int32, const int32, const int32)> ComputeIntersection = [&](const int32 Index1, const int32& Index2, const int32 MinIndex)
+		{
+			if (IsoParameter > Points[Index1][UIndex] && IsoParameter <= Points[Index2][UIndex])
+			{
+				double Intersection = (IsoParameter - Points[Index1][UIndex]) / (Points[Index2][UIndex] - Points[Index1][UIndex]);
+				OutIntersections.Add(PolylineTools::LinearInterpolation(Coordinates, MinIndex, Intersection));
+			}
+		};
 
-			TFunction<void(const int32, const int32, const int32)> ComputeIntersection = [&](const int32 Index1, const int32& Index2, const int32 MinIndex)
+		for (int32 Index = 1; Index < Points.Num(); ++Index)
+		{
+			if (!FMath::IsNearlyEqual(Points[Index - 1][UIndex], Points[Index][UIndex]))
 			{
-				if (IsoParameter > Points[Index1][UIndex] && IsoParameter <= Points[Index2][UIndex])
+				if (Points[Index - 1][UIndex] < Points[Index][UIndex])
 				{
-					double Intersection = (IsoParameter - Points[Index1][UIndex]) / (Points[Index2][UIndex] - Points[Index1][UIndex]);
-					OutIntersections.Add(PolylineTools::LinearInterpolation(Coordinates, MinIndex, Intersection));
+					ComputeIntersection(Index - 1, Index, Index - 1);
 				}
-			};
-
-			for (int32 Index = 1; Index < Points.Num(); ++Index)
-			{
-				if (!FMath::IsNearlyEqual(Points[Index - 1][UIndex], Points[Index][UIndex]))
+				else
 				{
-					if (Points[Index - 1][UIndex] < Points[Index][UIndex])
-					{
-						ComputeIntersection(Index - 1, Index, Index - 1);
-					}
-					else
-					{
-						ComputeIntersection(Index, Index - 1, Index - 1);
-					}
+					ComputeIntersection(Index, Index - 1, Index - 1);
 				}
 			}
 		}
-	};
+	}
 };
 
+} // ns CADKernel
+
diff --git a/Engine/Source/Runtime/Datasmith/CADKernel/Public/CADKernel/Geo/Sampling/PolylineTools.h b/Engine/Source/Runtime/Datasmith/CADKernel/Public/CADKernel/Geo/Sampling/PolylineTools.h
--- a/Engine/Source/Runtime/Datasmith/CADKernel/Public/CADKernel/Geo/Sampling/PolylineTools.h
+++ b/Engine/Source/Runtime/Datasmith/CADKernel/Public/CADKernel/Geo/Sampling/PolylineTools.h
@@ -16,6 +16,42 @@
 namespace CADKernel
 {
 
+struct FPolylineBBox
+{
+	FPoint Max;
+	FPoint Min;
+	FPoint MaxPoints[3];
+	FPoint MinPoints[3];
+	double CoordinateOfMaxPoint[3];
+	double CoordinateOfMinPoint[3];
+
+	FPolylineBBox()
+		: Max(-HUGE_VALUE, -HUGE_VALUE, -HUGE_VALUE)
+		, Min(HUGE_VALUE, HUGE_VALUE, HUGE_VALUE)
+	{
+	}
+
+	void Update(const double Coordinate, const FPoint& Point)
+	{
+		for (int32 Index = 0; Index < 3; ++Index)
+		{
+			if (Point[Index] > Max[Index])
+			{
+				Max[Index] = Point[Index];
+				MaxPoints[Index] = Point;
+				CoordinateOfMaxPoint[Index] = Coordinate;
+			}
+
+			if (Point[Index] < Min[Index])
+			{
+				Min[Index] = Point[Index];
+				MinPoints[Index] = Point;
+				CoordinateOfMinPoint[Index] = Coordinate;
+			}
+		}
+	}
+};
+
 namespace PolylineTools
 {
 
@@ -552,32 +588,59 @@ public:
 	 */
 	void GetSubPolyline(const FLinearBoundary& InBoundary, TArray<double>& OutCoordinates, TArray<PointType>& OutPoints) const
 	{
-		int32 StartIndex = 0;
-		int32 EndIndex = 0;
-		GetStartEndIndex(PolylineCoordinates, InBoundary, StartIndex, EndIndex);
+		int32 BoundaryIndices[2];
+		GetStartEndIndex(InBoundary, BoundaryIndices);
 
-		OutCoordinates.Empty(EndIndex - StartIndex + 2);
-		OutPoints.Empty(EndIndex - StartIndex + 2);
+		int32 NewSize = BoundaryIndices[1] - BoundaryIndices[0] + 2;
 
-		if (FMath::IsNearlyEqual(PolylineCoordinates[StartIndex + 1], InBoundary.Min, (double)SMALL_NUMBER))
+		OutCoordinates.Empty(NewSize);
+		OutPoints.Empty(NewSize);
+
+		if (FMath::IsNearlyEqual(PolylineCoordinates[BoundaryIndices[0] + 1], InBoundary.Min, (double)SMALL_NUMBER))
 		{
-			StartIndex++;
+			BoundaryIndices[0]++;
 		}
 
-		if (!FMath::IsNearlyEqual(PolylineCoordinates[EndIndex], InBoundary.Max, (double)SMALL_NUMBER))
+		if (FMath::IsNearlyEqual(PolylineCoordinates[BoundaryIndices[1]], InBoundary.Max, (double)SMALL_NUMBER))
 		{
-			EndIndex--;
+			BoundaryIndices[1]--;
 		}
 
-		OutPoints.Emplace(ComputePoint(StartIndex, InBoundary.Min));
-		OutPoints.Append(PolylinePoints.GetData() + StartIndex + 1, EndIndex - StartIndex);
-		OutPoints.Emplace(ComputePoint(EndIndex, InBoundary.Max));
+		OutPoints.Emplace(ComputePoint(BoundaryIndices[0], InBoundary.Min));
+		OutPoints.Append(PolylinePoints.GetData() + BoundaryIndices[0] + 1, BoundaryIndices[1] - BoundaryIndices[0]);
+		OutPoints.Emplace(ComputePoint(BoundaryIndices[1], InBoundary.Max));
 
 		OutCoordinates.Add(InBoundary.Min);
-		OutCoordinates.Append(PolylineCoordinates.GetData() + StartIndex + 1, EndIndex - StartIndex);
+		OutCoordinates.Append(PolylineCoordinates.GetData() + BoundaryIndices[0] + 1, BoundaryIndices[1] - BoundaryIndices[0]);
 		OutCoordinates.Add(InBoundary.Max);
 	}
 
+	/**
+	 * Update the Curve bounding box with this subset of polyline
+	 */
+	void UpdateSubPolylineBBox(const FLinearBoundary& InBoundary, FPolylineBBox& OutBBox) const
+	{
+		int32 BoundaryIndices[2];
+		GetStartEndIndex(InBoundary, BoundaryIndices);
+
+		if (FMath::IsNearlyEqual(PolylineCoordinates[BoundaryIndices[0] + 1], InBoundary.Min, (double)SMALL_NUMBER))
+		{
+			BoundaryIndices[0]++;
+		}
+
+		if (FMath::IsNearlyEqual(PolylineCoordinates[BoundaryIndices[1]], InBoundary.Max, (double)SMALL_NUMBER))
+		{
+			BoundaryIndices[1]--;
+		}
+
+		OutBBox.Update(InBoundary.Min, ComputePoint(BoundaryIndices[0], InBoundary.Min));
+		OutBBox.Update(InBoundary.Max, ComputePoint(BoundaryIndices[1], InBoundary.Max));
+		for (int32 Index = BoundaryIndices[0] + 1; Index <= BoundaryIndices[1]; ++Index)
+		{
+			OutBBox.Update(PolylineCoordinates[Index], PolylinePoints[Index]);
+		}
+	}
+
 	double ComputeLength() const
 	{
 		return PolylineTools::ComputeLength(PolylinePoints);
diff --git a/Engine/Source/Runtime/Datasmith/CADKernel/Public/CADKernel/Geo/Sampling/SurfacicPolyline.h b/Engine/Source/Runtime/Datasmith/CADKernel/Public/CADKernel/Geo/Sampling/SurfacicPolyline.h
--- a/Engine/Source/Runtime/Datasmith/CADKernel/Public/CADKernel/Geo/Sampling/SurfacicPolyline.h
+++ b/Engine/Source/Runtime/Datasmith/CADKernel/Public/CADKernel/Geo/Sampling/SurfacicPolyline.h
@@ -12,415 +12,416 @@
 
 namespace CADKernel
 {
-	class FInfoEntity;
 
-	class FSurfacicPolyline
-	{
+class FInfoEntity;
 
-	public:
+class FSurfacicPolyline
+{
 
-		TArray<double> Coordinates;
-		TArray<FPoint2D> Points2D;
-		TArray<FPoint> Points3D;
-		TArray<FVector> Normals;
-		TArray<FPoint> Tangents;
+public:
 
-		FSurfacicBoundary BoundingBox;
+	TArray<double> Coordinates;
+	TArray<FPoint2D> Points2D;
+	TArray<FPoint> Points3D;
+	TArray<FVector> Normals;
+	TArray<FPoint> Tangents;
 
-		bool bWithNormals;
-		bool bWithTangent;
+	FSurfacicBoundary BoundingBox;
 
-		FSurfacicPolyline(TSharedRef<FSurface> InCarrierSurface, TSharedRef<FCurve> InCurve2D);
+	bool bWithNormals;
+	bool bWithTangent;
 
-		FSurfacicPolyline(TSharedRef<FSurface> InCarrierSurface, TSharedRef<FCurve> InCurve2D, const double Tolerance);
+	FSurfacicPolyline(TSharedRef<FSurface> InCarrierSurface, TSharedRef<FCurve> InCurve2D);
 
-		FSurfacicPolyline(TSharedRef<FSurface> InCarrierSurface, TSharedRef<FCurve> InCurve2D, const double ChordTolerance, const double ParamTolerance, bool bInWithNormals/* = false*/, bool bWithTangent/* = false*/);
+	FSurfacicPolyline(TSharedRef<FSurface> InCarrierSurface, TSharedRef<FCurve> InCurve2D, const double Tolerance);
 
-		FSurfacicPolyline(bool bInWithNormals = false, bool bInWithTangent = false)
-			: bWithNormals(bInWithNormals)
-			, bWithTangent(bInWithTangent)
-		{
-		}
+	FSurfacicPolyline(TSharedRef<FSurface> InCarrierSurface, TSharedRef<FCurve> InCurve2D, const double ChordTolerance, const double ParamTolerance, bool bInWithNormals/* = false*/, bool bWithTangent/* = false*/);
 
-		void Serialize(FCADKernelArchive& Ar)
-		{
-			Ar.Serialize(Points3D);
-			Ar.Serialize(Points2D);
-			Ar.Serialize(Normals);
-			Ar.Serialize(Coordinates);
-			Ar.Serialize(BoundingBox);
-			Ar << bWithNormals;
-			Ar << bWithTangent;
-		}
+	FSurfacicPolyline(bool bInWithNormals = false, bool bInWithTangent = false)
+		: bWithNormals(bInWithNormals)
+		, bWithTangent(bInWithTangent)
+	{
+	}
+
+	void Serialize(FCADKernelArchive& Ar)
+	{
+		Ar.Serialize(Points3D);
+		Ar.Serialize(Points2D);
+		Ar.Serialize(Normals);
+		Ar.Serialize(Coordinates);
+		Ar.Serialize(BoundingBox);
+		Ar << bWithNormals;
+		Ar << bWithTangent;
+	}
+
+	FInfoEntity& GetInfo(FInfoEntity&) const;
 
-		FInfoEntity& GetInfo(FInfoEntity&) const;
+	TSharedPtr<FEntityGeom> ApplyMatrix(const FMatrixH&) const;
 
-		TSharedPtr<FEntityGeom> ApplyMatrix(const FMatrixH&) const;
+	void CheckIfDegenerated(const double Tolerance3D, const FSurfacicTolerance& Tolerances2D, const FLinearBoundary& Boudary, bool& bDegeneration2D, bool& bDegeneration3D, double& Length3D) const;
 
-		void CheckIfDegenerated(const double Tolerance3D, const FSurfacicTolerance& Tolerances2D, const FLinearBoundary& Boudary, bool& bDegeneration2D, bool& bDegeneration3D, double& Length3D) const;
+	void GetExtremities(const FLinearBoundary& InBoundary, const double Tolerance3D, const FSurfacicTolerance& Tolerances2D, FSurfacicCurveExtremities& Extremities) const;
 
-		void GetExtremities(const FLinearBoundary& InBoundary, const double Tolerance3D, const FSurfacicTolerance& Tolerances2D, FSurfacicCurveExtremities& Extremities) const;
+	FPoint Approximate3DPoint(double InCoordinate) const
+	{
+		TPolylineApproximator<FPoint> Approximator3D(Coordinates, Points3D);
+		return Approximator3D.ApproximatePoint(InCoordinate);
+	}
+
+	void Approximate3DPoints(const TArray<double>& InCoordinates, TArray<FPoint>& OutPoints) const
+	{
+		TPolylineApproximator<FPoint> Approximator3D(Coordinates, Points3D);
+		Approximator3D.ApproximatePoints(InCoordinates, OutPoints);
+	}
+
+	FPoint2D Approximate2DPoint(double InCoordinate) const
+	{
+		TPolylineApproximator<FPoint2D> Approximator(Coordinates, Points2D);
+		return Approximator.ApproximatePoint(InCoordinate);
+	}
 
-		FPoint Approximate3DPoint(double InCoordinate) const
+	FPoint GetTangentAt(double InCoordinate) const
+	{
+		FDichotomyFinder Finder(Coordinates);
+		int32 Index = Finder.Find(InCoordinate);
+		return Points3D[Index + 1] - Points3D[Index];
+	}
+
+	FPoint2D GetTangent2DAt(double InCoordinate) const
+	{
+		FDichotomyFinder Finder(Coordinates);
+		int32 Index = Finder.Find(InCoordinate);
+		return Points2D[Index + 1] - Points2D[Index];
+	}
+
+	FSurfacicTolerance ComputeTolerance(const double Tolerance3D, const FSurfacicTolerance& MinToleranceIso, const int32 Index) const
+	{
+		double Distance3D = Points3D[Index].Distance(Points3D[Index + 1]);
+		if (FMath::IsNearlyZero(Distance3D, (double)SMALL_NUMBER))
 		{
-			TPolylineApproximator<FPoint> Approximator3D(Coordinates, Points3D);
-			return Approximator3D.ApproximatePoint(InCoordinate);
+			return FPoint2D::FarawayPoint;
 		}
-
-		void Approximate3DPoints(const TArray<double>& InCoordinates, TArray<FPoint>& OutPoints) const
+		else
 		{
-			TPolylineApproximator<FPoint> Approximator3D(Coordinates, Points3D);
-			Approximator3D.ApproximatePoints(InCoordinates, OutPoints);
+			FSurfacicTolerance Tolerance2D = Points2D[Index] - Points2D[Index + 1];
+			return Max(Abs(Tolerance2D) * Tolerance3D / Distance3D, MinToleranceIso);
 		}
+	};
 
-		FPoint2D Approximate2DPoint(double InCoordinate) const
+	double ComputeLinearToleranceAt(const double Tolerance3D, const double MinLinearTolerance, const int32 Index) const
+	{
+		double Distance3D = Points3D[Index].Distance(Points3D[Index + 1]);
+		if (FMath::IsNearlyZero(Distance3D, (double)SMALL_NUMBER))
 		{
-			TPolylineApproximator<FPoint2D> Approximator(Coordinates, Points2D);
-			return Approximator.ApproximatePoint(InCoordinate);
+			return (Coordinates.Last() - Coordinates[0]) / 10.;
 		}
-
-		FPoint GetTangentAt(double InCoordinate) const
+		else
 		{
-			FDichotomyFinder Finder(Coordinates);
-			int32 Index = Finder.Find(InCoordinate);
-			return Points3D[Index + 1] - Points3D[Index];
+			double LinearDistance = Coordinates[Index + 1] - Coordinates[Index];
+			return FMath::Max(LinearDistance / Distance3D * Tolerance3D, MinLinearTolerance);
 		}
+	};
 
-		FPoint2D GetTangent2DAt(double InCoordinate) const
+	void Approximate2DPoints(const TArray<double>& InCoordinates, TArray<FPoint2D>& OutPoints) const
+	{
+		TPolylineApproximator<FPoint2D> Approximator(Coordinates, Points2D);
+		Approximator.ApproximatePoints(InCoordinates, OutPoints);
+	}
+
+	void ApproximatePolyline(FSurfacicPolyline& OutPolyline) const
+	{
+		if (OutPolyline.Coordinates.IsEmpty())
 		{
-			FDichotomyFinder Finder(Coordinates);
-			int32 Index = Finder.Find(InCoordinate);
-			return Points2D[Index + 1] - Points2D[Index];
+			return;
 		}
 
-		FSurfacicTolerance ComputeTolerance(const double Tolerance3D, const FSurfacicTolerance& MinToleranceIso, const int32 Index) const
+		TFunction<void(FIndexOfCoordinateFinder&)> ComputePoints = [&](FIndexOfCoordinateFinder& Finder)
 		{
-			double Distance3D = Points3D[Index].Distance(Points3D[Index + 1]);
-			if (FMath::IsNearlyZero(Distance3D, (double) SMALL_NUMBER))
-			{
-				return FPoint2D::FarawayPoint;
-			}
-			else
-			{
-				FSurfacicTolerance Tolerance2D = Points2D[Index] - Points2D[Index + 1];
-				return Max(Abs(Tolerance2D)*Tolerance3D/Distance3D, MinToleranceIso);
-			}
-		};
+			int32 CoordinateCount = OutPolyline.Coordinates.Num();
 
-		double ComputeLinearToleranceAt(const double Tolerance3D, const double MinLinearTolerance, const int32 Index) const
-		{
-			double Distance3D = Points3D[Index].Distance(Points3D[Index + 1]);
-			if (FMath::IsNearlyZero(Distance3D, (double)SMALL_NUMBER))
+			TArray<int32> SegmentIndexes;
+			SegmentIndexes.Reserve(CoordinateCount);
+
+			TArray<double> SegmentCoordinates;
+			SegmentCoordinates.Reserve(CoordinateCount);
+
+
+			//for (int32 Index = 0; Index < CoordinateCount; ++Index)
+			for (const double Coordinate : OutPolyline.Coordinates)
 			{
-				return (Coordinates.Last() - Coordinates[0]) / 10.;
+				const int32& Index = SegmentIndexes.Emplace_GetRef(Finder.Find(Coordinate));
+				SegmentCoordinates.Emplace(PolylineTools::SectionCoordinate(Coordinates, Index, Coordinate));
 			}
-			else
+
+			OutPolyline.Points2D.Reserve(CoordinateCount);
+			for (int32 Index = 0; Index < CoordinateCount; ++Index)
 			{
-				double LinearDistance = Coordinates[Index+1] - Coordinates[Index];
-				return FMath::Max(LinearDistance / Distance3D * Tolerance3D, MinLinearTolerance);
+				int32 SegmentIndex = SegmentIndexes[Index];
+				double SegmentCoordinate = SegmentCoordinates[Index];
+				OutPolyline.Points2D.Emplace(PolylineTools::LinearInterpolation(Points2D, SegmentIndex, SegmentCoordinate));
 			}
-		};
 
-		void Approximate2DPoints(const TArray<double>& InCoordinates, TArray<FPoint2D>& OutPoints) const
-		{
-			TPolylineApproximator<FPoint2D> Approximator(Coordinates, Points2D);
-			Approximator.ApproximatePoints(InCoordinates, OutPoints);
-		}
-
-		void ApproximatePolyline(FSurfacicPolyline& OutPolyline) const
-		{
-			if (OutPolyline.Coordinates.IsEmpty())
+			OutPolyline.Points3D.Reserve(CoordinateCount);
+			for (int32 Index = 0; Index < CoordinateCount; ++Index)
 			{
-				return;
+				int32 SegmentIndex = SegmentIndexes[Index];
+				double SegmentCoordinate = SegmentCoordinates[Index];
+				OutPolyline.Points3D.Emplace(PolylineTools::LinearInterpolation(Points3D, SegmentIndex, SegmentCoordinate));
 			}
 
-			TFunction<void(FIndexOfCoordinateFinder&)> ComputePoints = [&](FIndexOfCoordinateFinder& Finder)
+			if (bWithNormals)
 			{
-				int32 CoordinateCount = OutPolyline.Coordinates.Num();
-
-				TArray<int32> SegmentIndexes;
-				SegmentIndexes.Reserve(CoordinateCount);
-
-				TArray<double> SegmentCoordinates;
-				SegmentCoordinates.Reserve(CoordinateCount);
-
-
-				//for (int32 Index = 0; Index < CoordinateCount; ++Index)
-				for (const double Coordinate : OutPolyline.Coordinates)
-				{
-					const int32& Index = SegmentIndexes.Emplace_GetRef(Finder.Find(Coordinate));
-					SegmentCoordinates.Emplace(PolylineTools::SectionCoordinate(Coordinates, Index, Coordinate));
-				}
-
-				OutPolyline.Points2D.Reserve(CoordinateCount);
 				for (int32 Index = 0; Index < CoordinateCount; ++Index)
 				{
 					int32 SegmentIndex = SegmentIndexes[Index];
 					double SegmentCoordinate = SegmentCoordinates[Index];
-					OutPolyline.Points2D.Emplace(PolylineTools::LinearInterpolation(Points2D, SegmentIndex, SegmentCoordinate));
+					OutPolyline.Normals.Emplace(PolylineTools::LinearInterpolation(Normals, SegmentIndex, SegmentCoordinate));
 				}
+			}
 
-				OutPolyline.Points3D.Reserve(CoordinateCount);
+			if (bWithTangent)
+			{
 				for (int32 Index = 0; Index < CoordinateCount; ++Index)
 				{
 					int32 SegmentIndex = SegmentIndexes[Index];
 					double SegmentCoordinate = SegmentCoordinates[Index];
-					OutPolyline.Points3D.Emplace(PolylineTools::LinearInterpolation(Points3D, SegmentIndex, SegmentCoordinate));
+					OutPolyline.Tangents.Emplace(PolylineTools::LinearInterpolation(Tangents, SegmentIndex, SegmentCoordinate));
 				}
+			}
+		};
 
-				if (bWithNormals)
-				{
-					for (int32 Index = 0; Index < CoordinateCount; ++Index)
-					{
-						int32 SegmentIndex = SegmentIndexes[Index];
-						double SegmentCoordinate = SegmentCoordinates[Index];
-						OutPolyline.Normals.Emplace(PolylineTools::LinearInterpolation(Normals, SegmentIndex, SegmentCoordinate));
-					}
-				}
+		FDichotomyFinder DichotomyFinder(Coordinates);
+		int32 StartIndex = 0;
+		int32 EndIndex;
 
-				if (bWithTangent)
-				{
-					for (int32 Index = 0; Index < CoordinateCount; ++Index)
-					{
-						int32 SegmentIndex = SegmentIndexes[Index];
-						double SegmentCoordinate = SegmentCoordinates[Index];
-						OutPolyline.Tangents.Emplace(PolylineTools::LinearInterpolation(Tangents, SegmentIndex, SegmentCoordinate));
-					}
-				}
-			};
+		bool bUseDichotomy = false;
+		StartIndex = DichotomyFinder.Find(OutPolyline.Coordinates[0]);
+		EndIndex = DichotomyFinder.Find(OutPolyline.Coordinates.Last());
+		bUseDichotomy = PolylineTools::IsDichotomyToBePreferred(EndIndex - StartIndex, Coordinates.Num());
 
-			FDichotomyFinder DichotomyFinder(Coordinates);
-			int32 StartIndex = 0;
-			int32 EndIndex;
 
-			bool bUseDichotomy = false;
-			StartIndex = DichotomyFinder.Find(OutPolyline.Coordinates[0]);
-			EndIndex = DichotomyFinder.Find(OutPolyline.Coordinates.Last());
-			bUseDichotomy = PolylineTools::IsDichotomyToBePreferred(EndIndex - StartIndex, Coordinates.Num());
+		if (bUseDichotomy)
+		{
+			DichotomyFinder.StartLower = StartIndex;
+			DichotomyFinder.StartUpper = EndIndex;
+			ComputePoints(DichotomyFinder);
+		}
+		else
+		{
+			FLinearFinder LinearFinder(Coordinates, StartIndex);
+			ComputePoints(LinearFinder);
+		}
+	}
 
 
-			if (bUseDichotomy)
-			{
-				DichotomyFinder.StartLower = StartIndex;
-				DichotomyFinder.StartUpper = EndIndex;
-				ComputePoints(DichotomyFinder);
-			}
-			else
-			{
-				FLinearFinder LinearFinder(Coordinates, StartIndex);
-				ComputePoints(LinearFinder);
-			}
-		}
+	void Sample(const FLinearBoundary& Boundary, const double DesiredSegmentLength, TArray<double>& OutCoordinates) const
+	{
+		TPolylineApproximator<FPoint> Approximator3D(Coordinates, Points3D);
+		Approximator3D.SamplePolyline(Boundary, DesiredSegmentLength, OutCoordinates);
+	}
 
+	double GetCoordinateOfProjectedPoint(const FLinearBoundary& Boundary, const FPoint& PointOnEdge, FPoint& ProjectedPoint) const
+	{
+		TPolylineApproximator<FPoint> Approximator3D(Coordinates, Points3D);
+		return Approximator3D.ProjectPointToPolyline(Boundary, PointOnEdge, ProjectedPoint);
+	}
 
-		void Sample(const FLinearBoundary& Boundary, const double DesiredSegmentLength, TArray<double>& OutCoordinates) const
-		{
-			TPolylineApproximator<FPoint> Approximator3D(Coordinates, Points3D);
-			Approximator3D.SamplePolyline(Boundary, DesiredSegmentLength, OutCoordinates);
-		}
+	double GetCoordinateOfProjectedPoint(const FLinearBoundary& Boundary, const FPoint2D& PointOnEdge, FPoint2D& ProjectedPoint) const
+	{
+		TPolylineApproximator<FPoint2D> Approximator2D(Coordinates, Points2D);
+		return Approximator2D.ProjectPointToPolyline(Boundary, PointOnEdge, ProjectedPoint);
+	}
 
-		double GetCoordinateOfProjectedPoint(const FLinearBoundary& Boundary, const FPoint& PointOnEdge, FPoint& ProjectedPoint) const
-		{
-			TPolylineApproximator<FPoint> Approximator3D(Coordinates, Points3D);
-			return Approximator3D.ProjectPointToPolyline(Boundary, PointOnEdge, ProjectedPoint);
-		}
+	void ProjectPoints(const FLinearBoundary& InBoundary, const TArray<FPoint>& InPointsToProject, TArray<double>& ProjectedPointCoordinates, TArray<FPoint>& ProjectedPoints) const
+	{
+		TPolylineApproximator<FPoint> Approximator3D(Coordinates, Points3D);
+		Approximator3D.ProjectPointsToPolyline(InBoundary, InPointsToProject, ProjectedPointCoordinates, ProjectedPoints);
+	}
 
-		double GetCoordinateOfProjectedPoint(const FLinearBoundary& Boundary, const FPoint2D& PointOnEdge, FPoint2D& ProjectedPoint) const
-		{
-			TPolylineApproximator<FPoint2D> Approximator2D(Coordinates, Points2D);
-			return Approximator2D.ProjectPointToPolyline(Boundary, PointOnEdge, ProjectedPoint);
-		}
+	void ProjectPoints(const FLinearBoundary& InBoundary, const TArray<FPoint2D>& InPointsToProject, TArray<double>& ProjectedPointCoordinates, TArray<FPoint2D>& ProjectedPoints) const
+	{
+		TPolylineApproximator<FPoint2D> Approximator(Coordinates, Points2D);
+		Approximator.ProjectPointsToPolyline(InBoundary, InPointsToProject, ProjectedPointCoordinates, ProjectedPoints);
+	}
+
+	/**
+	 * Project each point of a coincidental polyline on the Polyline.
+	 */
+	void ProjectCoincidentalPolyline(const TArray<FPoint>& InPointsToProject, bool bSameOrientation, TArray<double>& OutProjectedPointCoordinates) const
+	{
+		TPolylineApproximator<FPoint> Approximator3D(Coordinates, Points3D);
+		Approximator3D.ProjectCoincidentalPolyline(InPointsToProject, bSameOrientation, OutProjectedPointCoordinates);
+	}
 
-		void ProjectPoints(const FLinearBoundary& InBoundary, const TArray<FPoint>& InPointsToProject, TArray<double>& ProjectedPointCoordinates, TArray<FPoint>& ProjectedPoints) const
-		{
-			TPolylineApproximator<FPoint> Approximator3D(Coordinates, Points3D);
-			Approximator3D.ProjectPointsToPolyline(InBoundary, InPointsToProject, ProjectedPointCoordinates, ProjectedPoints);
-		}
+	/**
+	 * The main idea of this algorithm is to process starting for the beginning of the curve to the end of the curve.
+	 */
+	void ComputeIntersectionsWithIsos(const FLinearBoundary& InBoundary, const TArray<double>& InIsoCoordinates, const EIso InTypeIso, const FSurfacicTolerance& ToleranceIso, TArray<double>& OutIntersection) const;
 
-		void ProjectPoints(const FLinearBoundary& InBoundary, const TArray<FPoint2D>& InPointsToProject, TArray<double>& ProjectedPointCoordinates, TArray<FPoint2D>& ProjectedPoints) const
-		{
-			TPolylineApproximator<FPoint2D> Approximator(Coordinates, Points2D);
-			Approximator.ProjectPointsToPolyline(InBoundary, InPointsToProject, ProjectedPointCoordinates, ProjectedPoints);
-		}
+	const TArray<double>& GetCoordinates() const
+	{
+		return Coordinates;
+	}
 
-		/**
-		 * Project each point of a coincidental polyline on the Polyline.
-		 */
-		void ProjectCoincidentalPolyline(const TArray<FPoint>& InPointsToProject, bool bSameOrientation, TArray<double>& OutProjectedPointCoordinates) const
-		{
-			TPolylineApproximator<FPoint> Approximator3D(Coordinates, Points3D);
-			Approximator3D.ProjectCoincidentalPolyline(InPointsToProject, bSameOrientation, OutProjectedPointCoordinates);
-		}
+	const TArray<FPoint2D>& Get2DPoints() const
+	{
+		return Points2D;
+	}
 
-		/**
-		 * The main idea of this algorithm is to process starting for the beginning of the curve to the end of the curve.
-		 */
-		void ComputeIntersectionsWithIsos(const FLinearBoundary& InBoundary, const TArray<double>& InIsoCoordinates, const EIso InTypeIso, const FSurfacicTolerance& ToleranceIso, TArray<double>& OutIntersection) const;
+	const FPoint& GetPointAt(int32 Index) const
+	{
+		return Points3D[Index];
+	}
 
-		const TArray<double>& GetCoordinates() const
-		{
-			return Coordinates;
-		}
+	const TArray<FPoint>& GetPoints() const
+	{
+		return Points3D;
+	}
 
-		const TArray<FPoint2D>& Get2DPoints() const
-		{
-			return Points2D;
-		}
+	const TArray<FVector>& GetNormals() const
+	{
+		return Normals;
+	}
 
-		const FPoint& GetPointAt(int32 Index) const
-		{
-			return Points3D[Index];
-		}
+	const TArray<FPoint>& GetTangents() const
+	{
+		return Tangents;
+	}
 
-		const TArray<FPoint>& GetPoints() const
-		{
-			return Points3D;
-		}
+	void SwapCoordinates(TArray<double>& NewCoordinates)
+	{
+		Swap(NewCoordinates, Coordinates);
+	}
 
-		const TArray<FVector>& GetNormals() const
-		{
-			return Normals;
-		}
+	/**
+	 * @return the size of the polyline i.e. the count of points.
+	 */
+	int32 Size() const
+	{
+		return Points2D.Num();
+	}
 
-		const TArray<FPoint>& GetTangents() const
+	/**
+	 * Get the sub 2d polyline bounded by the input InBoundary in the orientation of the input InOrientation and append it to the output OutPoints
+	 */
+	void GetSubPolyline(const FLinearBoundary& InBoundary, const EOrientation InOrientation, TArray<FPoint2D>& OutPoints) const
+	{
+		TPolylineApproximator<FPoint2D> Approximator(Coordinates, Points2D);
+		Approximator.GetSubPolyline(InBoundary, InOrientation, OutPoints);
+	}
+
+	/**
+	 * Get the sub polyline bounded by the input InBoundary in the orientation of the input InOrientation and append it to the output OutPoints
+	 */
+	void GetSubPolyline(const FLinearBoundary& InBoundary, const EOrientation InOrientation, TArray<FPoint>& OutPoints) const
+	{
+		TPolylineApproximator<FPoint> Approximator3D(Coordinates, Points3D);
+		Approximator3D.GetSubPolyline(InBoundary, InOrientation, OutPoints);
+	}
+
+	/**
+	 * Reserves memory such that the polyline can contain at least Number elements.
+	 *
+	 * @param Number The number of elements that the polyline should be able to contain after allocation.
+	 */
+	void Reserve(int32 Number)
+	{
+		Points3D.Reserve(Number);
+		Points2D.Reserve(Number);
+		Coordinates.Reserve(Number);
+		if (bWithNormals)
 		{
-			return Tangents;
+			Normals.Reserve(Number);
 		}
+	}
 
-		void SwapCoordinates(TArray<double>& NewCoordinates)
-		{
-			Swap(NewCoordinates, Coordinates);
-		}
+	/**
+	 * Empties the polyline.
+	 *
+	 * @param Slack (Optional) The expected usage size after empty operation. Default is 0.
+	 */
+	void Empty(int32 Slack = 0)
+	{
+		Points3D.Empty(Slack);
+		Points2D.Empty(Slack);
+		Normals.Empty(Slack);
+		Coordinates.Empty(Slack);
+	}
 
-		/**
-		 * @return the size of the polyline i.e. the count of points.
-		 */
-		int32 Size() const
-		{
-			return Points2D.Num();
-		}
+	void EmplaceAt(int32 Index, FSurfacicPolyline& Polyline, int32 PointIndex)
+	{
+		Coordinates.EmplaceAt(Index, Polyline.Coordinates[PointIndex]);
 
-		/**
-		 * Get the sub 2d polyline bounded by the input InBoundary in the orientation of the input InOrientation and append it to the output OutPoints
-		 */
-		void GetSubPolyline(const FLinearBoundary& InBoundary, const EOrientation InOrientation, TArray<FPoint2D>& OutPoints) const
+		Points2D.EmplaceAt(Index, Polyline.Points2D[PointIndex]);
+		Points3D.EmplaceAt(Index, Polyline.Points3D[PointIndex]);
+		if (bWithNormals)
 		{
-			TPolylineApproximator<FPoint2D> Approximator(Coordinates, Points2D);
-			Approximator.GetSubPolyline(InBoundary, InOrientation, OutPoints);
+			Normals.EmplaceAt(Index, Polyline.Normals[PointIndex]);
 		}
-
-		/**
-		 * Get the sub polyline bounded by the input InBoundary in the orientation of the input InOrientation and append it to the output OutPoints
-		 */
-		void GetSubPolyline(const FLinearBoundary& InBoundary, const EOrientation InOrientation, TArray<FPoint>& OutPoints) const
+		if (bWithTangent)
 		{
-			TPolylineApproximator<FPoint> Approximator3D(Coordinates, Points3D);
-			Approximator3D.GetSubPolyline(InBoundary, InOrientation, OutPoints);
+			Tangents.EmplaceAt(Index, Polyline.Tangents[PointIndex]);
 		}
+	}
 
-		/**
-		 * Reserves memory such that the polyline can contain at least Number elements.
-		 *
-		 * @param Number The number of elements that the polyline should be able to contain after allocation.
-		 */
-		void Reserve(int32 Number)
+	void RemoveAt(int32 Index)
+	{
+		Coordinates.RemoveAt(Index);
+		Points2D.RemoveAt(Index);
+		Points3D.RemoveAt(Index);
+		if (bWithNormals)
 		{
-			Points3D.Reserve(Number);
-			Points2D.Reserve(Number);
-			Coordinates.Reserve(Number);
-			if (bWithNormals)
-			{
-				Normals.Reserve(Number);
-			}
+			Normals.RemoveAt(Index);
 		}
-
-		/**
-		 * Empties the polyline.
-		 *
-		 * @param Slack (Optional) The expected usage size after empty operation. Default is 0.
-		 */
-		void Empty(int32 Slack = 0)
+		if (bWithTangent)
 		{
-			Points3D.Empty(Slack);
-			Points2D.Empty(Slack);
-			Normals.Empty(Slack);
-			Coordinates.Empty(Slack);
+			Tangents.RemoveAt(Index);
 		}
+	}
 
-		void EmplaceAt(int32 Index, FSurfacicPolyline& Polyline, int32 PointIndex)
+	void Pop()
+	{
+		Coordinates.Pop(false);
+		Points2D.Pop(false);
+		Points3D.Pop(false);
+		if (bWithNormals)
 		{
-			Coordinates.EmplaceAt(Index, Polyline.Coordinates[PointIndex]);
-
-			Points2D.EmplaceAt(Index, Polyline.Points2D[PointIndex]);
-			Points3D.EmplaceAt(Index, Polyline.Points3D[PointIndex]);
-			if (bWithNormals)
-			{
-				Normals.EmplaceAt(Index, Polyline.Normals[PointIndex]);
-			}
-			if (bWithTangent)
-			{
-				Tangents.EmplaceAt(Index, Polyline.Tangents[PointIndex]);
-			}
+			Normals.Pop(false);
 		}
-
-		void RemoveAt(int32 Index)
+		if (bWithTangent)
 		{
-			Coordinates.RemoveAt(Index);
-			Points2D.RemoveAt(Index);
-			Points3D.RemoveAt(Index);
-			if (bWithNormals)
-			{
-				Normals.RemoveAt(Index);
-			}
-			if (bWithTangent)
-			{
-				Tangents.RemoveAt(Index);
-			}
+			Tangents.Pop(false);
 		}
+	}
 
-		void Pop()
+	void Reverse()
+	{
+		Algo::Reverse(Coordinates);
+		Algo::Reverse(Points2D);
+		Algo::Reverse(Points3D);
+		if (bWithNormals)
 		{
-			Coordinates.Pop(false);
-			Points2D.Pop(false);
-			Points3D.Pop(false);
-			if (bWithNormals)
-			{
-				Normals.Pop(false);
-			}
-			if (bWithTangent)
-			{
-				Tangents.Pop(false);
-			}
+			Algo::Reverse(Normals);
 		}
-
-		void Reverse()
+		if (bWithTangent)
 		{
-			Algo::Reverse(Coordinates);
-			Algo::Reverse(Points2D);
-			Algo::Reverse(Points3D);
-			if (bWithNormals)
-			{
-				Algo::Reverse(Normals);
-			}
-			if (bWithTangent)
-			{
-				Algo::Reverse(Tangents);
-			}
+			Algo::Reverse(Tangents);
 		}
+	}
 
-		double GetLength(const FLinearBoundary& InBoundary) const
-		{
-			TPolylineApproximator<FPoint> Approximator3D(Coordinates, Points3D);
-			return Approximator3D.ComputeLengthOfSubPolyline(InBoundary);
-		}
+	double GetLength(const FLinearBoundary& InBoundary) const
+	{
+		TPolylineApproximator<FPoint> Approximator3D(Coordinates, Points3D);
+		return Approximator3D.ComputeLengthOfSubPolyline(InBoundary);
+	}
 
-		double Get2DLength(const FLinearBoundary& InBoundary) const
-		{
-			TPolylineApproximator<FPoint2D> Approximator(Coordinates, Points2D);
-			return Approximator.ComputeLengthOfSubPolyline(InBoundary);
-		}
+	double Get2DLength(const FLinearBoundary& InBoundary) const
+	{
+		TPolylineApproximator<FPoint2D> Approximator(Coordinates, Points2D);
+		return Approximator.ComputeLengthOfSubPolyline(InBoundary);
+	}
 
-	};
+};
 
-}
+} // ns CADKernel
 
diff --git a/Engine/Source/Runtime/Datasmith/CADKernel/Public/CADKernel/Geo/Surfaces/Surface.h b/Engine/Source/Runtime/Datasmith/CADKernel/Public/CADKernel/Geo/Surfaces/Surface.h
--- a/Engine/Source/Runtime/Datasmith/CADKernel/Public/CADKernel/Geo/Surfaces/Surface.h
+++ b/Engine/Source/Runtime/Datasmith/CADKernel/Public/CADKernel/Geo/Surfaces/Surface.h
@@ -187,6 +187,17 @@ namespace CADKernel
 
 		virtual void EvaluateNormals(const TArray<FPoint2D>& Points2D, TArray<FVector>& Normals) const;
 
+		virtual FVector EvaluateNormal(const FPoint2D& InPoint2D) const
+		{
+			int32 DerivativeOrder = 1;
+			FSurfacicPoint Point3D;
+			EvaluatePoint(InPoint2D, Point3D, DerivativeOrder);
+
+			FVector Normal = Point3D.GradientU ^ Point3D.GradientV;
+			Normal.Normalize();
+			return Normal;
+		}
+
 		/**
 		 * Divide the parametric space in the desired number of regular subdivisions and compute the associated PointGrid
 		 */
diff --git a/Engine/Source/Runtime/Datasmith/CADKernel/Public/CADKernel/Mesh/Meshers/MesherReport.h b/Engine/Source/Runtime/Datasmith/CADKernel/Public/CADKernel/Mesh/Meshers/MesherReport.h
--- a/Engine/Source/Runtime/Datasmith/CADKernel/Public/CADKernel/Mesh/Meshers/MesherReport.h
+++ b/Engine/Source/Runtime/Datasmith/CADKernel/Public/CADKernel/Mesh/Meshers/MesherReport.h
@@ -83,19 +83,6 @@ struct CADKERNEL_API FMesherChronos
 	FDuration GlobalMeshThinZones;
 
 	FMesherChronos()
-		: GlobalDuration(FChrono::Init())
-		, ApplyCriteriaDuration(FChrono::Init())
-		, IsolateQuadPatchDuration(FChrono::Init())
-		, GlobalMeshDuration(FChrono::Init())
-		, GlobalPointCloudDuration(FChrono::Init())
-		, GlobalGeneratePointCloudDuration(FChrono::Init())
-		, GlobalTriangulateDuration(FChrono::Init())
-		, GlobalDelaunayDuration(FChrono::Init())
-		, GlobalMeshAndGetLoopNodes(FChrono::Init())
-		, GlobalMeshEdges(FChrono::Init())
-		, GlobalThinZones(FChrono::Init())
-		, GlobalFindThinZones(FChrono::Init())
-		, GlobalMeshThinZones(FChrono::Init())
 	{}
 
 	void PrintTimeElapse() const
diff --git a/Engine/Source/Runtime/Datasmith/CADKernel/Public/CADKernel/Mesh/Meshers/ParametricMesher.h b/Engine/Source/Runtime/Datasmith/CADKernel/Public/CADKernel/Mesh/Meshers/ParametricMesher.h
--- a/Engine/Source/Runtime/Datasmith/CADKernel/Public/CADKernel/Mesh/Meshers/ParametricMesher.h
+++ b/Engine/Source/Runtime/Datasmith/CADKernel/Public/CADKernel/Mesh/Meshers/ParametricMesher.h
@@ -50,12 +50,9 @@ protected:
 	 */
 	const double ConstMinCurvature = 0.001;
 
-	TSharedRef<FModelMesh> MeshModel;
-	TSharedRef<FMesherParameters> Parameters;
+	FModelMesh& MeshModel;
 
-	TArray<TSharedPtr<FTopologicalFace>> Faces;
-	TArray<TSharedPtr<FTopologicalEdge>> Edges;
-	TArray<TSharedPtr<FTopologicalVertex>> Vertices;
+	TArray<FTopologicalFace*> Faces;
 
 	FMesherReport MesherReport;
 
@@ -63,42 +60,24 @@ protected:
 
 public:
 
-	FParametricMesher(TSharedRef<FModelMesh> MeshModel);
+	FParametricMesher(FModelMesh& MeshModel);
 
-	const TSharedRef<FModelMesh>& GetMeshModel() const
+	const FModelMesh& GetMeshModel() const
 	{
 		return MeshModel;
 	}
 
-	TSharedRef<FModelMesh>& GetMeshModel()
+	FModelMesh& GetMeshModel()
 	{
 		return MeshModel;
 	}
 
-	void InitParameters(const FString& ParametersString);
+	void MeshEntities(TArray<FTopologicalShapeEntity*>& InEntities);
 
-	const TSharedRef<FMesherParameters>& GetParameters() const
+	void MeshEntity(FTopologicalShapeEntity& InEntity)
 	{
-		return Parameters;
-	}
-
-	void MeshEntities(TArray<TSharedPtr<FEntity>>& InEntities);
-
-	template<typename EntityType>
-	void MeshEntity(TSharedRef<EntityType>& Entity)
-	{
-		MeshEntity((TSharedRef<FTopologicalEntity>&) Entity);
-	}
-
-	void MeshEntity(TSharedRef<FModel>& InModel)
-	{
-		MeshEntity((TSharedRef<FTopologicalEntity>&) InModel);
-	}
-
-	void MeshEntity(TSharedRef<FTopologicalEntity>& InEntity)
-	{
-		TArray<TSharedPtr<FEntity>> Entities;
-		Entities.Add(InEntity);
+		TArray<FTopologicalShapeEntity*> Entities;
+		Entities.Add(&InEntity);
 		MeshEntities(Entities);
 	}
 
@@ -127,7 +106,7 @@ protected:
 	void MeshSurfaceByFront(TArray<FCostToFace>& QuadTrimmedSurfaceSet);
 
 	void ApplyEdgeCriteria(FTopologicalEdge& Edge);
-	void ApplyFaceCriteria(TSharedRef<FTopologicalFace> Face);
+	void ApplyFaceCriteria(FTopologicalFace& Face);
 
 	/**
 	 * Generate Edge Elements on active edge from Edge cutting points
diff --git a/Engine/Source/Runtime/Datasmith/CADKernel/Public/CADKernel/Mesh/Structure/EdgeMesh.h b/Engine/Source/Runtime/Datasmith/CADKernel/Public/CADKernel/Mesh/Structure/EdgeMesh.h
--- a/Engine/Source/Runtime/Datasmith/CADKernel/Public/CADKernel/Mesh/Structure/EdgeMesh.h
+++ b/Engine/Source/Runtime/Datasmith/CADKernel/Public/CADKernel/Mesh/Structure/EdgeMesh.h
@@ -19,7 +19,7 @@ namespace CADKernel
 		 */
 		TArray<int32> EdgeVerticesIndex;
 
-		FEdgeMesh(TSharedRef<FModelMesh> Model, TSharedRef<FTopologicalEntity> TopologicalEntity)
+		FEdgeMesh(FModelMesh& Model, FTopologicalEntity& TopologicalEntity)
 			: FMesh(Model, TopologicalEntity)
 		{
 		}
diff --git a/Engine/Source/Runtime/Datasmith/CADKernel/Public/CADKernel/Mesh/Structure/FaceMesh.h b/Engine/Source/Runtime/Datasmith/CADKernel/Public/CADKernel/Mesh/Structure/FaceMesh.h
--- a/Engine/Source/Runtime/Datasmith/CADKernel/Public/CADKernel/Mesh/Structure/FaceMesh.h
+++ b/Engine/Source/Runtime/Datasmith/CADKernel/Public/CADKernel/Mesh/Structure/FaceMesh.h
@@ -27,7 +27,7 @@ namespace CADKernel
 		TArray<FVector2D> UVMap;
 
 	public:
-		FFaceMesh(TSharedRef<FModelMesh> InMeshModel, TSharedRef<FTopologicalEntity> InTopologicalEntity)
+		FFaceMesh(FModelMesh& InMeshModel, FTopologicalEntity& InTopologicalEntity)
 			: FMesh(InMeshModel, InTopologicalEntity)
 		{
 		}
diff --git a/Engine/Source/Runtime/Datasmith/CADKernel/Public/CADKernel/Mesh/Structure/Mesh.h b/Engine/Source/Runtime/Datasmith/CADKernel/Public/CADKernel/Mesh/Structure/Mesh.h
--- a/Engine/Source/Runtime/Datasmith/CADKernel/Public/CADKernel/Mesh/Structure/Mesh.h
+++ b/Engine/Source/Runtime/Datasmith/CADKernel/Public/CADKernel/Mesh/Structure/Mesh.h
@@ -24,10 +24,10 @@ namespace CADKernel
 
 	public:
 
-		FMesh(TSharedRef<FModelMesh> InMeshModel, TSharedRef<FTopologicalEntity> InTopologicalEntity)
+		FMesh(FModelMesh& InMeshModel, FTopologicalEntity& InTopologicalEntity)
 			: FEntityGeom()
-			, ModelMesh(*InMeshModel)
-			, TopologicalEntity(*InTopologicalEntity)
+			, ModelMesh(InMeshModel)
+			, TopologicalEntity(InTopologicalEntity)
 		{
 		}
 
diff --git a/Engine/Source/Runtime/Datasmith/CADKernel/Public/CADKernel/Mesh/Structure/VertexMesh.h b/Engine/Source/Runtime/Datasmith/CADKernel/Public/CADKernel/Mesh/Structure/VertexMesh.h
--- a/Engine/Source/Runtime/Datasmith/CADKernel/Public/CADKernel/Mesh/Structure/VertexMesh.h
+++ b/Engine/Source/Runtime/Datasmith/CADKernel/Public/CADKernel/Mesh/Structure/VertexMesh.h
@@ -12,7 +12,7 @@ namespace CADKernel
 	protected:
 
 	public:
-		FVertexMesh(TSharedRef<FModelMesh> Model, TSharedRef<FTopologicalEntity> TopologicalEntity)
+		FVertexMesh(FModelMesh& Model, FTopologicalEntity& TopologicalEntity)
 			: FMesh(Model, TopologicalEntity)
 		{
 		}
diff --git a/Engine/Source/Runtime/Datasmith/CADKernel/Public/CADKernel/Topo/Body.h b/Engine/Source/Runtime/Datasmith/CADKernel/Public/CADKernel/Topo/Body.h
--- a/Engine/Source/Runtime/Datasmith/CADKernel/Public/CADKernel/Topo/Body.h
+++ b/Engine/Source/Runtime/Datasmith/CADKernel/Public/CADKernel/Topo/Body.h
@@ -9,105 +9,114 @@
 
 namespace CADKernel
 {
-	class FShell;
 
-	class CADKERNEL_API FBody : public FTopologicalEntity, public FMetadataDictionary
-	{
-		friend FEntity;
+class FShell;
+
+class CADKERNEL_API FBody : public FTopologicalShapeEntity
+{
+	friend FEntity;
 
-	private:
-		TArray<TSharedPtr<FShell>> Shells;
+private:
+	TArray<TSharedPtr<FShell>> Shells;
 
-		FBody() = default;
+	FBody() = default;
 
-		FBody(const TArray<TSharedPtr<FShell>>& InShells)
+	FBody(const TArray<TSharedPtr<FShell>>& InShells)
+	{
+		for (TSharedPtr<FShell> Shell : InShells)
 		{
-			for (TSharedPtr<FShell> Shell : InShells)
+			if (Shell.IsValid())
 			{
-				if (Shell.IsValid())
-				{
-					AddShell(Shell.ToSharedRef());
-				}
+				AddShell(Shell.ToSharedRef());
 			}
 		}
+	}
 
-	public:
+public:
 
-		virtual void Serialize(FCADKernelArchive& Ar) override
-		{
-			FTopologicalEntity::Serialize(Ar);
-			SerializeIdents(Ar, Shells);
-			FMetadataDictionary::Serialize(Ar);
-		}
-
-		virtual void SpawnIdent(FDatabase& Database) override
-		{
-			if (!FEntity::SetId(Database))
-			{
-				return;
-			}
+	virtual void Serialize(FCADKernelArchive& Ar) override
+	{
+		FTopologicalShapeEntity::Serialize(Ar);
+		SerializeIdents(Ar, Shells);
+	}
 
-			SpawnIdentOnEntities(Shells, Database);
-		}
 
-		virtual void ResetMarkersRecursively() override
+	virtual void SpawnIdent(FDatabase& Database) override
+	{
+		if (!FEntity::SetId(Database))
 		{
-			ResetMarkers();
-			ResetMarkersRecursivelyOnEntities(Shells);
+			return;
 		}
 
+		SpawnIdentOnEntities(Shells, Database);
+	}
+
+	virtual void ResetMarkersRecursively() override
+	{
+		ResetMarkers();
+		ResetMarkersRecursivelyOnEntities(Shells);
+	}
+
 #ifdef CADKERNEL_DEV
-		virtual FInfoEntity& GetInfo(FInfoEntity&) const override;
+	virtual FInfoEntity& GetInfo(FInfoEntity&) const override;
 #endif
 
-		virtual EEntity GetEntityType() const override
-		{
-			return EEntity::Body;
-		}
+	virtual EEntity GetEntityType() const override
+	{
+		return EEntity::Body;
+	}
 
-		void AddShell(TSharedRef<FShell> Shell);
+	void AddShell(TSharedRef<FShell> Shell);
 
-		void RemoveEmptyShell();
+	void RemoveEmptyShell(FModel& Model);
 
-		void Empty()
-		{
-			Shells.Empty();
-		}
+	void Empty()
+	{
+		Shells.Empty();
+	}
 
-		const TArray<TSharedPtr<FShell>>& GetShells() const
+	const TArray<TSharedPtr<FShell>>& GetShells() const
+	{
+		return Shells;
+	}
+
+	virtual int32 FaceCount() const override
+	{
+		int32 FaceCount = 0;
+		for (const TSharedPtr<FShell>& Shell : Shells)
 		{
-			return Shells;
+			FaceCount += Shell->FaceCount();
 		}
+		return FaceCount;
+	}
 
-		virtual int32 FaceCount() const override
+	virtual void GetFaces(TArray<FTopologicalFace*>& Faces) override
+	{
+		for (const TSharedPtr<FShell>& Shell : Shells)
 		{
-			int32 FaceCount = 0;
-			for (const TSharedPtr<FShell>& Shell : Shells)
-			{
-				FaceCount += Shell->FaceCount();
-			}
-			return FaceCount;
+			Shell->GetFaces(Faces);
 		}
+	}
 
-		virtual void GetFaces(TArray<TSharedPtr<FTopologicalFace>>& Faces) override
+	virtual void SpreadBodyOrientation() override
+	{
+		for (const TSharedPtr<FShell>& Shell : Shells)
 		{
-			for (const TSharedPtr<FShell>& Shell : Shells)
-			{
-				Shell->GetFaces(Faces);
-			}
+			Shell->SpreadBodyOrientation();
 		}
+	}
 
-		virtual void SpreadBodyOrientation() override
+	void Orient()
+	{
+		for (const TSharedPtr<FShell>& Shell : Shells)
 		{
-			for (const TSharedPtr<FShell>& Shell : Shells)
-			{
-				Shell->SpreadBodyOrientation();
-			}
+			Shell->Orient();
 		}
+	}
+
+	virtual void FillTopologyReport(FTopologyReport& Report) const override;
 
-		virtual void FillTopologyReport(FTopologyReport& Report) const override;
+};
 
-		virtual TSharedPtr<FEntityGeom> ApplyMatrix(const FMatrixH& InMatrix) const override;
-	};
 }
 
diff --git a/Engine/Source/Runtime/Datasmith/CADKernel/Public/CADKernel/Topo/Linkable.h b/Engine/Source/Runtime/Datasmith/CADKernel/Public/CADKernel/Topo/Linkable.h
--- a/Engine/Source/Runtime/Datasmith/CADKernel/Public/CADKernel/Topo/Linkable.h
+++ b/Engine/Source/Runtime/Datasmith/CADKernel/Public/CADKernel/Topo/Linkable.h
@@ -27,7 +27,7 @@ public:
 
 	virtual void Serialize(FCADKernelArchive& Ar) override
 	{
-		FEntityGeom::Serialize(Ar);
+		FTopologicalEntity::Serialize(Ar);
 		SerializeIdent(Ar, TopologicalLink);
 	}
 
@@ -49,7 +49,7 @@ public:
 	{
 		ensureCADKernel(TopologicalLink.IsValid());
 
-		if (TopologicalLink->GetTwinsEntitieNum() == 1)
+		if (TopologicalLink->GetTwinEntityNum() == 1)
 		{
 			return true;
 		}
@@ -90,21 +90,21 @@ public:
 		return (Entity->TopologicalLink == TopologicalLink);
 	}
 
-	int32 GetTwinsEntityCount() const
+	int32 GetTwinEntityCount() const
 	{
 		ensureCADKernel(TopologicalLink.IsValid());
-		return TopologicalLink->GetTwinsEntitieNum();
+		return TopologicalLink->GetTwinEntityNum();
 	}
 
 	bool HasTwin() const
 	{
-		return GetTwinsEntityCount() != 1;
+		return GetTwinEntityCount() != 1;
 	}
 
-	const TArray<EntityType*>& GetTwinsEntities() const
+	const TArray<EntityType*>& GetTwinEntities() const
 	{
 		ensureCADKernel(TopologicalLink.IsValid());
-		return TopologicalLink->GetTwinsEntities();
+		return TopologicalLink->GetTwinEntities();
 	}
 
 	virtual void RemoveFromLink()
@@ -119,7 +119,11 @@ public:
 	 */
 	void UnlinkTwinEntities()
 	{
-		TopologicalLink->UnlinkTwinEntities();
+		TArray<EntityType*> Twins = GetTwinEntities();
+		for (EntityType* Entity : Twins)
+		{
+			Entity->ResetTopologicalLink();
+		}
 	}
 
 	const bool IsThinZone() const
@@ -151,13 +155,13 @@ protected:
 			return;
 		}
 
-		if (Link2->GetTwinsEntities().Num() > Link1->GetTwinsEntities().Num())
+		if (Link2->GetTwinEntityNum() > Link1->GetTwinEntityNum())
 		{
 			Swap(Link1, Link2);
 		}
 
-		Link1->AddEntities(Link2->GetTwinsEntities());
-		for (EntityType* Entity : Link2->GetTwinsEntities())
+		Link1->AddEntities(Link2->GetTwinEntities());
+		for (EntityType* Entity : Link2->GetTwinEntities())
 		{
 			Entity->SetTopologicalLink(Link1);
 		}
diff --git a/Engine/Source/Runtime/Datasmith/CADKernel/Public/CADKernel/Topo/Model.h b/Engine/Source/Runtime/Datasmith/CADKernel/Public/CADKernel/Topo/Model.h
--- a/Engine/Source/Runtime/Datasmith/CADKernel/Public/CADKernel/Topo/Model.h
+++ b/Engine/Source/Runtime/Datasmith/CADKernel/Public/CADKernel/Topo/Model.h
@@ -2,7 +2,7 @@
 #pragma once
 
 #include "CADKernel/Core/EntityGeom.h"
-#include "CADKernel/Topo/TopologicalFace.h"
+#include "CADKernel/Topo/TopologicalShapeEntity.h"
 
 namespace CADKernel
 {
@@ -13,7 +13,7 @@ namespace CADKernel
 	class FGroup;
 	class FTopologicalVertex;
 
-	class CADKERNEL_API FModel : public FTopologicalEntity, public FMetadataDictionary
+	class CADKERNEL_API FModel : public FTopologicalShapeEntity
 	{
 		friend FEntity;
 
@@ -32,10 +32,9 @@ namespace CADKernel
 
 		virtual void Serialize(FCADKernelArchive& Ar) override
 		{
-			FTopologicalEntity::Serialize(Ar);
+			FTopologicalShapeEntity::Serialize(Ar);
 			SerializeIdents(Ar, (TArray<TSharedPtr<FEntity>>&) Bodies);
 			SerializeIdents(Ar, (TArray<TSharedPtr<FEntity>>&) Faces);
-			FMetadataDictionary::Serialize(Ar);
 
 			if(Ar.IsLoading())
 			{
@@ -112,8 +111,16 @@ namespace CADKernel
 			Bodies.Remove(InBody);
 		}
 
+		void RemoveBody(const FBody* InBody)
+		{
+			int32 Index = Bodies.IndexOfByPredicate([&](const TSharedPtr<FBody>& Body){ return (InBody == Body.Get()); });
+			if (Index != INDEX_NONE)
+			{
+				Bodies.RemoveAt(Index);
+			}
+		}
+
 		void PrintBodyAndShellCount();
-		void RemoveEmptyBodies();
 
 		void RemoveEntity(TSharedPtr<FTopologicalEntity> InEntity);
 
@@ -158,7 +165,7 @@ namespace CADKernel
 			return Faces;
 		}
 
-		virtual void GetFaces(TArray<TSharedPtr<FTopologicalFace>>& OutFaces) override;
+		virtual void GetFaces(TArray<FTopologicalFace*>& OutFaces) override;
 
 		virtual int32 FaceCount() const override;
 
@@ -169,8 +176,6 @@ namespace CADKernel
 
 		virtual void SpreadBodyOrientation() override;
 
-		virtual TSharedPtr<FEntityGeom> ApplyMatrix(const FMatrixH& InMatrix) const override;
-
 		// Topo functions
 
 		/**
@@ -186,15 +191,13 @@ namespace CADKernel
 		void FixModelTopology(double JoiningTolerance);
 
 		void MergeInto(TSharedPtr<FBody> Body, TArray<TSharedPtr<FTopologicalEntity>>& InEntities);
-		void Split(TSharedPtr<FBody> Body, TArray<TSharedPtr<FBody>>& OutNewBody);
-		void Join(TArray<TSharedPtr<FBody>> Bodies, double Tolerance);
 
 		/**
 		 * Fore each shell of each body, try to stitch topological gap
 		 */
 		void HealModelTopology(double JoiningTolerance);
 
-
+		void Orient();
 	};
 
 } // namespace CADKernel
diff --git a/Engine/Source/Runtime/Datasmith/CADKernel/Public/CADKernel/Topo/Shell.h b/Engine/Source/Runtime/Datasmith/CADKernel/Public/CADKernel/Topo/Shell.h
--- a/Engine/Source/Runtime/Datasmith/CADKernel/Public/CADKernel/Topo/Shell.h
+++ b/Engine/Source/Runtime/Datasmith/CADKernel/Public/CADKernel/Topo/Shell.h
@@ -6,12 +6,14 @@
 #include "CADKernel/Core/MetadataDictionary.h"
 #include "CADKernel/Core/Types.h"
 #include "CADKernel/Geo/GeoEnum.h"
-#include "CADKernel/Topo/TopologicalEntity.h"
+#include "CADKernel/Topo/TopologicalShapeEntity.h"
 
 namespace CADKernel
 {
+
 class FTopologicalFace;
 class FBody;
+class FTopologyReport;
 
 struct FFaceSubset;
 
@@ -34,19 +36,18 @@ public:
 	}
 };
 
-class CADKERNEL_API FShell : public FTopologicalEntity, public FMetadataDictionary
+class CADKERNEL_API FShell : public FTopologicalShapeEntity
 {
 	friend FEntity;
 	friend class FBody;
 
 private:
 	TArray<FOrientedFace> TopologicalFaces;
-	FBody* HostedBy;
 
 	FShell() = default;
 
 	FShell(const TArray<FOrientedFace> InTopologicalFaces, bool bIsInnerShell = false)
-		: FTopologicalEntity()
+		: FTopologicalShapeEntity()
 		, TopologicalFaces(InTopologicalFaces)
 	{
 		if (bIsInnerShell)
@@ -64,10 +65,8 @@ public:
 
 	virtual void Serialize(FCADKernelArchive& Ar) override
 	{
-		FTopologicalEntity::Serialize(Ar);
+		FTopologicalShapeEntity::Serialize(Ar);
 		SerializeIdents(Ar, (TArray<TOrientedEntity<FEntity>>&) TopologicalFaces);
-		SerializeIdent(Ar, &HostedBy);
-		FMetadataDictionary::Serialize(Ar);
 	}
 
 	virtual void SpawnIdent(FDatabase& Database) override
@@ -86,24 +85,6 @@ public:
 		ResetMarkersRecursivelyOnEntities((TArray<TOrientedEntity<FEntity>>&) TopologicalFaces);
 	}
 
-	FBody* GetHost()
-	{
-		return HostedBy;
-	}
-
-	void SetHost(FBody* Body)
-	{
-		HostedBy = Body;
-	}
-
-	void CompleteMetadata()
-	{
-		if (HostedBy != nullptr)
-		{
-			CompleteDictionary((const FMetadataDictionary&)*HostedBy);
-		}
-	}
-
 	void Empty(int32 NewSize = 0);
 
 	void Add(TSharedRef<FTopologicalFace> InTopologicalFace, EOrientation InOrientation);
@@ -118,8 +99,6 @@ public:
 		return EEntity::Shell;
 	}
 
-	TSharedPtr<FEntityGeom> ApplyMatrix(const FMatrixH& InMatrix) const;
-
 	virtual int32 FaceCount() const override
 	{
 		return TopologicalFaces.Num();
@@ -141,13 +120,18 @@ public:
 		return TopologicalFaces;
 	}
 
-	virtual void GetFaces(TArray<TSharedPtr<FTopologicalFace>>& OutFaces) override;
+	virtual void GetFaces(TArray<FTopologicalFace*>& OutFaces) override;
 
 	virtual void Merge(TSharedPtr<FShell>& Shell);
 
 	virtual void SpreadBodyOrientation() override;
 
-	void CheckTopology(TArray<FFaceSubset>& SubShells);
+	/**
+	 * Update each FOrientedFace::Direction according to FOrientedFace::Entity->IsBackOriented() flag
+	 */
+	virtual void UpdateShellOrientation();
+
+	void CheckTopology(TArray<FFaceSubset>& Subshells);
 	virtual void FillTopologyReport(FTopologyReport& Report) const override;
 
 	/**
@@ -174,6 +158,16 @@ public:
 	{
 		States &= ~EHaveStates::IsInner;
 	}
+
+	/**
+	 * Orient each connected sub-shell: each connected face will have the same orientation 
+	 * Each connected subset of faces will be oriented:
+	 * - towards the outside if it's a closed subset (without border edge) 
+	 * - according to its main/average orientation if it's an open shell 
+	 * @return SwapFaceCount for report purpose
+	 */
+	int32 Orient();
 };
+
 }
 
diff --git a/Engine/Source/Runtime/Datasmith/CADKernel/Public/CADKernel/Topo/TopologicalEdge.h b/Engine/Source/Runtime/Datasmith/CADKernel/Public/CADKernel/Topo/TopologicalEdge.h
--- a/Engine/Source/Runtime/Datasmith/CADKernel/Public/CADKernel/Topo/TopologicalEdge.h
+++ b/Engine/Source/Runtime/Datasmith/CADKernel/Public/CADKernel/Topo/TopologicalEdge.h
@@ -174,8 +174,6 @@ public:
 		return EEntity::TopologicalEdge;
 	}
 
-	virtual TSharedPtr<FEntityGeom> ApplyMatrix(const FMatrixH& InMatrix) const override;
-
 	// ======   Topological Function   ======
 
 	void LinkVertex();
@@ -209,12 +207,12 @@ public:
 			return nullptr;
 		}
 
-		if (TopologicalLink->GetTwinsEntities().Num() < 2)
+		if (TopologicalLink->GetTwinEntities().Num() < 2)
 		{
 			return nullptr;
 		}
 
-		FTopologicalEdge* FirstTwinEdge = (TopologicalLink->GetTwinsEntities()[0] == this) ? TopologicalLink->GetTwinsEntities()[1] : TopologicalLink->GetTwinsEntities()[0];
+		FTopologicalEdge* FirstTwinEdge = (TopologicalLink->GetTwinEntities()[0] == this) ? TopologicalLink->GetTwinEntities()[1] : TopologicalLink->GetTwinEntities()[0];
 		return FirstTwinEdge;
 	}
 
@@ -405,7 +403,7 @@ public:
 		return Mesh.ToSharedRef();
 	}
 
-	TSharedRef<FEdgeMesh> GetOrCreateMesh(const TSharedRef<FModelMesh>& MeshModel);
+	TSharedRef<FEdgeMesh> GetOrCreateMesh(FModelMesh& MeshModel);
 
 	/**
 	 * Generate a sampling of the curve.
@@ -694,7 +692,7 @@ public:
 	 */
 	bool IsBorder()
 	{
-		return GetTwinsEntityCount() == 1;
+		return GetTwinEntityCount() == 1;
 	}
 
 	/**
@@ -702,7 +700,7 @@ public:
 	 */
 	bool IsSurfacic()
 	{
-		return GetTwinsEntityCount() == 2;
+		return GetTwinEntityCount() == 2;
 	}
 
 	/**
diff --git a/Engine/Source/Runtime/Datasmith/CADKernel/Public/CADKernel/Topo/TopologicalEntity.h b/Engine/Source/Runtime/Datasmith/CADKernel/Public/CADKernel/Topo/TopologicalEntity.h
--- a/Engine/Source/Runtime/Datasmith/CADKernel/Public/CADKernel/Topo/TopologicalEntity.h
+++ b/Engine/Source/Runtime/Datasmith/CADKernel/Public/CADKernel/Topo/TopologicalEntity.h
@@ -1,7 +1,7 @@
 // Copyright Epic Games, Inc. All Rights Reserved.
 #pragma once
 
-#include "CADKernel/Core/EntityGeom.h"
+#include "CADKernel/Core/Entity.h"
 
 namespace CADKernel
 {
@@ -9,28 +9,24 @@ namespace CADKernel
 	class FTopologicalFace;
 	class FTopologyReport;
 
-	class CADKERNEL_API FTopologicalEntity : public FEntityGeom
+	class CADKERNEL_API FTopologicalEntity : public FEntity
 	{
+		friend class FCoreTechBridge;
+
+	protected:
+		FIdent CtKioId = 0;
+
 	public:
 
-		virtual int32 FaceCount() const
+		FIdent GetKioId() const
 		{
-			return 0;
+			return CtKioId;
 		}
 
-		virtual void GetFaces(TArray<TSharedPtr<FTopologicalFace>>& OutFaces) {}
-
-		/**
-		 * Each face of model is set by its orientation. This allow to make oriented mesh and to keep the face orientation in topological function.
-		 * Marker2 of spread face is set. It must be reset after the process
-		 */
-		virtual void SpreadBodyOrientation() {}
-
-		virtual void FillTopologyReport(FTopologyReport& Report) const {}
-
 		virtual void Serialize(FCADKernelArchive& Ar) override
 		{
-			FEntityGeom::Serialize(Ar);
+			FEntity::Serialize(Ar);
+			Ar << CtKioId;
 		}
 		
 #ifdef CADKERNEL_DEV
