Update CADKernel : add oriention, fix one infinite loop

diff --git a/Engine/Source/Runtime/Datasmith/CADKernel/Public/CADKernel/Topo/TopologicalFace.h b/Engine/Source/Runtime/Datasmith/CADKernel/Public/CADKernel/Topo/TopologicalFace.h
--- a/Engine/Source/Runtime/Datasmith/CADKernel/Public/CADKernel/Topo/TopologicalFace.h
+++ b/Engine/Source/Runtime/Datasmith/CADKernel/Public/CADKernel/Topo/TopologicalFace.h
@@ -1,13 +1,13 @@
 // Copyright Epic Games, Inc. All Rights Reserved.
 #pragma once
 
-#include "CADKernel/Core/MetadataDictionary.h"
 #include "CADKernel/Geo/GeoPoint.h"
 #include "CADKernel/Geo/Surfaces/Surface.h"
-#include "CADKernel/Math/Aabb.h"
+#include "CADKernel/Geo/Sampling/PolylineTools.h"
 #include "CADKernel/Math/Curvature.h"
 #include "CADKernel/Topo/TopologicalEntity.h"
 #include "CADKernel/Topo/TopologicalLoop.h"
+#include "CADKernel/Topo/TopologicalShapeEntity.h"
 
 namespace CADKernel
 {
@@ -27,6 +27,8 @@ enum class EQuadType : uint8
 	Other
 };
 
+struct FBBoxWithNormal;
+
 class FGrid;
 class FFaceMesh;
 class FThinZone;
@@ -34,16 +36,14 @@ class FThinZoneFinder;
 class FBezierSurface;
 class FSegmentCurve;
 
-class CADKERNEL_API FTopologicalFace : public FTopologicalEntity, public FMetadataDictionary
+class CADKERNEL_API FTopologicalFace : public FTopologicalShapeEntity
 {
 	friend class FEntity;
-	//friend class FShell;
 
 protected:
 
 	TSharedPtr<FSurface> CarrierSurface;
 	TArray<TSharedPtr<FTopologicalLoop>> Loops;
-	FShell* HostedBy;
 
 	mutable TCache<FSurfacicBoundary> Boundary;
 
@@ -74,7 +74,7 @@ protected:
 	 * This constructor has to be completed with one of the three "AddBoundaries" methods to be finalized.
 	 */
 	FTopologicalFace(const TSharedPtr<FSurface>& InCarrierSurface)
-		: FTopologicalEntity()
+		: FTopologicalShapeEntity()
 		, CarrierSurface(InCarrierSurface)
 		, Mesh(TSharedPtr<FFaceMesh>())
 	{
@@ -92,11 +92,9 @@ public:
 
 	virtual void Serialize(FCADKernelArchive& Ar) override
 	{
-		FTopologicalEntity::Serialize(Ar);
+		FTopologicalShapeEntity::Serialize(Ar);
 		SerializeIdent(Ar, CarrierSurface);
 		SerializeIdents(Ar, (TArray<TSharedPtr<FEntity>>&) Loops);
-		SerializeIdent(Ar, &HostedBy);
-		FMetadataDictionary::Serialize(Ar);
 	}
 
 	virtual void SpawnIdent(FDatabase& Database) override;
@@ -117,8 +115,6 @@ public:
 		return EEntity::TopologicalFace;
 	}
 
-	virtual TSharedPtr<FEntityGeom> ApplyMatrix(const FMatrixH& InMatrix) const override;
-
 	const FSurfacicTolerance& GetIsoTolerances() const
 	{
 		return CarrierSurface->GetIsoTolerances();
@@ -143,37 +139,19 @@ public:
 		return 1;
 	}
 
-	virtual void GetFaces(TArray<TSharedPtr<FTopologicalFace>>& OutFaces) override
+	virtual void GetFaces(TArray<FTopologicalFace*>& OutFaces) override
 	{
 		if (!HasMarker1())
 		{
-			OutFaces.Emplace(StaticCastSharedRef<FTopologicalFace>(AsShared()));
+			OutFaces.Emplace(this);
 			SetMarker1();
 		}
 	}
 
-	void ResetHost()
-	{
-		HostedBy = nullptr;
-	}
-
-	FShell* GetHost()
+	virtual void SpreadBodyOrientation() override
 	{
-		return HostedBy;
 	}
 
-	void SetHost(FShell* Shell)
-	{
-		HostedBy = Shell;
-	}
-
-	void CompleteMetadata()
-	{
-		if (HostedBy != nullptr)
-		{
-			CompleteDictionary((const FMetadataDictionary&)*HostedBy);
-		}
-	}
 
 	// ======   Loop Functions   ======
 
@@ -215,7 +193,6 @@ public:
 	/**
 	 * Get a sampling of each loop of the face
 	 * @param OutLoopSamplings an array of 2d points
-	 * @param OutAABBs an array of 2d axis aligned bounding box of each boundary
 	 */
 	const void Get2DLoopSampling(TArray<TArray<FPoint2D>>& OutLoopSamplings) const;
 
@@ -281,6 +258,12 @@ public:
 	 */
 	void Presample();
 
+	/**
+	 * Update the bounding box with an approximation of the surface based of n iso curves 
+	 * @param ApproximationFactor: the factor apply to the geometric tolerance to defined the SAG error of the ISO
+	 */
+	void UpdateBBox(int32 IsoCount, const double ApproximationFactor, FBBoxWithNormal& BBox);
+
 	// ======   Topo Functions   ======
 
 	/**
@@ -297,7 +280,7 @@ public:
 
 	// ======   Meshing Function   ======
 
-	TSharedRef<FFaceMesh> GetOrCreateMesh(const TSharedRef<FModelMesh>& ShellMesh);
+	TSharedRef<FFaceMesh> GetOrCreateMesh(FModelMesh& ModelMesh);
 
 	const bool HasTesselation() const
 	{
@@ -398,6 +381,18 @@ public:
 		return ((States & EHaveStates::IsBackOriented) == EHaveStates::IsBackOriented);
 	}
 
+	void SwapOrientation() const
+	{
+		if (IsBackOriented())
+		{
+			ResetBackOriented();
+		}
+		else
+		{
+			SetBackOriented();
+		}
+	}
+
 	void SetBackOriented() const
 	{
 		States |= EHaveStates::IsBackOriented;
@@ -549,18 +544,159 @@ struct FFaceSubset
 	TArray<FTopologicalFace*> Faces;
 	int32 BorderEdgeCount = 0;
 	int32 NonManifoldEdgeCount = 0;
-	FShell* MainShell;
-	FBody* MainBody;
+	FTopologicalShapeEntity* MainShell;
+	FTopologicalShapeEntity* MainBody;
 	FString MainName;
 	uint32 MainColor;
 	uint32 MainMaterial;
 
-	void SetMainShell(TMap<FShell*, int32>& ShellToFaceCount);
-	void SetMainBody(TMap<FBody*, int32>& BodyToFaceCount);
+	void SetMainShell(TMap<FTopologicalShapeEntity*, int32>& ShellToFaceCount);
+	void SetMainBody(TMap<FTopologicalShapeEntity*, int32>& BodyToFaceCount);
 	void SetMainName(TMap<FString, int32>& NameToFaceCount);
 	void SetMainColor(TMap<uint32, int32>& ColorToFaceCount);
 	void SetMainMaterial(TMap<uint32, int32>& MaterialToFaceCount);
 };
 
+struct FBBoxWithNormal
+{
+	FPoint Max;
+	FPoint Min;
+	FPoint MaxPoints[3];
+	FPoint MinPoints[3];
+	FPoint2D MaxCoordinates[3];
+	FPoint2D MinCoordinates[3];
+	FVector MaxPointNormals[3];
+	FVector MinPointNormals[3];
+	bool MaxNormalNeedUpdate[3];
+	bool MinNormalNeedUpdate[3];
+
+	FBBoxWithNormal()
+		: Max(-HUGE_VALUE, -HUGE_VALUE, -HUGE_VALUE)
+		, Min(HUGE_VALUE, HUGE_VALUE, HUGE_VALUE)
+	{
+	}
+
+	/**
+	 * BBox Length is the sum of the length of the 3 sides
+	 */
+	double Length()
+	{
+		double XLength = Max.X - Min.X;
+		double YLength = Max.Y - Min.Y;
+		double ZLength = Max.Z - Min.Z;
+		return XLength + YLength + ZLength;
+	}
+
+	/**
+	 * Return false if the process failed
+	 */
+	bool CheckOrientation(bool bHasWrongOrientation)
+	{
+		int32 GoodOrientation = 0;
+		int32 WrongOrientation = 0;
+
+		const FVector BBoxNormals[] = { FVector(1., 0., 0.) , FVector(0., 1., 0.) , FVector(0., 0., 1.) };
+
+		for (int32 Index = 0; Index < 3; ++Index)
+		{
+			double DotProduct = MaxPointNormals[Index] | BBoxNormals[Index];
+			if (DotProduct > KINDA_SMALL_NUMBER)
+			{
+				GoodOrientation++;
+			}
+			else if(DotProduct < KINDA_SMALL_NUMBER)
+			{
+				WrongOrientation++;
+			}
+
+			DotProduct = MinPointNormals[Index] | BBoxNormals[Index];
+			if (DotProduct < KINDA_SMALL_NUMBER)
+			{
+				GoodOrientation++;
+			}
+			else if (DotProduct > KINDA_SMALL_NUMBER)
+			{
+				WrongOrientation++;
+			}
+		}
+
+		bHasWrongOrientation = (GoodOrientation < WrongOrientation);
+		return ((GoodOrientation >= 3) || (WrongOrientation >= 3)) && (GoodOrientation != WrongOrientation);
+	}
+
+	void Update(const FPolylineBBox& PolylineBBox, const EIso IsoType, double IsoCoordinate)
+	{
+		for (int32 Index = 0; Index < 3; ++Index)
+		{
+			if (PolylineBBox.Max[Index] > Max[Index])
+			{
+				Max[Index] = PolylineBBox.Max[Index];
+				MaxPoints[Index] = PolylineBBox.MaxPoints[Index];
+				MaxCoordinates[Index] = (IsoType == IsoU) ? FPoint2D(IsoCoordinate, PolylineBBox.CoordinateOfMaxPoint[Index]) : FPoint2D(PolylineBBox.CoordinateOfMaxPoint[Index], IsoCoordinate);
+				MaxNormalNeedUpdate[Index] = true;
+			}
+
+			if (PolylineBBox.Min[Index] < Min[Index])
+			{
+				Min[Index] = PolylineBBox.Min[Index];
+				MinPoints[Index] = PolylineBBox.MinPoints[Index];
+				MinCoordinates[Index] = (IsoType == IsoU) ? FPoint2D(IsoCoordinate, PolylineBBox.CoordinateOfMinPoint[Index]) : FPoint2D(PolylineBBox.CoordinateOfMinPoint[Index], IsoCoordinate);
+				MinNormalNeedUpdate[Index] = true;
+			}
+		}
+	}
+
+	void Update(const FPoint& Point, const FPoint2D InPoint2D)
+	{
+		for (int32 Index = 0; Index < 3; ++Index)
+		{
+			if (Point[Index] > Max[Index])
+			{
+				Max[Index] = Point[Index];
+				MaxPoints[Index] = Point;
+				MaxCoordinates[Index] = InPoint2D;
+				MaxNormalNeedUpdate[Index] = true;
+			}
+
+			if (Point[Index] < Min[Index])
+			{
+				Min[Index] = Point[Index];
+				MinPoints[Index] = Point;
+				MinCoordinates[Index] = InPoint2D;
+				MinNormalNeedUpdate[Index] = true;
+			}
+		}
+	}
+
+	void UpdateNormal(const FTopologicalFace& Face)
+	{
+		const FSurface& Surface = *Face.GetCarrierSurface();
+		bool bSwapOrientation = Face.IsBackOriented();
+
+		for (int32 Index = 0; Index < 3; ++Index)
+		{
+			if (MaxNormalNeedUpdate[Index])
+			{
+				MaxPointNormals[Index] = Surface.EvaluateNormal(MaxCoordinates[Index]);
+				if (bSwapOrientation)
+				{
+					MaxPointNormals[Index] *= -1.;
+				}
+				MaxNormalNeedUpdate[Index] = false;
+			}
+
+			if (MinNormalNeedUpdate[Index])
+			{
+				MinPointNormals[Index] = Surface.EvaluateNormal(MinCoordinates[Index]);
+				if (bSwapOrientation)
+				{
+					MinPointNormals[Index] *= -1.;
+				}
+				MinNormalNeedUpdate[Index] = false;
+			}
+		}
+	}
+};
+
 
 } // namespace CADKernel
diff --git a/Engine/Source/Runtime/Datasmith/CADKernel/Public/CADKernel/Topo/TopologicalLink.h b/Engine/Source/Runtime/Datasmith/CADKernel/Public/CADKernel/Topo/TopologicalLink.h
--- a/Engine/Source/Runtime/Datasmith/CADKernel/Public/CADKernel/Topo/TopologicalLink.h
+++ b/Engine/Source/Runtime/Datasmith/CADKernel/Public/CADKernel/Topo/TopologicalLink.h
@@ -7,243 +7,245 @@
 
 namespace CADKernel
 {
-	class FTopologicalEdge;
-	class FTopologicalVertex;
 
-	template<typename EntityType>
-	class CADKERNEL_API TTopologicalLink : public FEntity
-	{
-		friend FEntity;
+class FTopologicalEdge;
+class FTopologicalVertex;
 
-	protected:
-		friend EntityType;
-		EntityType* ActiveEntity;
+template<typename EntityType>
+class CADKERNEL_API TTopologicalLink : public FTopologicalEntity
+{
+	friend FEntity;
 
-		TArray<EntityType*> TwinsEntities;
+protected:
+	friend EntityType;
+	EntityType* ActiveEntity;
 
-		TTopologicalLink()
-			: ActiveEntity(nullptr)
-		{
-		}
+	TArray<EntityType*> TwinEntities;
 
-		TTopologicalLink(EntityType& Entity)
-			: ActiveEntity(&Entity)
-		{
-			TwinsEntities.Add(&Entity);
-		}
+	TTopologicalLink()
+		: ActiveEntity(nullptr)
+	{
+	}
 
-	public:
+	TTopologicalLink(EntityType& Entity)
+		: ActiveEntity(&Entity)
+	{
+		TwinEntities.Add(&Entity);
+	}
 
-		virtual void Serialize(FCADKernelArchive& Ar) override
-		{
-#ifdef CADKERNEL_DEV
-			if (Ar.IsSaving())
-			{
-				ensureCADKernel(ActiveEntity != nullptr);
-				ensureCADKernel(!ActiveEntity->IsDeleted());
-			}
-#endif
-			FEntity::Serialize(Ar);
-			SerializeIdent(Ar, &ActiveEntity, false);
-			SerializeIdents(Ar, TwinsEntities, false);
-		}
+public:
 
-		virtual void Delete()
+	virtual void Serialize(FCADKernelArchive& Ar) override
+	{
+#ifdef CADKERNEL_DEV
+		if (Ar.IsSaving())
 		{
-			TwinsEntities.Empty();
-			ActiveEntity = nullptr;
-			SetDeleted();
+			ensureCADKernel(ActiveEntity != nullptr);
+			ensureCADKernel(!ActiveEntity->IsDeleted());
 		}
+#endif
+		FEntity::Serialize(Ar);
+		SerializeIdent(Ar, &ActiveEntity, false);
+		SerializeIdents(Ar, TwinEntities, false);
+	}
 
-		const EntityType* GetActiveEntity() const
-		{
-			ensureCADKernel(ActiveEntity);
-			return ActiveEntity;
-		}
+	virtual void Delete()
+	{
+		TwinEntities.Empty();
+		ActiveEntity = nullptr;
+		SetDeleted();
+	}
 
-		EntityType* GetActiveEntity()
-		{
-			ensureCADKernel(ActiveEntity);
-			return ActiveEntity;
-		}
+	const EntityType* GetActiveEntity() const
+	{
+		ensureCADKernel(ActiveEntity);
+		return ActiveEntity;
+	}
 
-		int32 GetTwinsEntitieNum() const
-		{
-			return (int32)TwinsEntities.Num();
-		}
+	EntityType* GetActiveEntity()
+	{
+		ensureCADKernel(ActiveEntity);
+		return ActiveEntity;
+	}
 
-		const TArray<EntityType*>& GetTwinsEntities() const
-		{
-			return TwinsEntities;
-		}
+	int32 GetTwinEntityNum() const
+	{
+		return (int32)TwinEntities.Num();
+	}
 
-		void ActivateEntity(const EntityType& NewActiveEntity)
-		{
-			TFunction<bool()> CheckEntityIsATwin = [&]() {
-				for (EntityType* Entity : TwinsEntities)
+	const TArray<EntityType*>& GetTwinEntities() const
+	{
+		return TwinEntities;
+	}
+
+	void ActivateEntity(const EntityType& NewActiveEntity)
+	{
+		TFunction<bool()> CheckEntityIsATwin = [&]() {
+			for (EntityType* Entity : TwinEntities)
+			{
+				if (Entity == &NewActiveEntity)
 				{
-					if (Entity == &NewActiveEntity)
-					{
-						return true;
-					}
+					return true;
 				}
-				FMessage::Error(TEXT("FTopologicalLink::ActivateEntity, the topological entity is not found in the twins entities"));
-				return false;
-			};
+			}
+			FMessage::Error(TEXT("FTopologicalLink::ActivateEntity, the topological entity is not found in the twins entities"));
+			return false;
+		};
 
-			ensureCADKernel(CheckEntityIsATwin());
-			ActiveEntity = &NewActiveEntity;
-		}
+		ensureCADKernel(CheckEntityIsATwin());
+		ActiveEntity = &NewActiveEntity;
+	}
 
-		void RemoveEntity(TSharedPtr<EntityType>& Entity)
-		{
-			EntityType* EntityPtr = *Entity;
-			RemoveEntity(*EntityPtr);
-		}
+	void RemoveEntity(TSharedPtr<EntityType>& Entity)
+	{
+		EntityType* EntityPtr = *Entity;
+		RemoveEntity(*EntityPtr);
+	}
 
-		void RemoveEntity(EntityType& Entity)
+	void RemoveEntity(EntityType& Entity)
+	{
+		TwinEntities.Remove(&Entity);
+		if (&Entity == ActiveEntity && TwinEntities.Num() > 0)
 		{
-			TwinsEntities.Remove(&Entity);
-			if (&Entity == ActiveEntity && TwinsEntities.Num() > 0)
-			{
-				ActiveEntity = TwinsEntities.HeapTop();
-			}
-
-			if (TwinsEntities.Num() == 0)
-			{
-				ActiveEntity = nullptr;
-				Delete();
-			}
+			ActiveEntity = TwinEntities.HeapTop();
 		}
 
-		void UnlinkTwinEntities()
+		if (TwinEntities.Num() == 0)
 		{
-			for (EntityType* Entity : TwinsEntities)
-			{
-				Entity->ResetTopologicalLink();
-			}
-			TwinsEntities.Empty();
+			ActiveEntity = nullptr;
+			Delete();
 		}
+	}
+
+	//void UnlinkTwinEntities()
+	//{
+	//	for (EntityType* Entity : TwinsEntities)
+	//	{
+	//		Entity->ResetTopologicalLink();
+	//	}
+	//	TwinsEntities.Empty();
+	//}
 
 #ifdef CADKERNEL_DEV
-		virtual FInfoEntity& GetInfo(FInfoEntity& Info) const override;
+	virtual FInfoEntity& GetInfo(FInfoEntity& Info) const override;
 #endif
 
-		virtual EEntity GetEntityType() const override
-		{
-			return EEntity::EdgeLink;
-		}
+	virtual EEntity GetEntityType() const override
+	{
+		return EEntity::EdgeLink;
+	}
 
-		void AddEntity(EntityType* Entity)
-		{
-			TwinsEntities.Add(Entity);
-		}
+	void AddEntity(EntityType* Entity)
+	{
+		TwinEntities.Add(Entity);
+	}
 
-		void AddEntity(EntityType& Entity)
-		{
-			TwinsEntities.Add(&Entity);
-		}
+	void AddEntity(EntityType& Entity)
+	{
+		TwinEntities.Add(&Entity);
+	}
 
-		template <typename LinkableType>
-		void AddEntity(const LinkableType* Entity)
-		{
-			TwinsEntities.Add((EntityType*) Entity);
-		}
+	template <typename LinkableType>
+	void AddEntity(const LinkableType* Entity)
+	{
+		TwinEntities.Add((EntityType*)Entity);
+	}
 
-		template <typename ArrayType>
-		void AddEntities(const ArrayType& Entities)
-		{
-			TwinsEntities.Insert(Entities, TwinsEntities.Num());
-		}
+	template <typename ArrayType>
+	void AddEntities(const ArrayType& Entities)
+	{
+		TwinEntities.Insert(Entities, TwinEntities.Num());
+	}
 
-		/**
-		 * @return true if the Twin entity count link is modified
-		 */
-		virtual bool CleanLink()
+	/**
+	 * @return true if the Twin entity count link is modified
+	 */
+	virtual bool CleanLink()
+	{
+		TArray<EntityType*> NewTwinsEntities;
+		NewTwinsEntities.Reserve(TwinEntities.Num());
+		for (EntityType* Entity : TwinEntities)
 		{
-			TArray<EntityType*> NewTwinsEntities;
-			NewTwinsEntities.Reserve(TwinsEntities.Num());
-			for (EntityType* Entity : TwinsEntities)
+			if (Entity)
 			{
-				if (Entity)
-				{
-					NewTwinsEntities.Add(Entity);
-				}
+				NewTwinsEntities.Add(Entity);
 			}
+		}
 
-			if (NewTwinsEntities.Num() != TwinsEntities.Num())
+		if (NewTwinsEntities.Num() != TwinEntities.Num())
+		{
+			Swap(NewTwinsEntities, TwinEntities);
+			if (TwinEntities.Num())
 			{
-				Swap(NewTwinsEntities, TwinsEntities);
-				if(TwinsEntities.Num())
-				{
-					ActiveEntity = TwinsEntities.HeapTop();
-					return true;
-				}
+				ActiveEntity = TwinEntities.HeapTop();
+				return true;
 			}
-			return false;
 		}
-	};
+		return false;
+	}
+};
 
-	/**
-	 * TTopologicalLink overload dedicated to FVertex to manage the barycenter of twin vertices
-	 */
-	class CADKERNEL_API FVertexLink : public TTopologicalLink<FTopologicalVertex>
-	{
-		friend class FTopologicalVertex;
-	protected:
-		FPoint Barycenter;
+/**
+ * TTopologicalLink overload dedicated to FVertex to manage the barycenter of twin vertices
+ */
+class CADKERNEL_API FVertexLink : public TTopologicalLink<FTopologicalVertex>
+{
+	friend class FTopologicalVertex;
 
-		void SetBarycenter(const FPoint& Point)
-		{
-			Barycenter = Point;
-		}
+protected:
+	FPoint Barycenter;
 
-	public:
-		FVertexLink()
-			: Barycenter(FPoint::ZeroPoint)
-		{
-		}
+	void SetBarycenter(const FPoint& Point)
+	{
+		Barycenter = Point;
+	}
 
-		FVertexLink(FTopologicalVertex& Entity)
-			: TTopologicalLink<FTopologicalVertex>(Entity)
-			, Barycenter(FPoint::ZeroPoint)
-		{
-		}
+public:
+	FVertexLink()
+		: Barycenter(FPoint::ZeroPoint)
+	{
+	}
 
-		virtual void Serialize(FCADKernelArchive& Ar) override
-		{
-			TTopologicalLink<FTopologicalVertex>::Serialize(Ar);
-			Ar << Barycenter;
-		}
+	FVertexLink(FTopologicalVertex& Entity)
+		: TTopologicalLink<FTopologicalVertex>(Entity)
+		, Barycenter(FPoint::ZeroPoint)
+	{
+	}
+
+	virtual void Serialize(FCADKernelArchive& Ar) override
+	{
+		TTopologicalLink<FTopologicalVertex>::Serialize(Ar);
+		Ar << Barycenter;
+	}
 
 #ifdef CADKERNEL_DEV
-		virtual FInfoEntity& GetInfo(FInfoEntity& Info) const override;
+	virtual FInfoEntity& GetInfo(FInfoEntity& Info) const override;
 #endif
 
-		const FPoint& GetBarycenter() const
-		{
-			return Barycenter;
-		}
+	const FPoint& GetBarycenter() const
+	{
+		return Barycenter;
+	}
 
-		virtual bool CleanLink() override
+	virtual bool CleanLink() override
+	{
+		if (TTopologicalLink::CleanLink())
 		{
-			if(TTopologicalLink::CleanLink())
-			{
-				ComputeBarycenter();
-				DefineActiveEntity();
-				return true;
-			}
-			return false;
+			ComputeBarycenter();
+			DefineActiveEntity();
+			return true;
 		}
+		return false;
+	}
 
-		virtual EEntity GetEntityType() const override
-		{
-			return EEntity::VertexLink;
-		}
+	virtual EEntity GetEntityType() const override
+	{
+		return EEntity::VertexLink;
+	}
 
-		void ComputeBarycenter();
-		void DefineActiveEntity();
-	};
+	void ComputeBarycenter();
+	void DefineActiveEntity();
+};
 
 
 } // namespace CADKernel
diff --git a/Engine/Source/Runtime/Datasmith/CADKernel/Public/CADKernel/Topo/TopologicalLoop.h b/Engine/Source/Runtime/Datasmith/CADKernel/Public/CADKernel/Topo/TopologicalLoop.h
--- a/Engine/Source/Runtime/Datasmith/CADKernel/Public/CADKernel/Topo/TopologicalLoop.h
+++ b/Engine/Source/Runtime/Datasmith/CADKernel/Public/CADKernel/Topo/TopologicalLoop.h
@@ -127,6 +127,18 @@ public:
 		return Edges;
 	}
 
+	const FOrientedEdge* GetOrientedEdge(const FTopologicalEdge* InEdge) const 
+	{
+		for (const FOrientedEdge& Edge : Edges)
+		{
+			if (Edge.Entity.Get() == InEdge)
+			{
+				return &Edge;
+			}
+		}
+		return nullptr;
+	}
+
 	/**
 	 * Add active Edge that has not marker 1 in the edge array.
 	 * Marker 1 has to be reset at the end.
diff --git a/Engine/Source/Runtime/Datasmith/CADKernel/Public/CADKernel/Topo/TopologicalVertex.h b/Engine/Source/Runtime/Datasmith/CADKernel/Public/CADKernel/Topo/TopologicalVertex.h
--- a/Engine/Source/Runtime/Datasmith/CADKernel/Public/CADKernel/Topo/TopologicalVertex.h
+++ b/Engine/Source/Runtime/Datasmith/CADKernel/Public/CADKernel/Topo/TopologicalVertex.h
@@ -72,7 +72,7 @@ namespace CADKernel
 		 */
 		inline const FPoint& GetBarycenter() const
 		{
-			if (TopologicalLink.IsValid() && TopologicalLink->GetTwinsEntitieNum() > 1)
+			if (TopologicalLink.IsValid() && TopologicalLink->GetTwinEntityNum() > 1)
 			{
 				return TopologicalLink->GetBarycenter();
 			}
@@ -89,13 +89,13 @@ namespace CADKernel
 
 		void SetCoordinates(const FPoint& NewCoordinates)
 		{
-			if (GetLink()->GetTwinsEntities().Num() > 1)
+			if (GetLink()->GetTwinEntityNum() > 1)
 			{
 				// Update barycenter
-				FPoint BaryCenter = GetLink()->GetBarycenter() * (double)GetLink()->GetTwinsEntities().Num();
+				FPoint BaryCenter = GetLink()->GetBarycenter() * (double)GetLink()->GetTwinEntityNum();
 				BaryCenter -= Coordinates;
 				BaryCenter += NewCoordinates;
-				BaryCenter /= (double)GetLink()->GetTwinsEntities().Num();
+				BaryCenter /= (double)GetLink()->GetTwinEntityNum();
 				GetLink()->SetBarycenter(BaryCenter);
 			}
 			else
@@ -120,9 +120,7 @@ namespace CADKernel
 			return Coordinates.SquareDistance(Point);
 		}
 
-		virtual TSharedPtr<FEntityGeom> ApplyMatrix(const FMatrixH& InMatrix) const override;
-
-		TSharedRef<FVertexMesh> GetOrCreateMesh(TSharedRef<FModelMesh>& MeshModel);
+		TSharedRef<FVertexMesh> GetOrCreateMesh(FModelMesh& MeshModel);
 
 		const TSharedRef<FVertexMesh> GetMesh() const
 		{
@@ -193,7 +191,7 @@ namespace CADKernel
 			else
 			{
 				OutConnectedEdges.Reserve(100);
-				for (const FTopologicalVertex* Vertex : GetLink()->GetTwinsEntities())
+				for (const FTopologicalVertex* Vertex : GetLink()->GetTwinEntities())
 				{
 					OutConnectedEdges.Append(Vertex->ConnectedEdges);
 				}
@@ -209,7 +207,7 @@ namespace CADKernel
 			else
 			{
 				int32 Count = 0;
-				for (const FTopologicalVertex* Vertex : GetLink()->GetTwinsEntities())
+				for (const FTopologicalVertex* Vertex : GetLink()->GetTwinEntities())
 				{
 					Count += Vertex->ConnectedEdges.Num();
 				}
diff --git a/Engine/Source/Runtime/Datasmith/CADKernel/Public/CADKernel/Topo/Topomaker.h b/Engine/Source/Runtime/Datasmith/CADKernel/Public/CADKernel/Topo/Topomaker.h
--- a/Engine/Source/Runtime/Datasmith/CADKernel/Public/CADKernel/Topo/Topomaker.h
+++ b/Engine/Source/Runtime/Datasmith/CADKernel/Public/CADKernel/Topo/Topomaker.h
@@ -3,6 +3,8 @@
 
 #include "CADKernel/Core/Types.h"
 
+#include "CADKernel/Topo/TopomakerReport.h"
+
 namespace CADKernel
 {
 
@@ -13,7 +15,7 @@ class FTopologicalEdge;
 class FTopologicalFace;
 class FTopologicalVertex;
 
-class CADKERNEL_API FJoiner
+class CADKERNEL_API FTopomaker
 {
 	friend class FThinZone;
 	friend class FThinZoneFinder;
@@ -22,20 +24,21 @@ protected:
 
 	FSession& Session;
 
-	TArray<TSharedPtr<FShell>> Shells;
+	TArray<FShell*> Shells;
 	TArray<TSharedPtr<FTopologicalFace>> Faces;
 
-	double JoiningTolerance;
-	double JoiningToleranceSquare;
+	const double SewTolerance;
+	const double SewToleranceSquare;
+
+	FTopomakerReport Report;
 
 public:
 
-	FJoiner(FSession& InSession, const TSharedRef<FModel>& Shells, double Tolerance);
-	FJoiner(FSession& InSession, const TArray<TSharedPtr<FShell>>& Shells, double Tolerance);
-	FJoiner(FSession& InSession, const TArray<TSharedPtr<FTopologicalFace>>& Surfaces, double Tolerance);
+	FTopomaker(FSession& InSession, double Tolerance);
+	FTopomaker(FSession& InSession, const TArray<TSharedPtr<FShell>>& Shells, double Tolerance);
+	FTopomaker(FSession& InSession, const TArray<TSharedPtr<FTopologicalFace>>& Surfaces, double Tolerance);
 
-	void JoinFaces();
-	//void JoinFaces(bool bProcessOnlyBorderEdges, bool bProcessOnlyNonManifoldEdges);
+	void Sew();
 
 	/**
 	 * Check topology of each body
@@ -49,13 +52,20 @@ public:
 	/**
 	 * Split into connected shell and put each shell into the appropriate body
 	 */
-	void SplitIntoConnectedShell();
+	void SplitIntoConnectedShells();
+
+	void OrientShells();
 
 	/**
 	 * Unlink Non-Manifold Vertex i.e. Vertex belong to tow or more shell
 	 */
 	void UnlinkNonManifoldVertex();
 
+	void PrintReport()
+	{
+		Report.Print();
+	}
+
 private:
 
 	/**
@@ -68,6 +78,8 @@ private:
 
 	void RemoveFacesFromShell();
 
+	void RemoveEmptyShells();
+
 	/**
 	 * Return an array of active vertices.
 	 */
diff --git a/Engine/Source/Runtime/Datasmith/CADKernel/Public/CADKernel/Topo/TopomakerReport.h b/Engine/Source/Runtime/Datasmith/CADKernel/Public/CADKernel/Topo/TopomakerReport.h
new file mode 100644
--- /dev/null
+++ b/Engine/Source/Runtime/Datasmith/CADKernel/Public/CADKernel/Topo/TopomakerReport.h
@@ -0,0 +1,44 @@
+// Copyright Epic Games, Inc. All Rights Reserved.
+#pragma once
+
+#include "CADKernel/Core/Chrono.h"
+#include "CADKernel/Core/Types.h"
+#include "CADKernel/UI/Message.h"
+
+namespace CADKernel
+{
+
+class FBody;
+class FTopologicalEdge;
+class FTopologicalFace;
+class FTopologicalEntity;
+class FShell;
+
+class FTopomakerReport
+{
+private:
+	int32 OrientationFixedCount = 0;
+	int32 SwappedFaceCount = 0;
+
+public:
+
+	FDuration SewDuration;
+	FDuration OrientationDuration;
+
+	FTopomakerReport()
+	{
+	}
+
+	void AddSwappedFaceCount(int32 InSwappedFaceCount)
+	{
+		if (InSwappedFaceCount > 0)
+		{
+			OrientationFixedCount++;
+			SwappedFaceCount += InSwappedFaceCount;
+		}
+	}
+
+	void Print();
+};
+
+}
diff --git a/Engine/Source/Runtime/Datasmith/CADKernel/Public/CADKernel/Utils/StringUtil.h b/Engine/Source/Runtime/Datasmith/CADKernel/Public/CADKernel/Utils/StringUtil.h
--- a/Engine/Source/Runtime/Datasmith/CADKernel/Public/CADKernel/Utils/StringUtil.h
+++ b/Engine/Source/Runtime/Datasmith/CADKernel/Public/CADKernel/Utils/StringUtil.h
@@ -14,93 +14,87 @@
 #include <stdio.h>
 #include <stdarg.h>
 
-#ifdef DEV_CADKERNEL_LIB
-#include "CADKernel/Utils/UnrealEnvironment.h"
-#endif
+namespace CADKernel::Utils
+{
 
-namespace CADKernel
+inline std::string ToString(const TCHAR* WideString)
 {
-	namespace Utils
-	{
-		inline std::string ToString(const TCHAR* WideString)
-		{
-			return TCHAR_TO_UTF8(WideString);
-		}
+	return TCHAR_TO_UTF8(WideString);
+}
 
-		template<typename ValueType>
-		inline std::string ToString(const ValueType& v)
-		{
-			std::stringstream stream;
-			stream << v;
-			return stream.str();
-		}
+template<typename ValueType>
+inline std::string ToString(const ValueType& v)
+{
+	std::stringstream stream;
+	stream << v;
+	return stream.str();
+}
 
-		inline std::string ToString(const FPoint& p)
-		{
-			std::string Message = "dm.FPoint(";
-			Message += Utils::ToString(p.X) + "," + Utils::ToString(p.Y) + "," + Utils::ToString(p.Z) + ")";
-			return Message;
-		}
+inline std::string ToString(const FPoint& p)
+{
+	std::string Message = "dm.FPoint(";
+	Message += Utils::ToString(p.X) + "," + Utils::ToString(p.Y) + "," + Utils::ToString(p.Z) + ")";
+	return Message;
+}
 
-		inline FString ToFString(const char* Text)
-		{
-			return UTF8_TO_TCHAR(Text);
-		}
+inline FString ToFString(const char* Text)
+{
+	return UTF8_TO_TCHAR(Text);
+}
 
-		inline FString ToFString(const double& Value)
-		{
-			return FString::Printf(TEXT("%f"), Value);
-		}
+inline FString ToFString(const double& Value)
+{
+	return FString::Printf(TEXT("%f"), Value);
+}
 
-		inline FString ToFString(const int32& Value)
-		{
-			return FString::Printf(TEXT("%d"), Value);
-		}
+inline FString ToFString(const int32& Value)
+{
+	return FString::Printf(TEXT("%d"), Value);
+}
 
-		inline FString ToFString(const uint32& Value)
-		{
-			return FString::Printf(TEXT("%u"), Value);
-		}
+inline FString ToFString(const uint32& Value)
+{
+	return FString::Printf(TEXT("%u"), Value);
+}
 
-		inline FString ToFString(const bool& Value)
-		{
-			return FString::Printf(TEXT("%d"), Value);
-		}
+inline FString ToFString(const bool& Value)
+{
+	return FString::Printf(TEXT("%d"), Value);
+}
 
-		inline FString ToFString(const FPoint& p)
-		{
-			return FString::Printf(TEXT("dm.FPoint(%f, %f, %f)"), p.X, p.Y, p.Z);
-		}
+inline FString ToFString(const FPoint& p)
+{
+	return FString::Printf(TEXT("dm.FPoint(%f, %f, %f)"), p.X, p.Y, p.Z);
+}
 
-		void Explode(const FString& FullString, const TCHAR* Separator, TArray<FString>& StringArray);
+void Explode(const FString& FullString, const TCHAR* Separator, TArray<FString>& StringArray);
 
-		inline FString EscapeBackSlashes(const FString& InputString)
-		{
-			FString ret;
-			ret.Reserve(InputString.Len());
-			for (TCHAR Char : InputString)
-			{
-				if (Char == TEXT('\\')) {
-					ret.AppendChar(TEXT('\\'));
-				}
-				ret.AppendChar(Char);
-			}
-			return ret;
+inline FString EscapeBackSlashes(const FString& InputString)
+{
+	FString ret;
+	ret.Reserve(InputString.Len());
+	for (TCHAR Char : InputString)
+	{
+		if (Char == TEXT('\\')) {
+			ret.AppendChar(TEXT('\\'));
 		}
+		ret.AppendChar(Char);
+	}
+	return ret;
+}
 
-		inline void RemoveUnwantedChar(FString& String, const TCHAR UnwantedChar)
+inline void RemoveUnwantedChar(FString& String, const TCHAR UnwantedChar)
+{
+	FString NewString;
+	NewString.Reserve(String.Len());
+	for (const TCHAR& Char : String)
+	{
+		if (Char != UnwantedChar)
 		{
-			FString NewString;
-			NewString.Reserve(String.Len());
-			for (const TCHAR& Char : String)
-			{
-				if (Char != UnwantedChar)
-				{
-					NewString.AppendChar(Char);
-				}
-			}
-			Move(String, NewString);
+			NewString.AppendChar(Char);
 		}
+	}
+	Move(String, NewString);
+}
 
-	};
 }
