Update CADKernel : fix one infinite loop

diff --git a/Engine/Source/Runtime/Datasmith/CADKernel/Private/CADKernel/Core/Version.h b/Engine/Source/Runtime/Datasmith/CADKernel/Private/CADKernel/Core/Version.h
--- a/Engine/Source/Runtime/Datasmith/CADKernel/Private/CADKernel/Core/Version.h
+++ b/Engine/Source/Runtime/Datasmith/CADKernel/Private/CADKernel/Core/Version.h
@@ -2,8 +2,8 @@
 #pragma once
 
 //** ** Version du Toolkit** **//
-#define TOOLKIT_VERSION 5, 0, 18784723
-#define TOOLKIT_VERSION_ASCII "5.0.18784723"
+#define TOOLKIT_VERSION 5, 0, 18786868
+#define TOOLKIT_VERSION_ASCII "5.0.18786868"
 
 #define RELEASE_DATE	 __DATE__
 
diff --git a/Engine/Source/Runtime/Datasmith/CADKernel/Private/CADKernel/Mesh/Meshers/IsoTriangulator.cpp b/Engine/Source/Runtime/Datasmith/CADKernel/Private/CADKernel/Mesh/Meshers/IsoTriangulator.cpp
--- a/Engine/Source/Runtime/Datasmith/CADKernel/Private/CADKernel/Mesh/Meshers/IsoTriangulator.cpp
+++ b/Engine/Source/Runtime/Datasmith/CADKernel/Private/CADKernel/Mesh/Meshers/IsoTriangulator.cpp
@@ -2719,7 +2719,7 @@ void FIsoTriangulator::ConnectCellSubLoopsByNeighborhood(FCell& Cell)
 	{
 		FPoint2D& BaryCenter = LoopBarycenters.Emplace_GetRef(FPoint2D::ZeroPoint);
 
-		// the external loop is not proceed 
+		// the external loop is not processed 
 		if (Nodes[0]->GetLoopIndex() == 0)
 		{
 			continue;
diff --git a/Engine/Source/Runtime/Datasmith/CADKernel/Private/CADKernel/Mesh/Structure/LoopCleaner.cpp b/Engine/Source/Runtime/Datasmith/CADKernel/Private/CADKernel/Mesh/Structure/LoopCleaner.cpp
--- a/Engine/Source/Runtime/Datasmith/CADKernel/Private/CADKernel/Mesh/Structure/LoopCleaner.cpp
+++ b/Engine/Source/Runtime/Datasmith/CADKernel/Private/CADKernel/Mesh/Structure/LoopCleaner.cpp
@@ -136,7 +136,7 @@ bool FLoopCleaner::UncrossLoops()
 	F3DDebugSession _(bDisplay, TEXT("FixIntersectionBetweenLoops"));
 #endif
 
-	TSet<uint32> IntersectionAlreadyProceed;
+	TSet<uint32> IntersectionAlreadyProcessed;
 
 	LoopSegmentsIntersectionTool.Empty(LoopSegments.Num());
 	int32 Index = 0;
@@ -160,7 +160,7 @@ bool FLoopCleaner::UncrossLoops()
 		if (bDisplay)
 		{
 			LoopSegmentsIntersectionTool.Display(TEXT("IntersectionTool"));
-			F3DDebugSession _(*FString::Printf(TEXT("Segment to proceed %d %d"), Index, Iteration++));
+			F3DDebugSession _(*FString::Printf(TEXT("Segment to be processed %d %d"), Index, Iteration++));
 			Grid.DisplayIsoSegment(EGridSpace::UniformScaled, Segment, 0, EVisuProperty::BlueCurve);
 		}
 #endif
@@ -172,7 +172,7 @@ bool FLoopCleaner::UncrossLoops()
 			{
 				{
 					LoopSegmentsIntersectionTool.Display(TEXT("IntersectionTool"));
-					F3DDebugSession _(*FString::Printf(TEXT("Segment to proceed %d %d"), Index, Iteration++));
+					F3DDebugSession _(*FString::Printf(TEXT("Segment to be processed %d %d"), Index, Iteration++));
 					Grid.DisplayIsoSegment(EGridSpace::UniformScaled, Segment, 0, EVisuProperty::BlueCurve);
 				}
 				{
@@ -184,14 +184,14 @@ bool FLoopCleaner::UncrossLoops()
 #endif
 
 			uint32 IntersectionHash = GetTypeHash(*IntersectingSegment, Segment);
-			bool bNotProceed = !IntersectionAlreadyProceed.Find(IntersectionHash);
-			IntersectionAlreadyProceed.Add(IntersectionHash);
+			bool bNotProcessed = !IntersectionAlreadyProcessed.Find(IntersectionHash);
+			IntersectionAlreadyProcessed.Add(IntersectionHash);
 
 			bool bIsFixed = true;
 			bool bIsSameLoop = ((FLoopNode&)Segment.GetFirstNode()).GetLoopIndex() && ((FLoopNode&)Segment.GetFirstNode()).GetLoopIndex() == ((FLoopNode&)IntersectingSegment->GetFirstNode()).GetLoopIndex();
 
 			// Swapping segments is possible only with outer loop
-			if (bNotProceed)
+			if (bNotProcessed)
 			{
 				if (!TryToRemoveIntersectionOfTwoConsecutiveIntersectingSegments(*IntersectingSegment, Segment))
 				{
@@ -1161,7 +1161,10 @@ bool FLoopCleaner::RemoveIntersectionsOfSubLoop(int32 IntersectionIndex, int32 I
 					Swap(IntersectingSection, OppositeSection);
 				}
 
-				MoveIntersectingSectionBehindOppositeSection(IntersectingSection, OppositeSection);
+				if(!MoveIntersectingSectionBehindOppositeSection(IntersectingSection, OppositeSection))
+				{
+					return false;
+				}
 
 #ifdef DEBUG_REMOVE_INTERSECTIONS		
 				Grid.DisplayGridPolyline(TEXT("RemoveIntersectionsOfSubLoop start"), EGridSpace::UniformScaled, NodesOfLoop, true);
@@ -1238,6 +1241,34 @@ bool FLoopCleaner::RemoveIntersectionsOfSubLoop(int32 IntersectionIndex, int32 I
 	return true;
 }
 
+bool FLoopCleaner::RemovePickOrCoincidenceBetween(FLoopNode* StartNode, FLoopNode* StopNode)
+{
+	if (StartNode->IsDelete() || StopNode->IsDelete())
+	{
+		// should not happen => so the process is canceled
+		return false;
+	}
+
+	for (FLoopNode* Node = StartNode; Node != StopNode; )
+	{
+		FLoopNode* NodeToBeProcessed = Node;
+		Node = GetNext(Node);
+		if (Node == nullptr || Node->IsDelete())
+		{
+			return false;
+		}
+
+		const FPoint2D& PreviousPoint = NodeToBeProcessed->GetPreviousNode().Get2DPoint(EGridSpace::UniformScaled, Grid);
+		const FPoint2D& PointToBeProcessed = NodeToBeProcessed->Get2DPoint(EGridSpace::UniformScaled, Grid);
+		if (CheckAndRemovePick(PreviousPoint, NodeToBeProcessed->Get2DPoint(EGridSpace::UniformScaled, Grid), NodeToBeProcessed->GetNextNode().Get2DPoint(EGridSpace::UniformScaled, Grid), *NodeToBeProcessed) ||
+			CheckAndRemoveCoincidence(PreviousPoint, PointToBeProcessed, *NodeToBeProcessed))
+		{
+			Node = GetPrevious(Node);
+		}
+	}
+	return true;
+}
+
 bool FLoopCleaner::RemoveIntersection(TPair<double, double>& Intersection)
 {
 	int32 SubLoopSegmentCount = (int32)Intersection.Value - (int32)Intersection.Key;
@@ -1326,7 +1357,7 @@ bool FLoopCleaner::RemoveIntersection(TPair<double, double>& Intersection)
 	return true;
 };
 
-void FLoopCleaner::MoveIntersectingSectionBehindOppositeSection(LoopCleanerImpl::FLoopSection IntersectingSection, LoopCleanerImpl::FLoopSection OppositeSection)
+bool FLoopCleaner::MoveIntersectingSectionBehindOppositeSection(LoopCleanerImpl::FLoopSection IntersectingSection, LoopCleanerImpl::FLoopSection OppositeSection)
 {
 	using namespace LoopCleanerImpl;
 
@@ -1370,12 +1401,12 @@ void FLoopCleaner::MoveIntersectingSectionBehindOppositeSection(LoopCleanerImpl:
 	FLoopNode* NextNode = nullptr;
 	for (FLoopNode* Node = GetNext(FirstNodeIntersectingSection); Node != LastNodeIntersectingSection;)
 	{
-		FLoopNode* NodeToProceed = Node;
+		FLoopNode* NodeToBeProcessed = Node;
 		Node = GetNext(Node);
 
 		FPoint2D  CandiatePosition;
 		double MinSquareDistance = HUGE_VALUE;
-		FPoint2D PointToProject = NodeToProceed->Get2DPoint(EGridSpace::UniformScaled, Grid);
+		FPoint2D PointToProject = NodeToBeProcessed->Get2DPoint(EGridSpace::UniformScaled, Grid);
 		for (int32 Index = 1; Index < OppositeSectionCount; ++Index)
 		{
 			FPoint2D ProjectedPoint = ProjectPointOnSegment(PointToProject, *OppositeSectionPoint[Index - 1], *OppositeSectionPoint[Index], Coordinate);
@@ -1386,7 +1417,7 @@ void FLoopCleaner::MoveIntersectingSectionBehindOppositeSection(LoopCleanerImpl:
 				CandiatePosition = ProjectedPoint;
 			}
 		}
-		MoveNode(*NodeToProceed, CandiatePosition);
+		MoveNode(*NodeToBeProcessed, CandiatePosition);
 	}
 
 	for (FLoopNode* Node = GetNext(FirstNodeIntersectingSection); Node != LastNodeIntersectingSection; )
@@ -1398,6 +1429,7 @@ void FLoopCleaner::MoveIntersectingSectionBehindOppositeSection(LoopCleanerImpl:
 			Node = GetPrevious(Node);
 		}
 	}
+	return true;
 }
 
 void FLoopCleaner::MoveNodeBehindSegment(const FIsoSegment& IntersectingSegment, FLoopNode& NodeToMove)
@@ -2020,10 +2052,24 @@ bool FLoopCleaner::RemovePickRecursively(FLoopNode* Node0, FLoopNode* Node1)
 				break;
 			}
 
+			if (CheckAndRemoveCoincidence(*PreviousPoint, *Point1, *Node1))
+			{
+				if (PreviousNode->IsDelete())
+				{
+					break;
+				}
+
+				Point1 = NextPoint;
+ 				Node1 = NextNode;
+				NextNode = &NextNode->GetNextNode();
+				NextPoint = &NextNode->Get2DPoint(EGridSpace::UniformScaled, Grid);
+			}
+
 			Point0 = PreviousPoint;
 			Node0 = PreviousNode;
 			PreviousNode = &PreviousNode->GetPreviousNode();
 			PreviousPoint = &PreviousNode->Get2DPoint(EGridSpace::UniformScaled, Grid);
+
 			continue;
 		}
 
@@ -2034,6 +2080,19 @@ bool FLoopCleaner::RemovePickRecursively(FLoopNode* Node0, FLoopNode* Node1)
 				break;
 			}
 
+			if (CheckAndRemoveCoincidence(*Point0, *NextPoint, *Node0))
+			{
+				if (PreviousNode->IsDelete())
+				{
+					break;
+				}
+
+				Point0 = PreviousPoint;
+				Node0 = PreviousNode;
+				PreviousNode = &PreviousNode->GetPreviousNode();
+				PreviousPoint = &PreviousNode->Get2DPoint(EGridSpace::UniformScaled, Grid);
+			}
+
 			Point1 = NextPoint;
 			Node1 = NextNode;
 			NextNode = &NextNode->GetNextNode();
diff --git a/Engine/Source/Runtime/Datasmith/CADKernel/Private/CADKernel/Mesh/Structure/LoopCleaner.h b/Engine/Source/Runtime/Datasmith/CADKernel/Private/CADKernel/Mesh/Structure/LoopCleaner.h
--- a/Engine/Source/Runtime/Datasmith/CADKernel/Private/CADKernel/Mesh/Structure/LoopCleaner.h
+++ b/Engine/Source/Runtime/Datasmith/CADKernel/Private/CADKernel/Mesh/Structure/LoopCleaner.h
@@ -107,7 +107,7 @@ private:
 	int32 LoopIndex;  
 
 	/**
-	 * Array of the proceed loop's nodes
+	 * Array of the processed loop's nodes
 	 * Warning: some nodes can be deleted -> check with Node.IsDelete
 	 */
 	TArray<FLoopNode*> NodesOfLoop;
@@ -188,6 +188,9 @@ private:
 	bool RemoveLoopPicks();
 	bool RemovePickRecursively(FLoopNode* Node0, FLoopNode* Node1);
 
+	/**
+	 * @return true if the node has been deleted
+	 */
 	bool CheckAndRemovePick(const FPoint2D& PreviousPoint, const FPoint2D& NodeToRemovePoint, const FPoint2D& NextPoint, FLoopNode& NodeToRemove)
 	{
 		double Slop = ComputeUnorientedSlope(NodeToRemovePoint, PreviousPoint, NextPoint);
@@ -210,6 +213,22 @@ private:
 		return false;
 	};
 
+	/**
+	 * @return true if the node has been deleted
+	 */
+	bool CheckAndRemoveCoincidence(const FPoint2D& Point0, const FPoint2D& Point1, FLoopNode& NodeToRemove)
+	{
+		double SquareDistance = Point0.SquareDistance(Point1);
+		if (SquareDistance < SquareGeometricTolerance2)
+		{
+			return RemoveNodeOfLoop(NodeToRemove);
+		}
+		return false;
+	};
+
+	/**
+	 * @return true if the node has been deleted
+	 */
 	bool RemoveNodeOfLoop(FLoopNode& NodeToRemove);
 
 	void FindLoopIntersections();
@@ -225,6 +244,11 @@ private:
 	bool RemoveOuterNode(const TPair<double, double>& Intersection);
 	bool SwapNodes(const TPair<double, double>& Intersection);
 
+	/**
+	 * @return false if the process failed => the surface cannot be meshed
+	 */
+	bool RemovePickOrCoincidenceBetween(FLoopNode* StartNode, FLoopNode* StopNode);
+
 	/**
 	 * Two cases:
 	 *    - the segments of the intersection a closed parallel and in same orientation. in this case, the sub-loop is a long pick. The sub-Loop is delete
@@ -240,7 +264,7 @@ private:
 	void SwapSubLoopOrientation(int32 FirstSegmentIndex, int32 LastSegmentIndex);
 	bool RemoveSubLoop(FLoopNode* StartNode, FLoopNode* EndNode);
 
-	void MoveIntersectingSectionBehindOppositeSection(LoopCleanerImpl::FLoopSection IntersectingSection, LoopCleanerImpl::FLoopSection OppositeSection);
+	bool MoveIntersectingSectionBehindOppositeSection(LoopCleanerImpl::FLoopSection IntersectingSection, LoopCleanerImpl::FLoopSection OppositeSection);
 	void MoveNodeBehindSegment(const FIsoSegment& IntersectingSegment, FLoopNode& NodeToMove);
 	void MoveNode(FLoopNode& NodeToMove, FPoint2D& NewPosition);
 
