Directlink: Fix to support LWC

diff --git a/Engine/Source/Runtime/Datasmith/DirectLink/Private/DirectLinkEndpoint.cpp b/Engine/Source/Runtime/Datasmith/DirectLink/Private/DirectLinkEndpoint.cpp
--- a/Engine/Source/Runtime/Datasmith/DirectLink/Private/DirectLinkEndpoint.cpp
+++ b/Engine/Source/Runtime/Datasmith/DirectLink/Private/DirectLinkEndpoint.cpp
@@ -1260,7 +1260,7 @@ void FInternalThreadState::Run()
 		}
 	}
 
-	UE_CLOG(SharedState.bDebugLog, LogDirectLinkNet, Display, TEXT("Endpoint '%s': Publishing FDirectLinkMsg_EndpointLifecycle Stop"), *SharedState.NiceName);
+	UE_CLOG(SharedState.bDebugLog, LogDirectLinkNet, Verbose, TEXT("Endpoint '%s': Publishing FDirectLinkMsg_EndpointLifecycle Stop"), *SharedState.NiceName);
 	MessageEndpoint->Publish(FMessageEndpoint::MakeMessage<FDirectLinkMsg_EndpointLifecycle>(FDirectLinkMsg_EndpointLifecycle::ELifecycle::Stop));
 	FMessageEndpoint::SafeRelease(MessageEndpoint);
 }
diff --git a/Engine/Source/Runtime/Datasmith/DirectLink/Public/DirectLinkSerialMethods.h b/Engine/Source/Runtime/Datasmith/DirectLink/Public/DirectLinkSerialMethods.h
--- a/Engine/Source/Runtime/Datasmith/DirectLink/Public/DirectLinkSerialMethods.h
+++ b/Engine/Source/Runtime/Datasmith/DirectLink/Public/DirectLinkSerialMethods.h
@@ -53,9 +53,8 @@ enum class ESerialMethod : uint8
 	Float_Default      = _f32        | _default,
 	Float_Array        = _f32        | _array,
 	Double_Default     = _f64        | _default,
-	Vector_Default     = _vector     | _default,
-	Vector4_Default    = _vector4    | _default,
-	Quat_Default       = _quat       | _default,
+	Vector_f32         = _vector     | _default, // FVector serialized as FVector3f
+	Quat_f32           = _quat       | _default, // FQuat serialized as FQuat4f
 	LinearColor_Default= _linearcolor| _default,
 	MD5Hash_Default    = _md5hash    | _default,
 
@@ -65,6 +64,7 @@ enum class ESerialMethod : uint8
 
 static constexpr EStoreType GetStoreType(ESerialMethod Method) { return EStoreType(uint8(Method) & _storeTypeMask); }
 
+// This list maps which serialization method should be used for each Type reflected by directlink
 template<typename T> struct TDefaultSerialMethod;
 template<> struct TDefaultSerialMethod<bool>            { constexpr static ESerialMethod Value = ESerialMethod::Uint8_Default;      };
 template<> struct TDefaultSerialMethod<uint8>           { constexpr static ESerialMethod Value = ESerialMethod::Uint8_Default;      };
@@ -77,20 +77,53 @@ template<> struct TDefaultSerialMethod<double>          { constexpr static ESeri
 template<> struct TDefaultSerialMethod<TArray<int32>>   { constexpr static ESerialMethod Value = ESerialMethod::Int32_Array;        };
 template<> struct TDefaultSerialMethod<TArray<FString>> { constexpr static ESerialMethod Value = ESerialMethod::String_Array;       };
 template<> struct TDefaultSerialMethod<TArray<float>>   { constexpr static ESerialMethod Value = ESerialMethod::Float_Array;        };
-template<> struct TDefaultSerialMethod<FVector>         { constexpr static ESerialMethod Value = ESerialMethod::Vector_Default;     };
-template<> struct TDefaultSerialMethod<FVector4>        { constexpr static ESerialMethod Value = ESerialMethod::Vector4_Default;    };
-template<> struct TDefaultSerialMethod<FQuat>           { constexpr static ESerialMethod Value = ESerialMethod::Quat_Default;       };
 template<> struct TDefaultSerialMethod<FLinearColor>    { constexpr static ESerialMethod Value = ESerialMethod::LinearColor_Default;};
 template<> struct TDefaultSerialMethod<FMD5Hash>        { constexpr static ESerialMethod Value = ESerialMethod::MD5Hash_Default;    };
 
-template<typename T> EStoreType GetStoreTypeForType() { return GetStoreType(TDefaultSerialMethod<T>::Value); }
+// for compatibility with pre-LWC endpoints, we send/receive LWC-aware types as their f32 variants
+template<> struct TDefaultSerialMethod<FVector>         { constexpr static ESerialMethod Value = ESerialMethod::Vector_f32;         };
+template<> struct TDefaultSerialMethod<FQuat>           { constexpr static ESerialMethod Value = ESerialMethod::Quat_f32;           };
+
 
+template<typename T> EStoreType GetStoreTypeForType() { return GetStoreType(TDefaultSerialMethod<T>::Value); }
 template<typename T> bool CanSerializeWithMethod(ESerialMethod Method) { return GetStoreTypeForType<T>() == GetStoreType(Method); }
 
 
+// Serialization implementations
 template<ESerialMethod Code, typename T>
 void Serial(FArchive& Ar, T* Param) { Ar << *Param; }
-template<> inline void Serial<ESerialMethod::Uint32_Packed >(FArchive& Ar, uint32*  Param) { Ar.SerializeIntPacked(*Param); }
+
+template<> inline void Serial<ESerialMethod::Uint32_Packed>(FArchive& Ar, uint32* ValuePtr) { Ar.SerializeIntPacked(*ValuePtr); }
+
+template<> inline void Serial<ESerialMethod::Vector_f32>(FArchive& Ar, FVector* ValuePtr)
+{
+	if (Ar.IsLoading())
+	{
+		FVector3f Tmp;
+		Ar << Tmp;
+		*ValuePtr = Tmp;
+	}
+	else if (Ar.IsSaving())
+	{
+		FVector3f Tmp(*ValuePtr);
+		Ar << Tmp;
+	}
+}
+
+template<> inline void Serial<ESerialMethod::Quat_f32>(FArchive& Ar, FQuat* ValuePtr)
+{
+	if (Ar.IsLoading())
+	{
+		FQuat4f Tmp;
+		Ar << Tmp;
+		*ValuePtr = FQuat(Tmp);
+	}
+	else if (Ar.IsSaving())
+	{
+		FQuat4f Tmp(*ValuePtr);
+		Ar << Tmp;
+	}
+}
 
 
 inline bool SerialAny(FArchive& Ar, void* data, ESerialMethod Method)
@@ -110,9 +143,8 @@ inline bool SerialAny(FArchive& Ar, void* data, ESerialMethod Method)
 		SerialAny_Case(ESerialMethod::Int32_Array        , TArray<int32>  );
 		SerialAny_Case(ESerialMethod::String_Array       , TArray<FString>);
 		SerialAny_Case(ESerialMethod::Float_Array        , TArray<float>  );
-		SerialAny_Case(ESerialMethod::Vector_Default     , FVector        );
-		SerialAny_Case(ESerialMethod::Vector4_Default    , FVector4       );
-		SerialAny_Case(ESerialMethod::Quat_Default       , FQuat          );
+		SerialAny_Case(ESerialMethod::Vector_f32         , FVector        );
+		SerialAny_Case(ESerialMethod::Quat_f32           , FQuat          );
 		SerialAny_Case(ESerialMethod::LinearColor_Default, FLinearColor   );
 		SerialAny_Case(ESerialMethod::MD5Hash_Default    , FMD5Hash       );
 
