Datasmith Direct Link - Fix missing MetaData on Datasmith DirectLink import.

diff --git a/Engine/Source/Runtime/Datasmith/DatasmithCore/Private/DatasmithSceneElementsImpl.cpp b/Engine/Source/Runtime/Datasmith/DatasmithCore/Private/DatasmithSceneElementsImpl.cpp
--- a/Engine/Source/Runtime/Datasmith/DatasmithCore/Private/DatasmithSceneElementsImpl.cpp
+++ b/Engine/Source/Runtime/Datasmith/DatasmithCore/Private/DatasmithSceneElementsImpl.cpp
@@ -1160,7 +1160,7 @@ const TSharedPtr< IDatasmithMetaDataElement >& FDatasmithSceneImpl::GetMetaData(
 
 TSharedPtr< IDatasmithMetaDataElement > FDatasmithSceneImpl::GetMetaData(const TSharedPtr<IDatasmithElement>& Element)
 {
-	if (TSharedPtr< IDatasmithMetaDataElement >* MetaDataElement = ElementToMetaDataMap.Find(Element))
+	if (TSharedPtr< IDatasmithMetaDataElement >* MetaDataElement = GetElementToMetaDataCache().Find(Element))
 	{
 		return *MetaDataElement;
 	}
@@ -1172,7 +1172,7 @@ TSharedPtr< IDatasmithMetaDataElement > FDatasmithSceneImpl::GetMetaData(const T
 
 const TSharedPtr< IDatasmithMetaDataElement >& FDatasmithSceneImpl::GetMetaData(const TSharedPtr<IDatasmithElement>& Element) const
 {
-	if (const TSharedPtr< IDatasmithMetaDataElement >* MetaDataElement = ElementToMetaDataMap.Find(Element))
+	if (const TSharedPtr< IDatasmithMetaDataElement >* MetaDataElement = GetElementToMetaDataCache().Find(Element))
 	{
 		return *MetaDataElement;
 	}
@@ -1186,7 +1186,7 @@ void FDatasmithSceneImpl::RemoveMetaData( const TSharedPtr<IDatasmithMetaDataEle
 {
 	if ( Element )
 	{
-		ElementToMetaDataMap.Remove( Element->GetAssociatedElement() );
+		GetElementToMetaDataCache().Remove( Element->GetAssociatedElement() );
 		MetaData.Remove( Element );
 	}
 }
@@ -1196,6 +1196,19 @@ void FDatasmithSceneImpl::RemoveMetaDataAt(int32 InIndex)
 	RemoveMetaData(GetMetaData(InIndex));
 }
 
+TMap< TSharedPtr< IDatasmithElement >, TSharedPtr< IDatasmithMetaDataElement> >& FDatasmithSceneImpl::GetElementToMetaDataCache() const
+{
+	if (ElementToMetaDataMap.Num() == 0)
+	{
+		for (const TSharedPtr<IDatasmithMetaDataElement>& MetaDataElement : MetaData.View())
+		{
+			ElementToMetaDataMap.Add(MetaDataElement->GetAssociatedElement(), MetaDataElement);
+		}
+	}
+
+	return ElementToMetaDataMap;
+}
+
 TSharedPtr<IDatasmithLevelSequenceElement> FDatasmithSceneImpl::GetLevelSequence(int32 InIndex)
 {
 	return LevelSequences.IsValidIndex(InIndex) ? LevelSequences[InIndex] : nullptr;
diff --git a/Engine/Source/Runtime/Datasmith/DatasmithCore/Private/DatasmithSceneElementsImpl.h b/Engine/Source/Runtime/Datasmith/DatasmithCore/Private/DatasmithSceneElementsImpl.h
--- a/Engine/Source/Runtime/Datasmith/DatasmithCore/Private/DatasmithSceneElementsImpl.h
+++ b/Engine/Source/Runtime/Datasmith/DatasmithCore/Private/DatasmithSceneElementsImpl.h
@@ -1620,7 +1620,7 @@ public:
 	virtual int32 GetLODScreenSizesCount() const override { return LODScreenSizes.Get().Num(); }
 	virtual float GetLODScreenSize(int32 InIndex) const override { return LODScreenSizes.Get().IsValidIndex( InIndex ) ? LODScreenSizes.Get()[InIndex] : 0.f; }
 
-	virtual void AddMetaData(const TSharedPtr< IDatasmithMetaDataElement >& InMetaData) override { MetaData.Add(InMetaData); ElementToMetaDataMap.Add(InMetaData->GetAssociatedElement(), InMetaData); }
+	virtual void AddMetaData(const TSharedPtr< IDatasmithMetaDataElement >& InMetaData) override { MetaData.Add(InMetaData); GetElementToMetaDataCache().Add(InMetaData->GetAssociatedElement(), InMetaData); }
 
 	virtual int32 GetMetaDataCount() const override { return MetaData.Num(); }
 	virtual TSharedPtr< IDatasmithMetaDataElement > GetMetaData(int32 InIndex) override;
@@ -1648,6 +1648,8 @@ public:
 	virtual void AttachActorToSceneRoot(const TSharedPtr< IDatasmithActorElement >& Child, EDatasmithActorAttachmentRule AttachmentRule) override;
 
 private:
+	TMap< TSharedPtr< IDatasmithElement >, TSharedPtr< IDatasmithMetaDataElement> >& GetElementToMetaDataCache() const;
+
 	TDatasmithReferenceArrayProxy<IDatasmithActorElement>            Actors;
 	TDatasmithReferenceArrayProxy<IDatasmithMeshElement>             Meshes;
 	TDatasmithReferenceArrayProxy<IDatasmithBaseMaterialElement>     Materials;
@@ -1673,6 +1675,6 @@ private:
 
 	TReflected<bool> bUseSky;
 
-	// #ue_directlink_reflect todo
-	TMap< TSharedPtr< IDatasmithElement >, TSharedPtr< IDatasmithMetaDataElement> > ElementToMetaDataMap;
+	// Internal cache for faster metadata access per-element, should be accessed via GetMetaDataCache(), do not use directly.
+	mutable TMap< TSharedPtr< IDatasmithElement >, TSharedPtr< IDatasmithMetaDataElement> > ElementToMetaDataMap;
 };
