Prepare Directlink to receive LWC types

diff --git a/Engine/Source/Runtime/Datasmith/DirectLink/Public/DirectLinkSerialMethods.h b/Engine/Source/Runtime/Datasmith/DirectLink/Public/DirectLinkSerialMethods.h
--- a/Engine/Source/Runtime/Datasmith/DirectLink/Public/DirectLinkSerialMethods.h
+++ b/Engine/Source/Runtime/Datasmith/DirectLink/Public/DirectLinkSerialMethods.h
@@ -10,6 +10,13 @@
 #include "Serialization/Archive.h"
 
 
+// LWC initiative changes the support of some geometry-related types from float to doubles.
+// As we want to be able to send data to apps built without LWC support, we send data with
+// the legacy float serialization.
+// When the support of pre-LWC products is no longer required, we can safely toggle this
+// macro off and send LWC types natively (as doubles).
+#define DL_SEND_LWC_TYPES_AS_F32 1
+
 namespace Reflect
 {
 
@@ -54,7 +61,9 @@ enum class ESerialMethod : uint8
 	Float_Array        = _f32        | _array,
 	Double_Default     = _f64        | _default,
 	Vector_f32         = _vector     | _default, // FVector serialized as FVector3f
+	Vector_f64         = _vector     | _alt1,    // FVector serialized as FVector3d (native since 5.0)
 	Quat_f32           = _quat       | _default, // FQuat serialized as FQuat4f
+	Quat_f64           = _quat       | _alt1,    // FQuat serialized as FQuat4d (native since 5.0)
 	LinearColor_Default= _linearcolor| _default,
 	MD5Hash_Default    = _md5hash    | _default,
 
@@ -79,10 +88,14 @@ template<> struct TDefaultSerialMethod<TArray<FString>> { constexpr static ESeri
 template<> struct TDefaultSerialMethod<TArray<float>>   { constexpr static ESerialMethod Value = ESerialMethod::Float_Array;        };
 template<> struct TDefaultSerialMethod<FLinearColor>    { constexpr static ESerialMethod Value = ESerialMethod::LinearColor_Default;};
 template<> struct TDefaultSerialMethod<FMD5Hash>        { constexpr static ESerialMethod Value = ESerialMethod::MD5Hash_Default;    };
-
-// for compatibility with pre-LWC endpoints, we send/receive LWC-aware types as their f32 variants
+#if DL_SEND_LWC_TYPES_AS_F32
+// for compatibility with pre-LWC endpoints, we send LWC-aware types as their f32 variants
 template<> struct TDefaultSerialMethod<FVector>         { constexpr static ESerialMethod Value = ESerialMethod::Vector_f32;         };
 template<> struct TDefaultSerialMethod<FQuat>           { constexpr static ESerialMethod Value = ESerialMethod::Quat_f32;           };
+#else
+template<> struct TDefaultSerialMethod<FVector>         { constexpr static ESerialMethod Value = ESerialMethod::Vector_f64;         };
+template<> struct TDefaultSerialMethod<FQuat>           { constexpr static ESerialMethod Value = ESerialMethod::Quat_f64;           };
+#endif
 
 
 template<typename T> EStoreType GetStoreTypeForType() { return GetStoreType(TDefaultSerialMethod<T>::Value); }
@@ -144,7 +157,9 @@ inline bool SerialAny(FArchive& Ar, void* data, ESerialMethod Method)
 		SerialAny_Case(ESerialMethod::String_Array       , TArray<FString>);
 		SerialAny_Case(ESerialMethod::Float_Array        , TArray<float>  );
 		SerialAny_Case(ESerialMethod::Vector_f32         , FVector        );
+		SerialAny_Case(ESerialMethod::Vector_f64         , FVector        );
 		SerialAny_Case(ESerialMethod::Quat_f32           , FQuat          );
+		SerialAny_Case(ESerialMethod::Quat_f64           , FQuat          );
 		SerialAny_Case(ESerialMethod::LinearColor_Default, FLinearColor   );
 		SerialAny_Case(ESerialMethod::MD5Hash_Default    , FMD5Hash       );
 
