Fix FTopologicalEdge::CreateEdgeByMergingEdges :
 - if BSpline::DuplicateNurbsCurveWithHigherDegree failed, the operation is canceled
 - fix the merge of NURBS (rational + non-rational nurbs)
 - clean FNURBSCurve constructor
 - If DuplicateNurbsCurveWithHigherDegree failed return TSharedPtr<FNURBSCurve>()

diff --git a/Engine/Plugins/Enterprise/DatasmithCADImporter/Source/CADInterfaces/Private/CoreTechBridge.cpp b/Engine/Plugins/Enterprise/DatasmithCADImporter/Source/CADInterfaces/Private/CoreTechBridge.cpp
--- a/Engine/Plugins/Enterprise/DatasmithCADImporter/Source/CADInterfaces/Private/CoreTechBridge.cpp
+++ b/Engine/Plugins/Enterprise/DatasmithCADImporter/Source/CADInterfaces/Private/CoreTechBridge.cpp
@@ -416,14 +416,7 @@ TSharedPtr<FTopologicalEdge> FCoreTechBridge::AddEdge(CT_OBJECT_ID CTCoedgeId, T
 	}
 
 	TSharedPtr<FCurve> Curve;
-	if (PoleDim == 2)
-	{
-		Curve = FEntity::MakeShared<FNURBSCurve>(Degre, Knots, Poles, 2);
-	}
-	else
-	{
-		Curve = FEntity::MakeShared<FNURBSCurve>(Degre, Knots, Poles, Weights, 2);
-	}
+	Curve = FEntity::MakeShared<FNURBSCurve>(Degre, Knots, Poles, Weights, 2);
 
 	if (!Curve.IsValid())
 	{
@@ -853,19 +846,17 @@ TSharedPtr<FCurve> FCoreTechBridge::AddNurbsCurve(CT_OBJECT_ID CTCurveId)
 		Poles[IPole].Set(RawPoles.GetData() + Index);
 	}
 
+	TArray<double> Weights;
 	if (PoleDim > 3)
 	{
-		TArray<double> Weights;
 		Weights.SetNum(PoleNum);
 		Index = 3;
 		for (CT_UINT32 IPole = 0; IPole < PoleNum; IPole++, Index += 4)
 		{
 			Weights[IPole] = RawPoles[Index];
 		}
-		return FEntity::MakeShared<FNURBSCurve>(Degre, Knots, Poles, Weights);
 	}
-
-	return FEntity::MakeShared<FNURBSCurve>(Degre, Knots, Poles);
+	return FEntity::MakeShared<FNURBSCurve>(Degre, Knots, Poles, Weights);
 }
 
 TSharedPtr<FCurve> FCoreTechBridge::AddLineCurve(CT_OBJECT_ID CTCurveId, CT_OBJECT_ID CTSurfaceId)
diff --git a/Engine/Source/Runtime/Datasmith/CADKernel/Private/CADKernel/Geo/Curves/NURBSCurve.cpp b/Engine/Source/Runtime/Datasmith/CADKernel/Private/CADKernel/Geo/Curves/NURBSCurve.cpp
--- a/Engine/Source/Runtime/Datasmith/CADKernel/Private/CADKernel/Geo/Curves/NURBSCurve.cpp
+++ b/Engine/Source/Runtime/Datasmith/CADKernel/Private/CADKernel/Geo/Curves/NURBSCurve.cpp
@@ -13,24 +13,13 @@
 namespace CADKernel
 { 
 
-FNURBSCurve::FNURBSCurve(int32 InDegre, const TArray<double>& InNodalVector, const TArray<FPoint>& InPoles, int8 InDimension)
-	: FCurve(InDimension)
-	, Degree(InDegre)
-	, NodalVector(InNodalVector)
-	, Poles(InPoles)
-	, bIsRational(false)
-{
-	Weights.Init(1.0, InPoles.Num());
-	Finalize();
-}
-
 FNURBSCurve::FNURBSCurve(int32 InDegre, const TArray<double>& InNodalVector, const TArray<FPoint>& InPoles, const TArray<double>& InWeights, int8 InDimension)
 	: FCurve(InDimension)
 	, Degree(InDegre)
 	, NodalVector(InNodalVector)
 	, Weights(InWeights)
 	, Poles(InPoles)
-	, bIsRational(true)
+	, bIsRational(!InWeights.IsEmpty())
 {
 	Finalize();
 }
@@ -46,9 +35,15 @@ FNURBSCurve::FNURBSCurve(FNurbsCurveData& NurbsCurveData)
 	Finalize();
 }
 
-FNURBSCurve::FNURBSCurve(const TSharedRef<FNURBSCurve>& Nurbs)
-	: FNURBSCurve(Nurbs->Degree, Nurbs->NodalVector, Nurbs->Poles, Nurbs->Weights, Nurbs->Dimension)
+FNURBSCurve::FNURBSCurve(const FNURBSCurve& Nurbs)
+	: FCurve(Nurbs.Dimension)
+	, Degree(Nurbs.Degree)
+	, NodalVector(Nurbs.NodalVector)
+	, Weights(Nurbs.Weights)
+	, Poles(Nurbs.Poles)
+	, bIsRational(Nurbs.bIsRational)
 {
+	Finalize();
 }
 
 TSharedPtr<FEntityGeom> FNURBSCurve::ApplyMatrix(const FMatrixH& InMatrix) const
diff --git a/Engine/Source/Runtime/Datasmith/CADKernel/Private/CADKernel/Math/BSpline.cpp b/Engine/Source/Runtime/Datasmith/CADKernel/Private/CADKernel/Math/BSpline.cpp
--- a/Engine/Source/Runtime/Datasmith/CADKernel/Private/CADKernel/Math/BSpline.cpp
+++ b/Engine/Source/Runtime/Datasmith/CADKernel/Private/CADKernel/Math/BSpline.cpp
@@ -813,23 +813,18 @@ namespace CADKernel
 			}
 		}
 
-		TSharedRef<FNURBSCurve> DuplicateNurbsCurveWithHigherDegree(int32 Degre, const FNURBSCurve& InCurve)
+		TSharedPtr<FNURBSCurve> DuplicateNurbsCurveWithHigherDegree(int32 Degre, const FNURBSCurve& InCurve)
 		{
 			TArray<FPoint> NewPoles;
 			TArray<double> NewNodalVector;
 			TArray<double> NewWeights;
 
 			DuplicateNurbsCurveWithHigherDegree(Degre, InCurve.GetPoles(), InCurve.GetNodalVector(), InCurve.GetWeights(), NewPoles, NewNodalVector, NewWeights);
-			
-			TSharedPtr<FCurve> Curve;
-			if (InCurve.IsRational())
-			{
-				return FEntity::MakeShared<FNURBSCurve>(Degre, NewNodalVector, NewPoles, NewWeights, InCurve.GetDimension());
-			}
-			else
+			if (NewPoles.IsEmpty())
 			{
-				return FEntity::MakeShared<FNURBSCurve>(Degre, NewNodalVector, NewPoles, InCurve.GetDimension());
+				return TSharedPtr<FNURBSCurve>();
 			}
+			return FEntity::MakeShared<FNURBSCurve>(Degre, NewNodalVector, NewPoles, NewWeights, InCurve.GetDimension());
 		}
 
 
@@ -1540,7 +1535,11 @@ namespace CADKernel
 			//int32 newN=n+1;
 			int32 newN = degre + 1;
 
-			ensureCADKernel(newPoles.Num() + newN - 1 == newNodalVector.Num());
+			if (newPoles.Num() + newN - 1 != newNodalVector.Num())
+			{
+				newPoles.Empty();
+				return;
+			}
 
 			TArray<bool> polesDone;
 			polesDone.Init(false, newPoles.Num());
diff --git a/Engine/Source/Runtime/Datasmith/CADKernel/Private/CADKernel/Topo/TopologicalEdge.cpp b/Engine/Source/Runtime/Datasmith/CADKernel/Private/CADKernel/Topo/TopologicalEdge.cpp
--- a/Engine/Source/Runtime/Datasmith/CADKernel/Private/CADKernel/Topo/TopologicalEdge.cpp
+++ b/Engine/Source/Runtime/Datasmith/CADKernel/Private/CADKernel/Topo/TopologicalEdge.cpp
@@ -177,6 +177,11 @@ void FTopologicalEdge::Link(FTopologicalEdge& Twin, double SquareJoiningToleranc
 		return;
 	}
 
+	if(IsDeleted() || Twin.IsDeleted())
+	{
+		return;
+	}
+
 	const FPoint& Edge1Vertex1 = GetStartVertex()->GetBarycenter();
 	const FPoint& Edge1Vertex2 = GetEndVertex()->GetBarycenter();
 	const FPoint& Edge2Vertex1 = Twin.GetStartVertex()->GetBarycenter();
@@ -654,10 +659,15 @@ TSharedPtr<FTopologicalEdge> FTopologicalEdge::CreateEdgeByMergingEdges(TArray<F
 		if (NURBS->GetDegree() < NurbsMaxDegree)
 		{
 			NURBS = BSpline::DuplicateNurbsCurveWithHigherDegree(NurbsMaxDegree, *NURBS);
+			if(!NURBS.IsValid())
+			{
+				// cancel
+				return TSharedPtr<FTopologicalEdge>();
+			}
 		}
 		else
 		{
-			NURBS = FEntity::MakeShared<FNURBSCurve>(NURBS.ToSharedRef());
+			NURBS = FEntity::MakeShared<FNURBSCurve>(*NURBS);
 		}
 
 		if (Edges[Index].Direction == EOrientation::Back)
@@ -705,9 +715,9 @@ TSharedPtr<FTopologicalEdge> FTopologicalEdge::CreateEdgeByMergingEdges(TArray<F
 	TArray<FPoint> NewPoles;
 	NurbsMaxDegree++;
 	NewNodalVector.Reserve(PoleCount + NurbsMaxDegree);
-	NewWeights.Reserve(PoleCount + NurbsMaxDegree);
 	NewPoles.Reserve(PoleCount + NurbsMaxDegree);
 
+	bool bIsRational = false;
 	for (const TSharedPtr<FNURBSCurve>& NurbsCurve : NurbsCurves)
 	{
 		if (!NurbsCurve.IsValid())
@@ -715,14 +725,59 @@ TSharedPtr<FTopologicalEdge> FTopologicalEdge::CreateEdgeByMergingEdges(TArray<F
 			continue;
 		}
 
-		if (!NewPoles.IsEmpty())
+		if (NurbsCurve->IsRational())
 		{
-			NewPoles.Pop();
-			NewWeights.Pop();
+			bIsRational = true;
+			break;
 		}
+	}
 
-		NewPoles.Append(NurbsCurve->GetPoles());
-		NewWeights.Append(NurbsCurve->GetWeights());
+	if(bIsRational)
+	{
+		NewWeights.Reserve(PoleCount + NurbsMaxDegree);
+		for (const TSharedPtr<FNURBSCurve>& NurbsCurve : NurbsCurves)
+		{
+			if (!NurbsCurve.IsValid())
+			{
+				continue;
+			}
+
+			if (!NewPoles.IsEmpty())
+			{
+				NewPoles.Pop();
+				NewWeights.Pop();
+			}
+
+			NewPoles.Append(NurbsCurve->GetPoles());
+			if(NurbsCurve->IsRational())
+			{
+				NewWeights.Append(NurbsCurve->GetWeights());
+			}
+			else
+			{
+				for (int32 Index = 0; Index < NurbsCurve->GetPoles().Num(); ++Index)
+				{
+					NewWeights.Add(1.);
+				}
+			}
+		}
+	}
+	else
+	{
+		for (const TSharedPtr<FNURBSCurve>& NurbsCurve : NurbsCurves)
+		{
+			if (!NurbsCurve.IsValid())
+			{
+				continue;
+			}
+
+			if (!NewPoles.IsEmpty())
+			{
+				NewPoles.Pop();
+			}
+
+			NewPoles.Append(NurbsCurve->GetPoles());
+		}
 	}
 
 	for (const TSharedPtr<FNURBSCurve>& NurbsCurve : NurbsCurves)
diff --git a/Engine/Source/Runtime/Datasmith/CADKernel/Public/CADKernel/Geo/Curves/NURBSCurve.h b/Engine/Source/Runtime/Datasmith/CADKernel/Public/CADKernel/Geo/Curves/NURBSCurve.h
--- a/Engine/Source/Runtime/Datasmith/CADKernel/Public/CADKernel/Geo/Curves/NURBSCurve.h
+++ b/Engine/Source/Runtime/Datasmith/CADKernel/Public/CADKernel/Geo/Curves/NURBSCurve.h
@@ -39,10 +39,9 @@ namespace CADKernel
 		TArray<double> HomogeneousPoles;
 		int32 PoleDimension;
 
-		FNURBSCurve(int32 InDegree, const TArray<double>& InNodalVector, const TArray<FPoint>& InPoles, int8 InDimension = 3);
 		FNURBSCurve(int32 InDegree, const TArray<double>& InNodalVector, const TArray<FPoint>& InPoles, const TArray<double>& InWeights, int8 InDimension = 3);
 		FNURBSCurve(FNurbsCurveData& NurbsCurveData);
-		FNURBSCurve(const TSharedRef<FNURBSCurve>& Nurbs);
+		FNURBSCurve(const FNURBSCurve& Nurbs);
 
 		FNURBSCurve() = default;
 
diff --git a/Engine/Source/Runtime/Datasmith/CADKernel/Public/CADKernel/Geo/Sampling/SurfacicPolyline.h b/Engine/Source/Runtime/Datasmith/CADKernel/Public/CADKernel/Geo/Sampling/SurfacicPolyline.h
--- a/Engine/Source/Runtime/Datasmith/CADKernel/Public/CADKernel/Geo/Sampling/SurfacicPolyline.h
+++ b/Engine/Source/Runtime/Datasmith/CADKernel/Public/CADKernel/Geo/Sampling/SurfacicPolyline.h
@@ -15,7 +15,7 @@ namespace CADKernel
 
 class FInfoEntity;
 
-class FSurfacicPolyline
+class CADKERNEL_API FSurfacicPolyline
 {
 
 public:
diff --git a/Engine/Source/Runtime/Datasmith/CADKernel/Public/CADKernel/Math/BSpline.h b/Engine/Source/Runtime/Datasmith/CADKernel/Public/CADKernel/Math/BSpline.h
--- a/Engine/Source/Runtime/Datasmith/CADKernel/Public/CADKernel/Math/BSpline.h
+++ b/Engine/Source/Runtime/Datasmith/CADKernel/Public/CADKernel/Math/BSpline.h
@@ -34,7 +34,7 @@ namespace CADKernel
 		CADKERNEL_API void EvaluatePoint(const FNURBSSurface&, const FPoint2D& InPoint2D, FSurfacicPoint& OutPoint3D, int32 InDerivativeOrder);
 		CADKERNEL_API void EvaluatePointGrid(const FNURBSSurface&, const FCoordinateGrid& Coords, FSurfacicSampling& OutPoints, bool bComputeNormals);
 
-		CADKERNEL_API TSharedRef<FNURBSCurve> DuplicateNurbsCurveWithHigherDegree(int32 degre, const FNURBSCurve& InCurve);
+		CADKERNEL_API TSharedPtr<FNURBSCurve> DuplicateNurbsCurveWithHigherDegree(int32 degre, const FNURBSCurve& InCurve);
 	}
 }
 
