DirectLink - Adding missing support for Hierarchical-Instanced-StaticMesh Actors.

To support the HISM actors we must serialize arrays of Transforms.
Right now we're missing support to either serialize FTransform directly or to serialize TArrays of FVector and FQuat. I chose to add the _transform StoreType as it seems to be the most hastle-free solution.

diff --git a/Engine/Source/Runtime/Datasmith/DatasmithCore/Private/DatasmithSceneElementsImpl.cpp b/Engine/Source/Runtime/Datasmith/DatasmithCore/Private/DatasmithSceneElementsImpl.cpp
--- a/Engine/Source/Runtime/Datasmith/DatasmithCore/Private/DatasmithSceneElementsImpl.cpp
+++ b/Engine/Source/Runtime/Datasmith/DatasmithCore/Private/DatasmithSceneElementsImpl.cpp
@@ -163,6 +163,7 @@ FDatasmithMaterialIDElementImpl::FDatasmithMaterialIDElementImpl(const TCHAR* In
 FDatasmithHierarchicalInstancedStaticMeshActorElementImpl::FDatasmithHierarchicalInstancedStaticMeshActorElementImpl(const TCHAR* InName)
 	: FDatasmithMeshActorElementImpl< IDatasmithHierarchicalInstancedStaticMeshActorElement >(InName, EDatasmithElementType::HierarchicalInstanceStaticMesh)
 {
+	Store.RegisterParameter(Instances, "Instaces");
 }
 
 FDatasmithHierarchicalInstancedStaticMeshActorElementImpl::~FDatasmithHierarchicalInstancedStaticMeshActorElementImpl()
@@ -171,34 +172,34 @@ FDatasmithHierarchicalInstancedStaticMeshActorElementImpl::~FDatasmithHierarchic
 
 int32 FDatasmithHierarchicalInstancedStaticMeshActorElementImpl::GetInstancesCount() const
 {
-	return Instances.Num();
+	return Instances.Get().Num();
 }
 
 void FDatasmithHierarchicalInstancedStaticMeshActorElementImpl::ReserveSpaceForInstances(int32 NumIntances)
 {
-	Instances.Reserve(NumIntances);
+	Instances.Get().Reserve(NumIntances);
 }
 
 int32 FDatasmithHierarchicalInstancedStaticMeshActorElementImpl::AddInstance(const FTransform& Transform)
 {
-	Instances.Add(Transform);
-	return Instances.Num() - 1;
+	Instances.Get().Add(Transform);
+	return Instances.Get().Num() - 1;
 }
 
 FTransform FDatasmithHierarchicalInstancedStaticMeshActorElementImpl::GetInstance(int32 InstanceIndex) const
 {
-	if (Instances.IsValidIndex(InstanceIndex))
+	if (Instances.Get().IsValidIndex(InstanceIndex))
 	{
-		return Instances[InstanceIndex];
+		return Instances.Get()[InstanceIndex];
 	}
 	return FTransform();
 }
 
 void FDatasmithHierarchicalInstancedStaticMeshActorElementImpl::RemoveInstance(int32 InstanceIndex)
 {
-	if (Instances.IsValidIndex(InstanceIndex))
+	if (Instances.Get().IsValidIndex(InstanceIndex))
 	{
-		Instances.RemoveAtSwap(InstanceIndex);
+		Instances.Get().RemoveAtSwap(InstanceIndex);
 	}
 }
 
diff --git a/Engine/Source/Runtime/Datasmith/DatasmithCore/Private/DatasmithSceneElementsImpl.h b/Engine/Source/Runtime/Datasmith/DatasmithCore/Private/DatasmithSceneElementsImpl.h
--- a/Engine/Source/Runtime/Datasmith/DatasmithCore/Private/DatasmithSceneElementsImpl.h
+++ b/Engine/Source/Runtime/Datasmith/DatasmithCore/Private/DatasmithSceneElementsImpl.h
@@ -402,9 +402,9 @@ private:
 template< typename InterfaceType = IDatasmithMeshActorElement >
 class FDatasmithMeshActorElementImpl : public FDatasmithActorElementImpl< InterfaceType >
 {
+public:
 	using FDatasmithElementImpl< InterfaceType >::Store;
 
-public:
 	explicit FDatasmithMeshActorElementImpl(const TCHAR* InName);
 
 	virtual void AddMaterialOverride(const TCHAR* InMaterialName, int32 Id) override;
@@ -429,16 +429,15 @@ private:
 
 template < typename InterfaceType >
 FDatasmithMeshActorElementImpl< InterfaceType >::FDatasmithMeshActorElementImpl(const TCHAR* InName)
-	: FDatasmithActorElementImpl< InterfaceType >(InName, EDatasmithElementType::StaticMeshActor)
-{
-	this->RegisterReferenceProxy(Materials, "Materials");
-	Store.RegisterParameter(StaticMeshPathName, "StaticMeshPathName");
-}
+	: FDatasmithMeshActorElementImpl< InterfaceType >(InName, EDatasmithElementType::None)
+{}
 
 template < typename InterfaceType >
 FDatasmithMeshActorElementImpl< InterfaceType >::FDatasmithMeshActorElementImpl(const TCHAR* InName, EDatasmithElementType ElementType)
 	: FDatasmithActorElementImpl< InterfaceType >(InName, EDatasmithElementType::StaticMeshActor | ElementType)
 {
+	this->RegisterReferenceProxy(Materials, "Materials");
+	Store.RegisterParameter(StaticMeshPathName, "StaticMeshPathName");
 }
 
 template < typename InterfaceType >
@@ -531,7 +530,7 @@ public:
 	virtual void RemoveInstance(int32 InstanceIndex) override;
 
 private:
-	TArray<FTransform> Instances;// #ue_directlink_reflect
+	TReflected<TArray<FTransform>> Instances;
 };
 
 enum class LightActorFlags : uint8
diff --git a/Engine/Source/Runtime/Datasmith/DirectLink/Public/DirectLinkSerialMethods.h b/Engine/Source/Runtime/Datasmith/DirectLink/Public/DirectLinkSerialMethods.h
--- a/Engine/Source/Runtime/Datasmith/DirectLink/Public/DirectLinkSerialMethods.h
+++ b/Engine/Source/Runtime/Datasmith/DirectLink/Public/DirectLinkSerialMethods.h
@@ -31,6 +31,7 @@ enum EStoreType : uint8
 	_vector, _vector4, _quat,
 	_linearcolor,
 	_md5hash,
+	_transform,
 	_storeTypeLast
 };
 
@@ -66,6 +67,8 @@ enum class ESerialMethod : uint8
 	Quat_f64           = _quat       | _alt1,    // FQuat serialized as FQuat4d (native since 5.0)
 	LinearColor_Default= _linearcolor| _default,
 	MD5Hash_Default    = _md5hash    | _default,
+	Transform_Default  = _transform  | _default, // Added in 5.0, uses LWC.
+	Transform_Array    = _transform  | _array,   // Added in 5.0, uses LWC.
 
 	_NotImplementedYet = 0xff
 };
@@ -75,26 +78,28 @@ static constexpr EStoreType GetStoreType(ESerialMethod Method) { return EStoreTy
 
 // This list maps which serialization method should be used for each Type reflected by directlink
 template<typename T> struct TDefaultSerialMethod;
-template<> struct TDefaultSerialMethod<bool>            { constexpr static ESerialMethod Value = ESerialMethod::Uint8_Default;      };
-template<> struct TDefaultSerialMethod<uint8>           { constexpr static ESerialMethod Value = ESerialMethod::Uint8_Default;      };
-template<> struct TDefaultSerialMethod<int32>           { constexpr static ESerialMethod Value = ESerialMethod::Int32_Default;      };
-template<> struct TDefaultSerialMethod<uint32>          { constexpr static ESerialMethod Value = ESerialMethod::Uint32_Default;     };
-template<> struct TDefaultSerialMethod<uint64>          { constexpr static ESerialMethod Value = ESerialMethod::Uint64_Default;     };
-template<> struct TDefaultSerialMethod<FString>         { constexpr static ESerialMethod Value = ESerialMethod::String_Default;     };
-template<> struct TDefaultSerialMethod<float>           { constexpr static ESerialMethod Value = ESerialMethod::Float_Default;      };
-template<> struct TDefaultSerialMethod<double>          { constexpr static ESerialMethod Value = ESerialMethod::Double_Default;     };
-template<> struct TDefaultSerialMethod<TArray<int32>>   { constexpr static ESerialMethod Value = ESerialMethod::Int32_Array;        };
-template<> struct TDefaultSerialMethod<TArray<FString>> { constexpr static ESerialMethod Value = ESerialMethod::String_Array;       };
-template<> struct TDefaultSerialMethod<TArray<float>>   { constexpr static ESerialMethod Value = ESerialMethod::Float_Array;        };
-template<> struct TDefaultSerialMethod<FLinearColor>    { constexpr static ESerialMethod Value = ESerialMethod::LinearColor_Default;};
-template<> struct TDefaultSerialMethod<FMD5Hash>        { constexpr static ESerialMethod Value = ESerialMethod::MD5Hash_Default;    };
+template<> struct TDefaultSerialMethod<bool>               { constexpr static ESerialMethod Value = ESerialMethod::Uint8_Default;      };
+template<> struct TDefaultSerialMethod<uint8>              { constexpr static ESerialMethod Value = ESerialMethod::Uint8_Default;      };
+template<> struct TDefaultSerialMethod<int32>              { constexpr static ESerialMethod Value = ESerialMethod::Int32_Default;      };
+template<> struct TDefaultSerialMethod<uint32>             { constexpr static ESerialMethod Value = ESerialMethod::Uint32_Default;     };
+template<> struct TDefaultSerialMethod<uint64>             { constexpr static ESerialMethod Value = ESerialMethod::Uint64_Default;     };
+template<> struct TDefaultSerialMethod<FString>            { constexpr static ESerialMethod Value = ESerialMethod::String_Default;     };
+template<> struct TDefaultSerialMethod<float>              { constexpr static ESerialMethod Value = ESerialMethod::Float_Default;      };
+template<> struct TDefaultSerialMethod<double>             { constexpr static ESerialMethod Value = ESerialMethod::Double_Default;     };
+template<> struct TDefaultSerialMethod<TArray<int32>>      { constexpr static ESerialMethod Value = ESerialMethod::Int32_Array;        };
+template<> struct TDefaultSerialMethod<TArray<FString>>    { constexpr static ESerialMethod Value = ESerialMethod::String_Array;       };
+template<> struct TDefaultSerialMethod<TArray<float>>      { constexpr static ESerialMethod Value = ESerialMethod::Float_Array;        };
+template<> struct TDefaultSerialMethod<FLinearColor>       { constexpr static ESerialMethod Value = ESerialMethod::LinearColor_Default;};
+template<> struct TDefaultSerialMethod<FMD5Hash>           { constexpr static ESerialMethod Value = ESerialMethod::MD5Hash_Default;    };
+template<> struct TDefaultSerialMethod<FTransform>         { constexpr static ESerialMethod Value = ESerialMethod::Transform_Default;  };
+template<> struct TDefaultSerialMethod<TArray<FTransform>> { constexpr static ESerialMethod Value = ESerialMethod::Transform_Array;    };
 #if DL_SEND_LWC_TYPES_AS_F32
 // for compatibility with pre-LWC endpoints, we send LWC-aware types as their f32 variants
-template<> struct TDefaultSerialMethod<FVector>         { constexpr static ESerialMethod Value = ESerialMethod::Vector_f32;         };
-template<> struct TDefaultSerialMethod<FQuat>           { constexpr static ESerialMethod Value = ESerialMethod::Quat_f32;           };
+template<> struct TDefaultSerialMethod<FVector>            { constexpr static ESerialMethod Value = ESerialMethod::Vector_f32;         };
+template<> struct TDefaultSerialMethod<FQuat>              { constexpr static ESerialMethod Value = ESerialMethod::Quat_f32;           };
 #else
-template<> struct TDefaultSerialMethod<FVector>         { constexpr static ESerialMethod Value = ESerialMethod::Vector_f64;         };
-template<> struct TDefaultSerialMethod<FQuat>           { constexpr static ESerialMethod Value = ESerialMethod::Quat_f64;           };
+template<> struct TDefaultSerialMethod<FVector>            { constexpr static ESerialMethod Value = ESerialMethod::Vector_f64;         };
+template<> struct TDefaultSerialMethod<FQuat>>             { constexpr static ESerialMethod Value = ESerialMethod::Quat_f64;           };
 #endif
 
 
@@ -144,24 +149,26 @@ inline bool SerialAny(FArchive& Ar, void* data, ESerialMethod Method)
 #define SerialAny_Case(enumvalue, casttype) case enumvalue: Serial<enumvalue>(Ar, (casttype*)data); return true;
 	switch (Method)
 	{
-		SerialAny_Case(ESerialMethod::Bool_Default       , bool           );
-		SerialAny_Case(ESerialMethod::Uint8_Default      , uint8          );
-		SerialAny_Case(ESerialMethod::Int32_Default      , int32          );
-		SerialAny_Case(ESerialMethod::Uint32_Default     , uint32         );
-		SerialAny_Case(ESerialMethod::Uint64_Default     , uint64         );
-		SerialAny_Case(ESerialMethod::Uint32_Packed      , uint32         );
-		SerialAny_Case(ESerialMethod::String_Default     , FString        );
-		SerialAny_Case(ESerialMethod::Double_Default     , double         );
-		SerialAny_Case(ESerialMethod::Float_Default      , float          );
-		SerialAny_Case(ESerialMethod::Int32_Array        , TArray<int32>  );
-		SerialAny_Case(ESerialMethod::String_Array       , TArray<FString>);
-		SerialAny_Case(ESerialMethod::Float_Array        , TArray<float>  );
-		SerialAny_Case(ESerialMethod::Vector_f32         , FVector        );
-		SerialAny_Case(ESerialMethod::Vector_f64         , FVector        );
-		SerialAny_Case(ESerialMethod::Quat_f32           , FQuat          );
-		SerialAny_Case(ESerialMethod::Quat_f64           , FQuat          );
-		SerialAny_Case(ESerialMethod::LinearColor_Default, FLinearColor   );
-		SerialAny_Case(ESerialMethod::MD5Hash_Default    , FMD5Hash       );
+		SerialAny_Case(ESerialMethod::Bool_Default       , bool              );
+		SerialAny_Case(ESerialMethod::Uint8_Default      , uint8             );
+		SerialAny_Case(ESerialMethod::Int32_Default      , int32             );
+		SerialAny_Case(ESerialMethod::Uint32_Default     , uint32            );
+		SerialAny_Case(ESerialMethod::Uint64_Default     , uint64            );
+		SerialAny_Case(ESerialMethod::Uint32_Packed      , uint32            );
+		SerialAny_Case(ESerialMethod::String_Default     , FString           );
+		SerialAny_Case(ESerialMethod::Double_Default     , double            );
+		SerialAny_Case(ESerialMethod::Float_Default      , float             );
+		SerialAny_Case(ESerialMethod::Int32_Array        , TArray<int32>     );
+		SerialAny_Case(ESerialMethod::String_Array       , TArray<FString>   );
+		SerialAny_Case(ESerialMethod::Float_Array        , TArray<float>     );
+		SerialAny_Case(ESerialMethod::Vector_f32         , FVector           );
+		SerialAny_Case(ESerialMethod::Vector_f64         , FVector           );
+		SerialAny_Case(ESerialMethod::Quat_f32           , FQuat             );
+		SerialAny_Case(ESerialMethod::Quat_f64           , FQuat             );
+		SerialAny_Case(ESerialMethod::LinearColor_Default, FLinearColor      );
+		SerialAny_Case(ESerialMethod::MD5Hash_Default    , FMD5Hash          );
+		SerialAny_Case(ESerialMethod::Transform_Default  , FTransform        );
+		SerialAny_Case(ESerialMethod::Transform_Array    , TArray<FTransform>);
 
 		case ESerialMethod::None: return true;
 		case ESerialMethod::_NotImplementedYet:
