Fix CADKernel used with TechSoft, 3 mains problems:
- by default, from TechSoft, faces with a periodic basis surface are not always split on parametric seams
- CADKernel expects to have the external loop in first, that is not always the case with TechSoft
- The function to test the relative orienation of two coincident closed edges (StartVertex = EndVertex) has not been yet implemented
#swarm https://p4-swarm.epicgames.net/reviews/19379931

diff --git a/Engine/Plugins/Enterprise/DatasmithCADImporter/Source/CADInterfaces/Private/CoreTechBridge.cpp b/Engine/Plugins/Enterprise/DatasmithCADImporter/Source/CADInterfaces/Private/CoreTechBridge.cpp
--- a/Engine/Plugins/Enterprise/DatasmithCADImporter/Source/CADInterfaces/Private/CoreTechBridge.cpp
+++ b/Engine/Plugins/Enterprise/DatasmithCADImporter/Source/CADInterfaces/Private/CoreTechBridge.cpp
@@ -224,15 +224,17 @@ void FCoreTechBridge::AddFace(CT_OBJECT_ID CTFaceId, TSharedRef<FShell>& Shell)
 
 	TArray<TSharedPtr<FTopologicalLoop>> Loops;
 
+	bool bIsExternalLoop = true;
 	CTLoopIds.IteratorInitialize();
 	CT_OBJECT_ID CTLoopId;
 	while ((CTLoopId = CTLoopIds.IteratorIter()) != 0)
 	{
-		TSharedPtr<FTopologicalLoop> Loop = AddLoop(CTLoopId, Surface);
+		TSharedPtr<FTopologicalLoop> Loop = AddLoop(CTLoopId, Surface, bIsExternalLoop);
 		if (!Loop.IsValid())
 		{
 			continue;
 		}
+		bIsExternalLoop = false;
 
 		LinkEdgesLoop(CTLoopId, *Loop);
 
@@ -287,7 +289,7 @@ void FCoreTechBridge::Get2DCurvesRange(CT_OBJECT_ID CTFaceId, FSurfacicBoundary&
 	CT_FACE_IO::AskUVminmax(CTFaceId, OutBoundary[EIso::IsoU].Min, OutBoundary[EIso::IsoU].Max, OutBoundary[EIso::IsoV].Min, OutBoundary[EIso::IsoV].Max);
 }
 
-TSharedPtr<FTopologicalLoop> FCoreTechBridge::AddLoop(CT_OBJECT_ID CTLoopId, TSharedRef<FSurface>& Surface)
+TSharedPtr<FTopologicalLoop> FCoreTechBridge::AddLoop(CT_OBJECT_ID CTLoopId, TSharedRef<FSurface>& Surface, const bool bIsExternalLoop)
 {
 	Report.LoopCount++;
 
@@ -322,7 +324,7 @@ TSharedPtr<FTopologicalLoop> FCoreTechBridge::AddLoop(CT_OBJECT_ID CTLoopId, TSh
 		return TSharedPtr<FTopologicalLoop>();
 	}
 
-	return FTopologicalLoop::Make(Edges, Directions, GeometricTolerance);
+	return FTopologicalLoop::Make(Edges, Directions, bIsExternalLoop, GeometricTolerance);
 }
 
 void FCoreTechBridge::LinkEdgesLoop(CT_OBJECT_ID CTLoopId, FTopologicalLoop& Loop)
@@ -405,6 +407,10 @@ TSharedPtr<FTopologicalEdge> FCoreTechBridge::AddEdge(CT_OBJECT_ID CTCoedgeId, T
 			Weights[IPole] = RawPoles[Index];
 		}
 	}
+	else
+	{
+		Weights.Init(1., PoleNum);
+	}
 
 	const FSurfacicBoundary& SurfaceBounds = Surface->GetBoundary();
 
@@ -856,6 +862,11 @@ TSharedPtr<FCurve> FCoreTechBridge::AddNurbsCurve(CT_OBJECT_ID CTCurveId)
 			Weights[IPole] = RawPoles[Index];
 		}
 	}
+	else
+	{
+		Weights.Init(1., PoleNum);
+	}
+
 	return FEntity::MakeShared<FNURBSCurve>(Degre, Knots, Poles, Weights);
 }
 
diff --git a/Engine/Plugins/Enterprise/DatasmithCADImporter/Source/CADInterfaces/Private/CoreTechBridge.h b/Engine/Plugins/Enterprise/DatasmithCADImporter/Source/CADInterfaces/Private/CoreTechBridge.h
--- a/Engine/Plugins/Enterprise/DatasmithCADImporter/Source/CADInterfaces/Private/CoreTechBridge.h
+++ b/Engine/Plugins/Enterprise/DatasmithCADImporter/Source/CADInterfaces/Private/CoreTechBridge.h
@@ -64,7 +64,7 @@ public:
 
 private:
 
-	TSharedPtr<FTopologicalLoop> AddLoop(CT_OBJECT_ID CTLoopId, TSharedRef<FSurface>& Surface);
+	TSharedPtr<FTopologicalLoop> AddLoop(CT_OBJECT_ID CTLoopId, TSharedRef<FSurface>& Surface, const bool bIsExternalLoop);
 
 	/**
 	 * Build face's links with its neighbor have to be done after the loop is finalize.
diff --git a/Engine/Plugins/Enterprise/DatasmithCADImporter/Source/CADInterfaces/Private/CoreTechFileParserImpl.cpp b/Engine/Plugins/Enterprise/DatasmithCADImporter/Source/CADInterfaces/Private/CoreTechFileParserImpl.cpp
--- a/Engine/Plugins/Enterprise/DatasmithCADImporter/Source/CADInterfaces/Private/CoreTechFileParserImpl.cpp
+++ b/Engine/Plugins/Enterprise/DatasmithCADImporter/Source/CADInterfaces/Private/CoreTechFileParserImpl.cpp
@@ -640,6 +640,7 @@ namespace CADLibrary
 
 		int32 BodyIndex = CADFileData.AddBody(BodyId);
 		FArchiveBody& Body = CADFileData.GetBodyAt(BodyIndex);
+		Body.ParentId = ParentId;
 
 		ReadNodeMetaData(BodyId, Body.MetaData);
 
@@ -710,6 +711,8 @@ namespace CADLibrary
 
 		int32 Index = CADFileData.AddBody(BodyId);
 		FArchiveBody& ArchiveBody = CADFileData.GetBodyAt(Index);
+		ArchiveBody.ParentId = ParentId;
+
 		ReadNodeMetaData(BodyId, ArchiveBody.MetaData);
 
 		FBodyMesh& BodyMesh = CADFileData.AddBodyMesh(BodyId, ArchiveBody);
@@ -773,6 +776,7 @@ namespace CADLibrary
 	{
 		int32 Index = CADFileData.AddBody(CADKernelBody.GetHostId());
 		FArchiveBody& ArchiveBody = CADFileData.GetBodyAt(Index);
+		ArchiveBody.ParentId = ParentId;
 
 		CADKernelBody.ExtractMetaData(ArchiveBody.MetaData);
 
diff --git a/Engine/Plugins/Enterprise/DatasmithCADImporter/Source/CADInterfaces/Private/TechSoftBridge.cpp b/Engine/Plugins/Enterprise/DatasmithCADImporter/Source/CADInterfaces/Private/TechSoftBridge.cpp
--- a/Engine/Plugins/Enterprise/DatasmithCADImporter/Source/CADInterfaces/Private/TechSoftBridge.cpp
+++ b/Engine/Plugins/Enterprise/DatasmithCADImporter/Source/CADInterfaces/Private/TechSoftBridge.cpp
@@ -225,11 +225,12 @@ CADKernel::FBody* FTechSoftBridge::GetBody(A3DRiBrepModel* A3DBRepModel)
 	return nullptr;
 }
 
-CADKernel::FBody* FTechSoftBridge::AddBody(A3DRiBrepModel* A3DBRepModel, TMap<FString, FString> MetaData)
+CADKernel::FBody* FTechSoftBridge::AddBody(A3DRiBrepModel* A3DBRepModel, TMap<FString, FString> MetaData, const double InBodyScale)
 {
 	Report.BodyCount++;
 
-	UnitScale = Parser.GetFileUnit();
+	//CADKernel working unit is mm
+	BodyScale = InBodyScale * 10.;
 
 	FEntityMetaData BRepMetaData;
 	Parser.ExtractMetaData(A3DBRepModel, BRepMetaData);
@@ -287,7 +288,7 @@ void FTechSoftBridge::TraverseBrepData(const A3DTopoBrepData* A3DBrepData, TShar
 				{
 					if (TopoContextData->m_bHaveScale)
 					{
-						UnitScale *= TopoContextData->m_dScale;
+						BodyScale *= TopoContextData->m_dScale;
 					}
 				}
 			}
@@ -671,7 +672,7 @@ TSharedPtr<CADKernel::FTopologicalEdge> FTechSoftBridge::AddEdge(const A3DTopoCo
 	return Edge;
 }
 
-TSharedPtr<CADKernel::FTopologicalLoop> FTechSoftBridge::AddLoop(const A3DTopoLoop* A3DLoop, const TSharedRef<CADKernel::FSurface>& Surface, const TechSoftUtils::FUVReparameterization& UVReparameterization)
+TSharedPtr<CADKernel::FTopologicalLoop> FTechSoftBridge::AddLoop(const A3DTopoLoop* A3DLoop, const TSharedRef<CADKernel::FSurface>& Surface, const TechSoftUtils::FUVReparameterization& UVReparameterization, const bool bIsExternalLoop)
 {
 	Report.LoopCount++;
 
@@ -706,7 +707,7 @@ TSharedPtr<CADKernel::FTopologicalLoop> FTechSoftBridge::AddLoop(const A3DTopoLo
 		return TSharedPtr<CADKernel::FTopologicalLoop>();
 	}
 
-	return CADKernel::FTopologicalLoop::Make(Edges, Directions, GeometricTolerance);
+	return CADKernel::FTopologicalLoop::Make(Edges, Directions, bIsExternalLoop, GeometricTolerance);
 }
 
 void FTechSoftBridge::AddFace(const A3DTopoFace* A3DFace, CADKernel::EOrientation Orientation, TSharedRef<CADKernel::FShell>& Shell)
@@ -738,7 +739,7 @@ void FTechSoftBridge::AddFace(const A3DTopoFace* A3DFace, CADKernel::EOrientatio
 
 	if (TopoFaceData->m_bHasTrimDomain)
 	{
-		CADKernel::FSurfacicBoundary SurfaceBoundary = TechSoftUtils::GetSurfacicBoundary(TopoFaceData->m_sSurfaceDomain, UnitScale);
+		CADKernel::FSurfacicBoundary SurfaceBoundary = TechSoftUtils::GetSurfacicBoundary(TopoFaceData->m_sSurfaceDomain, BodyScale);
 		Surface->TrimBoundaryTo(SurfaceBoundary);
 	}
 
@@ -752,7 +753,8 @@ void FTechSoftBridge::AddFace(const A3DTopoFace* A3DFace, CADKernel::EOrientatio
 
 		for (A3DUns32 Index = 0; Index < TopoFaceData->m_uiLoopSize; ++Index)
 		{
-			TSharedPtr<CADKernel::FTopologicalLoop> Loop = FTechSoftBridge::AddLoop(TopoFaceData->m_ppLoops[Index], Surface, UVReparameterization);
+			const bool bIsExternalLoop = (Index == TopoFaceData->m_uiOuterLoopIndex);
+			TSharedPtr<CADKernel::FTopologicalLoop> Loop = FTechSoftBridge::AddLoop(TopoFaceData->m_ppLoops[Index], Surface, UVReparameterization, bIsExternalLoop);
 			if (!Loop.IsValid())
 			{
 				continue;
@@ -825,11 +827,11 @@ TSharedPtr<CADKernel::FSurface> FTechSoftBridge::AddSurface(const A3DSurfBase* A
 			return AddNurbsSurface(A3DSurface, OutUVReparameterization);
 
 		case kA3DTypeSurfCone:
-			OutUVReparameterization.SetScale(1, UnitScale);
+			OutUVReparameterization.SetScale(1, BodyScale);
 			return AddConeSurface(A3DSurface, OutUVReparameterization);
 
 		case kA3DTypeSurfCylinder:
-			OutUVReparameterization.SetScale(1, UnitScale);
+			OutUVReparameterization.SetScale(1, BodyScale);
 			return AddCylinderSurface(A3DSurface, OutUVReparameterization);
 
 		case kA3DTypeSurfCylindrical:
@@ -842,7 +844,7 @@ TSharedPtr<CADKernel::FSurface> FTechSoftBridge::AddSurface(const A3DSurfBase* A
 			return AddPipeSurface(A3DSurface, OutUVReparameterization);
 
 		case kA3DTypeSurfPlane:
-			OutUVReparameterization.SetScale(UnitScale, UnitScale);
+			OutUVReparameterization.SetScale(BodyScale, BodyScale);
 			return AddPlaneSurface(A3DSurface, OutUVReparameterization);
 
 		case kA3DTypeSurfRuled:
@@ -888,9 +890,9 @@ TSharedPtr<CADKernel::FSurface> FTechSoftBridge::AddConeSurface(const A3DSurfBas
 	}
 
 	OutUVReparameterization.AddUVTransform(A3DConeData->m_sParam);
-	CADKernel::FMatrixH CoordinateSystem = TechSoftUtils::CreateCoordinateSystem(A3DConeData->m_sTrsf, UnitScale);
-	CADKernel::FSurfacicBoundary Boundary = TechSoftUtils::GetSurfacicBoundary(A3DConeData->m_sParam.m_sUVDomain, UnitScale);
-	TSharedPtr<CADKernel::FSurface> Cone = CADKernel::FEntity::MakeShared<CADKernel::FConeSurface>(GeometricTolerance, CoordinateSystem, A3DConeData->m_dRadius * UnitScale, A3DConeData->m_dSemiAngle, Boundary);
+	CADKernel::FMatrixH CoordinateSystem = TechSoftUtils::CreateCoordinateSystem(A3DConeData->m_sTrsf, BodyScale);
+	CADKernel::FSurfacicBoundary Boundary = TechSoftUtils::GetSurfacicBoundary(A3DConeData->m_sParam.m_sUVDomain, BodyScale);
+	TSharedPtr<CADKernel::FSurface> Cone = CADKernel::FEntity::MakeShared<CADKernel::FConeSurface>(GeometricTolerance, CoordinateSystem, A3DConeData->m_dRadius * BodyScale, A3DConeData->m_dSemiAngle, Boundary);
 
 	return Cone;
 }
@@ -906,9 +908,9 @@ TSharedPtr<CADKernel::FSurface> FTechSoftBridge::AddCylinderSurface(const A3DSur
 	}
 
 	OutUVReparameterization.AddUVTransform(A3DCylinderData->m_sParam);
-	CADKernel::FMatrixH CoordinateSystem = TechSoftUtils::CreateCoordinateSystem(A3DCylinderData->m_sTrsf, UnitScale);
-	CADKernel::FSurfacicBoundary Boundary = TechSoftUtils::GetSurfacicBoundary(A3DCylinderData->m_sParam.m_sUVDomain, UnitScale);
-	TSharedPtr<CADKernel::FSurface> Cylinder = CADKernel::FEntity::MakeShared<CADKernel::FCylinderSurface>(GeometricTolerance, CoordinateSystem, A3DCylinderData->m_dRadius * UnitScale, Boundary);
+	CADKernel::FMatrixH CoordinateSystem = TechSoftUtils::CreateCoordinateSystem(A3DCylinderData->m_sTrsf, BodyScale);
+	CADKernel::FSurfacicBoundary Boundary = TechSoftUtils::GetSurfacicBoundary(A3DCylinderData->m_sParam.m_sUVDomain, BodyScale);
+	TSharedPtr<CADKernel::FSurface> Cylinder = CADKernel::FEntity::MakeShared<CADKernel::FCylinderSurface>(GeometricTolerance, CoordinateSystem, A3DCylinderData->m_dRadius * BodyScale, Boundary);
 
 	return Cylinder;
 }
@@ -961,8 +963,8 @@ TSharedPtr<CADKernel::FSurface> FTechSoftBridge::AddPlaneSurface(const A3DSurfPl
 	}
 
 	OutUVReparameterization.AddUVTransform(A3DPlaneData->m_sParam);
-	CADKernel::FMatrixH CoordinateSystem = TechSoftUtils::CreateCoordinateSystem(A3DPlaneData->m_sTrsf, UnitScale);
-	CADKernel::FSurfacicBoundary Boundary = TechSoftUtils::GetSurfacicBoundary(A3DPlaneData->m_sParam.m_sUVDomain, UnitScale);
+	CADKernel::FMatrixH CoordinateSystem = TechSoftUtils::CreateCoordinateSystem(A3DPlaneData->m_sTrsf, BodyScale);
+	CADKernel::FSurfacicBoundary Boundary = TechSoftUtils::GetSurfacicBoundary(A3DPlaneData->m_sParam.m_sUVDomain, BodyScale);
 	TSharedPtr<CADKernel::FSurface> Plane = CADKernel::FEntity::MakeShared<CADKernel::FPlaneSurface>(GeometricTolerance, CoordinateSystem, Boundary);
 
 	return Plane;
@@ -1020,9 +1022,9 @@ TSharedPtr<CADKernel::FSurface> FTechSoftBridge::AddSphereSurface(const A3DSurfB
 	}
 
 	OutUVReparameterization.AddUVTransform(A3DSphereData->m_sParam);
-	CADKernel::FMatrixH CoordinateSystem = TechSoftUtils::CreateCoordinateSystem(A3DSphereData->m_sTrsf, UnitScale);
-	CADKernel::FSurfacicBoundary Boundary = TechSoftUtils::GetSurfacicBoundary(A3DSphereData->m_sParam.m_sUVDomain, UnitScale);
-	TSharedRef<CADKernel::FSphericalSurface> Sphere = CADKernel::FEntity::MakeShared<CADKernel::FSphericalSurface>(GeometricTolerance, CoordinateSystem, A3DSphereData->m_dRadius * UnitScale, Boundary);
+	CADKernel::FMatrixH CoordinateSystem = TechSoftUtils::CreateCoordinateSystem(A3DSphereData->m_sTrsf, BodyScale);
+	CADKernel::FSurfacicBoundary Boundary = TechSoftUtils::GetSurfacicBoundary(A3DSphereData->m_sParam.m_sUVDomain, BodyScale);
+	TSharedRef<CADKernel::FSphericalSurface> Sphere = CADKernel::FEntity::MakeShared<CADKernel::FSphericalSurface>(GeometricTolerance, CoordinateSystem, A3DSphereData->m_dRadius * BodyScale, Boundary);
 
 	return Sphere;
 }
@@ -1038,9 +1040,9 @@ TSharedPtr<CADKernel::FSurface> FTechSoftBridge::AddTorusSurface(const A3DSurfBa
 	}
 
 	OutUVReparameterization.AddUVTransform(A3DTorusData->m_sParam);
-	CADKernel::FMatrixH CoordinateSystem = TechSoftUtils::CreateCoordinateSystem(A3DTorusData->m_sTrsf, UnitScale);
-	CADKernel::FSurfacicBoundary Boundary = TechSoftUtils::GetSurfacicBoundary(A3DTorusData->m_sParam.m_sUVDomain, UnitScale);
-	TSharedPtr<CADKernel::FSurface> Torus = CADKernel::FEntity::MakeShared<CADKernel::FTorusSurface>(GeometricTolerance, CoordinateSystem, A3DTorusData->m_dMajorRadius * UnitScale, A3DTorusData->m_dMinorRadius * UnitScale, Boundary);
+	CADKernel::FMatrixH CoordinateSystem = TechSoftUtils::CreateCoordinateSystem(A3DTorusData->m_sTrsf, BodyScale);
+	CADKernel::FSurfacicBoundary Boundary = TechSoftUtils::GetSurfacicBoundary(A3DTorusData->m_sParam.m_sUVDomain, BodyScale);
+	TSharedPtr<CADKernel::FSurface> Torus = CADKernel::FEntity::MakeShared<CADKernel::FTorusSurface>(GeometricTolerance, CoordinateSystem, A3DTorusData->m_dMajorRadius * BodyScale, A3DTorusData->m_dMinorRadius * BodyScale, Boundary);
 
 	return Torus;
 }
@@ -1199,11 +1201,11 @@ TSharedPtr<CADKernel::FSurface> FTechSoftBridge::AddSurfaceNurbs(const A3DSurfNu
 
 	TArray<CADKernel::FPoint> Poles;
 	TechSoftUtils::FillPointArray(NurbsData.PoleUCount, NurbsData.PoleVCount, A3DNurbsData.m_pCtrlPts, NurbsData.Poles);
-	if (!FMath::IsNearlyEqual(UnitScale, 1.))
+	if (!FMath::IsNearlyEqual(BodyScale, 1.))
 	{
 		for (CADKernel::FPoint& Point : NurbsData.Poles)
 		{
-			Point *= UnitScale;
+			Point *= BodyScale;
 		}
 	}
 
@@ -1250,4 +1252,4 @@ void FTechSoftBridge::AddMetadata(FEntityMetaData& MetaData, CADKernel::FTopolog
 }
 
 
-#endif // USE_TECHSOFT_SDK
\ No newline at end of file
+#endif // USE_TECHSOFT_SDK
diff --git a/Engine/Plugins/Enterprise/DatasmithCADImporter/Source/CADInterfaces/Private/TechSoftBridge.h b/Engine/Plugins/Enterprise/DatasmithCADImporter/Source/CADInterfaces/Private/TechSoftBridge.h
--- a/Engine/Plugins/Enterprise/DatasmithCADImporter/Source/CADInterfaces/Private/TechSoftBridge.h
+++ b/Engine/Plugins/Enterprise/DatasmithCADImporter/Source/CADInterfaces/Private/TechSoftBridge.h
@@ -140,11 +140,11 @@ private:
 	TMap<CADKernel::FBody*, const A3DEntity*> CADKernelToTechSoft;
 	TMap<const A3DTopoCoEdge*, TSharedPtr<CADKernel::FTopologicalEdge>> A3DEdgeToEdge;
 
-	double UnitScale = 1;
+	double BodyScale = 1;
 
 public:
 	FTechSoftBridge(FTechSoftFileParser& InParser, CADKernel::FSession& InSession, CADKernel::FCADFileReport& InReport);
-	CADKernel::FBody* AddBody(A3DRiBrepModel* A3DBRepModel, TMap<FString, FString> MetaData);
+	CADKernel::FBody* AddBody(A3DRiBrepModel* A3DBRepModel, TMap<FString, FString> MetaData, const double InBodyScale);
 	CADKernel::FBody* GetBody(A3DRiBrepModel* A3DBRepModel);
 	const A3DRiBrepModel* GetA3DBody(CADKernel::FBody* BRepModel);
 
@@ -156,7 +156,7 @@ private:
 	void TraverseShell(const A3DTopoShell* A3DShell, TSharedRef<CADKernel::FBody>& Body);
 	void AddFace(const A3DTopoFace* A3DFace, CADKernel::EOrientation Orientation, TSharedRef<CADKernel::FShell>& Body);
 
-	TSharedPtr<CADKernel::FTopologicalLoop> AddLoop(const A3DTopoLoop* A3DLoop, const TSharedRef<CADKernel::FSurface>& Surface, const TechSoftUtils::FUVReparameterization& UVReparameterization);
+	TSharedPtr<CADKernel::FTopologicalLoop> AddLoop(const A3DTopoLoop* A3DLoop, const TSharedRef<CADKernel::FSurface>& Surface, const TechSoftUtils::FUVReparameterization& UVReparameterization, const bool bIsExternalLoop);
 	TSharedPtr<CADKernel::FTopologicalEdge> AddEdge(const A3DTopoCoEdge* A3DCoedge, const TSharedRef<CADKernel::FSurface>& Surface, const TechSoftUtils::FUVReparameterization& UVReparameterization, CADKernel::EOrientation& OutOrientation);
 
 	TSharedPtr<CADKernel::FSurface> AddSurface(const A3DSurfBase* A3DSurface, TechSoftUtils::FUVReparameterization& OutUVReparameterization);
diff --git a/Engine/Plugins/Enterprise/DatasmithCADImporter/Source/CADInterfaces/Private/TechSoftFileParserCADKernelTessellator.cpp b/Engine/Plugins/Enterprise/DatasmithCADImporter/Source/CADInterfaces/Private/TechSoftFileParserCADKernelTessellator.cpp
--- a/Engine/Plugins/Enterprise/DatasmithCADImporter/Source/CADInterfaces/Private/TechSoftFileParserCADKernelTessellator.cpp
+++ b/Engine/Plugins/Enterprise/DatasmithCADImporter/Source/CADInterfaces/Private/TechSoftFileParserCADKernelTessellator.cpp
@@ -90,7 +90,7 @@ void FTechSoftFileParserCADKernelTessellator::SewAndMesh(TArray<A3DRiRepresentat
 		}
 		FArchiveBody& Body = CADFileData.GetBodyAt(*BodyIndex);
 
-		CADKernel::FBody* CADKernelBody = TechSoftBridge.AddBody(Representation, Body.MetaData);
+		CADKernel::FBody* CADKernelBody = TechSoftBridge.AddBody(Representation, Body.MetaData, Body.BodyUnit);
 	}
 
 	// Sew if needed
@@ -196,7 +196,7 @@ void FTechSoftFileParserCADKernelTessellator::GenerateBodyMesh(A3DRiRepresentati
 	CADKernel::FCADFileReport Report;
 	FTechSoftBridge TechSoftBridge(*this, CADKernelSession, Report);
 
-	CADKernel::FBody* CADKernelBody = TechSoftBridge.AddBody(Representation, ArchiveBody.MetaData);
+	CADKernel::FBody* CADKernelBody = TechSoftBridge.AddBody(Representation, ArchiveBody.MetaData, ArchiveBody.BodyUnit);
 
 	if (CADFileData.GetImportParameters().GetStitchingTechnique() == StitchingHeal)
 	{
@@ -283,7 +283,7 @@ A3DStatus FTechSoftFileParserCADKernelTessellator::AdaptBRepModel()
 	CopyAndAdaptBrepModelData->m_bUseSameParam = false;                        // If `A3D_TRUE`, surfaces will keep their parametrization when converted to NURBS.       
 	CopyAndAdaptBrepModelData->m_dTol = 1e-3;                                  // Tolerance value of resulting B-rep. The value is relative to the scale of the model.
 	CopyAndAdaptBrepModelData->m_bDeleteCrossingUV = false;                    // If `A3D_TRUE`, UV curves that cross seams of periodic surfaces are replaced by 3D curves 
-	CopyAndAdaptBrepModelData->m_bSplitFaces = false;                          // If `A3D_TRUE`, the faces with a periodic basis surface are split on parametric seams
+	CopyAndAdaptBrepModelData->m_bSplitFaces = true;                           // If `A3D_TRUE`, the faces with a periodic basis surface are split on parametric seams
 	CopyAndAdaptBrepModelData->m_bSplitClosedFaces = false;                    // If `A3D_TRUE`, the faces with a closed basis surface are split into faces at the parametric seam and mid-parameter
 	CopyAndAdaptBrepModelData->m_bForceComputeUV = true;                       // If `A3D_TRUE`, UV curves are computed from the B-rep data
 	CopyAndAdaptBrepModelData->m_bAllowUVCrossingSeams = true;                 // If `A3D_TRUE` and m_bForceComputeUV is set to `A3D_TRUE`, computed UV curves can cross seams.
diff --git a/Engine/Plugins/Enterprise/DatasmithCADImporter/Source/DatasmithOpenNurbsTranslator/Private/OpenNurbsBRepToCADKernelConverter.cpp b/Engine/Plugins/Enterprise/DatasmithCADImporter/Source/DatasmithOpenNurbsTranslator/Private/OpenNurbsBRepToCADKernelConverter.cpp
--- a/Engine/Plugins/Enterprise/DatasmithCADImporter/Source/DatasmithOpenNurbsTranslator/Private/OpenNurbsBRepToCADKernelConverter.cpp
+++ b/Engine/Plugins/Enterprise/DatasmithCADImporter/Source/DatasmithOpenNurbsTranslator/Private/OpenNurbsBRepToCADKernelConverter.cpp
@@ -110,7 +110,7 @@ TSharedRef<CADKernel::FSurface> FOpenNurbsBRepToCADKernelConverter::AddSurface(O
 	return CADKernel::FEntity::MakeShared<CADKernel::FNURBSSurface>(GeometricTolerance, NurbsData);
 }
 
-TSharedPtr<CADKernel::FTopologicalLoop> FOpenNurbsBRepToCADKernelConverter::AddLoop(const ON_BrepLoop& OpenNurbsLoop, TSharedRef<CADKernel::FSurface> & CarrierSurface)
+TSharedPtr<CADKernel::FTopologicalLoop> FOpenNurbsBRepToCADKernelConverter::AddLoop(const ON_BrepLoop& OpenNurbsLoop, TSharedRef<CADKernel::FSurface> & CarrierSurface, const bool bIsExternal)
 {
 	using namespace CADKernel;
 
@@ -146,7 +146,7 @@ TSharedPtr<CADKernel::FTopologicalLoop> FOpenNurbsBRepToCADKernelConverter::AddL
 		return TSharedPtr<FTopologicalLoop>();
 	}
 
-	return FTopologicalLoop::Make(Edges, Directions, GeometricTolerance);
+	return FTopologicalLoop::Make(Edges, Directions, bIsExternal, GeometricTolerance);
 }
 
 void FOpenNurbsBRepToCADKernelConverter::LinkEdgesLoop(const ON_BrepLoop& OpenNurbsLoop, CADKernel::FTopologicalLoop& Loop)
@@ -282,15 +282,17 @@ TSharedPtr<CADKernel::FTopologicalFace> FOpenNurbsBRepToCADKernelConverter::AddF
 		return Face;
 	}
 
+	bool bIsExternal = true;
 	int32 LoopCount = OpenNurbsFace.LoopCount();
 	for (int32 LoopIndex = 0; LoopIndex < LoopCount; ++LoopIndex)
 	{
 		const ON_BrepLoop& OpenNurbsLoop = *OpenNurbsFace.Loop(LoopIndex);
-		TSharedPtr<FTopologicalLoop> Loop = AddLoop(OpenNurbsLoop, Surface);
+		TSharedPtr<FTopologicalLoop> Loop = AddLoop(OpenNurbsLoop, Surface, bIsExternal);
 		if(Loop.IsValid())
 		{
 			LinkEdgesLoop(OpenNurbsLoop, *Loop);
 			Face->AddLoop(Loop);
+			bIsExternal = false;
 		}
 	}
 
diff --git a/Engine/Plugins/Enterprise/DatasmithCADImporter/Source/DatasmithOpenNurbsTranslator/Private/OpenNurbsBRepToCADKernelConverter.h b/Engine/Plugins/Enterprise/DatasmithCADImporter/Source/DatasmithOpenNurbsTranslator/Private/OpenNurbsBRepToCADKernelConverter.h
--- a/Engine/Plugins/Enterprise/DatasmithCADImporter/Source/DatasmithOpenNurbsTranslator/Private/OpenNurbsBRepToCADKernelConverter.h
+++ b/Engine/Plugins/Enterprise/DatasmithCADImporter/Source/DatasmithOpenNurbsTranslator/Private/OpenNurbsBRepToCADKernelConverter.h
@@ -53,7 +53,7 @@ private:
 	TSharedPtr<CADKernel::FTopologicalFace> AddFace(const ON_BrepFace& OpenNurbsFace);
 	TSharedRef<CADKernel::FSurface> AddSurface(ON_NurbsSurface& Surface);
 
-	TSharedPtr<CADKernel::FTopologicalLoop> AddLoop(const ON_BrepLoop& OpenNurbsLoop, TSharedRef<CADKernel::FSurface>& CarrierSurface);
+	TSharedPtr<CADKernel::FTopologicalLoop> AddLoop(const ON_BrepLoop& OpenNurbsLoop, TSharedRef<CADKernel::FSurface>& CarrierSurface, const bool bIsExternal);
 
 	/**
 	 * Build face's links with its neighbor have to be done after the loop is finalize.
diff --git a/Engine/Plugins/Enterprise/DatasmithCADImporter/Source/DatasmithOpenNurbsTranslator/Private/OpenNurbsBRepToTechSoftConverter.cpp b/Engine/Plugins/Enterprise/DatasmithCADImporter/Source/DatasmithOpenNurbsTranslator/Private/OpenNurbsBRepToTechSoftConverter.cpp
--- a/Engine/Plugins/Enterprise/DatasmithCADImporter/Source/DatasmithOpenNurbsTranslator/Private/OpenNurbsBRepToTechSoftConverter.cpp
+++ b/Engine/Plugins/Enterprise/DatasmithCADImporter/Source/DatasmithOpenNurbsTranslator/Private/OpenNurbsBRepToTechSoftConverter.cpp
@@ -65,13 +65,14 @@ struct FOpenNurbsSurfaceInfo
 
 		TArray<A3DDouble>& NodalVector = Axis == EAxis::U ? UNodalVector : VNodalVector;
 		NodalVector.Reserve(KnotSize);
-		NodalVector.Add(OpenNurbsSurface.SuperfluousKnot(Axis, 0));
+		NodalVector.Add(OpenNurbsSurface.Knot(Axis, 0));
 		uint32 KnotCount = OpenNurbsSurface.KnotCount(Axis);
 		for (uint32 i = 0; i < KnotCount; ++i)
 		{
 			NodalVector.Add(OpenNurbsSurface.Knot(Axis, i));
 		}
-		NodalVector.Add(OpenNurbsSurface.SuperfluousKnot(Axis, 1));
+		double LastValue = NodalVector[KnotCount];
+		NodalVector.Add(LastValue);
 
 		A3DDouble*& UKnots = Axis == EAxis::U ? Data.m_pdUKnots : Data.m_pdVKnots;
 		UKnots = NodalVector.GetData();
@@ -154,12 +155,13 @@ A3DCrvBase* FOpenNurbsBRepToTechSoftConverter::CreateCurve(const ON_NurbsCurve&
 	int32 KnotCount = OpenNurbsCurve.KnotCount();
 	TArray<A3DDouble> NodalVector;
 	NodalVector.Reserve(KnotCount + 2);
-	NodalVector.Emplace(OpenNurbsCurve.SuperfluousKnot(0));
+	NodalVector.Emplace(OpenNurbsCurve.Knot(0));
 	for (int Index = 0; Index < KnotCount; ++Index)
 	{
 		NodalVector.Emplace(OpenNurbsCurve.Knot(Index));
 	}
-	NodalVector.Emplace(OpenNurbsCurve.SuperfluousKnot(1));
+	double LastValue = NodalVector[KnotCount];
+	NodalVector.Emplace(LastValue);
 
 	TArray<A3DVector3dData> ControlPointArray;
 	TArray<A3DDouble> WeightArray;
diff --git a/Engine/Plugins/Enterprise/DatasmithCADImporter/Source/DatasmithWireTranslator/Private/AliasModelToCADKernelConverter.cpp b/Engine/Plugins/Enterprise/DatasmithCADImporter/Source/DatasmithWireTranslator/Private/AliasModelToCADKernelConverter.cpp
--- a/Engine/Plugins/Enterprise/DatasmithCADImporter/Source/DatasmithWireTranslator/Private/AliasModelToCADKernelConverter.cpp
+++ b/Engine/Plugins/Enterprise/DatasmithCADImporter/Source/DatasmithWireTranslator/Private/AliasModelToCADKernelConverter.cpp
@@ -143,7 +143,7 @@ TSharedPtr<CADKernel::FTopologicalEdge> FAliasModelToCADKernelConverter::AddEdge
 	return Edge;
 }
 
-TSharedPtr<CADKernel::FTopologicalLoop> FAliasModelToCADKernelConverter::AddLoop(const AlTrimBoundary& TrimBoundary, TSharedPtr<CADKernel::FSurface>& CarrierSurface)
+TSharedPtr<CADKernel::FTopologicalLoop> FAliasModelToCADKernelConverter::AddLoop(const AlTrimBoundary& TrimBoundary, TSharedPtr<CADKernel::FSurface>& CarrierSurface, const bool bIsExternal)
 {
 	TArray<TSharedPtr<CADKernel::FTopologicalEdge>> Edges;
 	TArray<CADKernel::EOrientation> Directions;
@@ -163,7 +163,7 @@ TSharedPtr<CADKernel::FTopologicalLoop> FAliasModelToCADKernelConverter::AddLoop
 		return TSharedPtr<CADKernel::FTopologicalLoop>();
 	}
 
-	TSharedPtr<CADKernel::FTopologicalLoop> Loop = CADKernel::FTopologicalLoop::Make(Edges, Directions, GeometricTolerance);
+	TSharedPtr<CADKernel::FTopologicalLoop> Loop = CADKernel::FTopologicalLoop::Make(Edges, Directions, bIsExternal, GeometricTolerance);
 	return Loop;
 }
 
@@ -202,14 +202,16 @@ TSharedPtr<CADKernel::FTopologicalFace> FAliasModelToCADKernelConverter::AddTrim
 		return TSharedPtr<CADKernel::FTopologicalFace>();
 	}
 
+	bool bIsExternal = true;
 	TArray<TSharedPtr<CADKernel::FTopologicalLoop>> Loops;
 	for (TUniquePtr<AlTrimBoundary> TrimBoundary(TrimRegion.firstBoundary()); TrimBoundary.IsValid(); TrimBoundary = TUniquePtr<AlTrimBoundary>(TrimBoundary->nextBoundary()))
 	{
-		TSharedPtr<CADKernel::FTopologicalLoop> Loop = AddLoop(*TrimBoundary, Surface);
+		TSharedPtr<CADKernel::FTopologicalLoop> Loop = AddLoop(*TrimBoundary, Surface, bIsExternal);
 		if (Loop.IsValid())
 		{
 			LinkEdgesLoop(*TrimBoundary, *Loop);
 			Loops.Add(Loop);
+			bIsExternal = false;
 		}
 	}
 
diff --git a/Engine/Plugins/Enterprise/DatasmithCADImporter/Source/DatasmithWireTranslator/Private/AliasModelToCADKernelConverter.h b/Engine/Plugins/Enterprise/DatasmithCADImporter/Source/DatasmithWireTranslator/Private/AliasModelToCADKernelConverter.h
--- a/Engine/Plugins/Enterprise/DatasmithCADImporter/Source/DatasmithWireTranslator/Private/AliasModelToCADKernelConverter.h
+++ b/Engine/Plugins/Enterprise/DatasmithCADImporter/Source/DatasmithWireTranslator/Private/AliasModelToCADKernelConverter.h
@@ -45,7 +45,7 @@ public:
 protected:
 	TSharedPtr<CADKernel::FTopologicalEdge> AddEdge(const AlTrimCurve& TrimCurve, TSharedPtr<CADKernel::FSurface>& CarrierSurface);
 
-	TSharedPtr<CADKernel::FTopologicalLoop> AddLoop(const AlTrimBoundary& TrimBoundary, TSharedPtr<CADKernel::FSurface>& CarrierSurface);
+	TSharedPtr<CADKernel::FTopologicalLoop> AddLoop(const AlTrimBoundary& TrimBoundary, TSharedPtr<CADKernel::FSurface>& CarrierSurface, const bool bIsExternal);
 	
 	/**
 	 * Build face's links with its neighbor have to be done after the loop is finalize.
diff --git a/Engine/Source/Runtime/Datasmith/CADKernel/Private/CADKernel/Mesh/Structure/Grid.cpp b/Engine/Source/Runtime/Datasmith/CADKernel/Private/CADKernel/Mesh/Structure/Grid.cpp
--- a/Engine/Source/Runtime/Datasmith/CADKernel/Private/CADKernel/Mesh/Structure/Grid.cpp
+++ b/Engine/Source/Runtime/Datasmith/CADKernel/Private/CADKernel/Mesh/Structure/Grid.cpp
@@ -932,204 +932,222 @@ void SlightlyDisplacedPolyline(TArray<FPoint2D>& D2Points, const FSurfacicBounda
 	D2Points.Last() += Normal;
 }
 
-//#define DEBUG_GET_BOUNDARY_MESH
-bool FGrid::GetMeshOfLoops()
+void FGrid::GetMeshOfLoop(const FTopologicalLoop& Loop)
 {
-	int32 ThinZoneNum = 0;
-	if (Face.HasThinZone())
+	int32 LoopNodeCount = 0;
+
+	for (const FOrientedEdge& Edge : Loop.GetEdges())
 	{
-		ThinZoneNum = ThinZoneFinder.GetThinZones().Num();
+		LoopNodeCount += Edge.Entity->GetLinkActiveEdge()->GetMesh()->GetNodeCount() + 2;
 	}
 
-	int32 LoopCount = Face.GetLoops().Num();
-	FaceLoops2D[EGridSpace::Default2D].Reserve(LoopCount + ThinZoneNum);
+	TArray<FPoint2D>& Loop2D = FaceLoops2D[EGridSpace::Default2D].Emplace_GetRef();
+	Loop2D.Empty(LoopNodeCount);
 
-	FaceLoops3D.Reserve(LoopCount);
-	NormalsOfFaceLoops.Reserve(LoopCount);
-	NodeIdsOfFaceLoops.Reserve(LoopCount);
+	TArray<FPoint>& Loop3D = FaceLoops3D.Emplace_GetRef();
+	Loop3D.Reserve(LoopNodeCount);
 
-#ifdef DEBUG_GET_BOUNDARY_MESH
-	F3DDebugSession _(bDisplay, ("GetLoopMesh"));
-#endif
+	TArray<FVector>& LoopNormals = NormalsOfFaceLoops.Emplace_GetRef();
+	LoopNormals.Reserve(LoopNodeCount);
 
-	for (const TSharedPtr<FTopologicalLoop>& Loop : Face.GetLoops())
-	{
-		int32 LoopNodeCount = 0;
+	TArray<int32>& LoopIds = NodeIdsOfFaceLoops.Emplace_GetRef();
+	LoopIds.Reserve(LoopNodeCount);
 
-		for (const FOrientedEdge& Edge : Loop->GetEdges())
-		{
-			LoopNodeCount += Edge.Entity->GetLinkActiveEdge()->GetMesh()->GetNodeCount() + 2;
-		}
-
-		TArray<FPoint2D>& Loop2D = FaceLoops2D[EGridSpace::Default2D].Emplace_GetRef();
-		Loop2D.Empty(LoopNodeCount);
-
-		TArray<FPoint>& Loop3D = FaceLoops3D.Emplace_GetRef();
-		Loop3D.Reserve(LoopNodeCount);
-
-		TArray<FVector>& LoopNormals = NormalsOfFaceLoops.Emplace_GetRef();
-		LoopNormals.Reserve(LoopNodeCount);
+	for (const FOrientedEdge& OrientedEdge : Loop.GetEdges())
+	{
+		const TSharedPtr<FTopologicalEdge>& Edge = OrientedEdge.Entity;
+		const TSharedRef<FTopologicalEdge>& ActiveEdge = Edge->GetLinkActiveEdge();
 
-		TArray<int32>& LoopIds = NodeIdsOfFaceLoops.Emplace_GetRef();
-		LoopIds.Reserve(LoopNodeCount);
+		bool bSameDirection = Edge->IsSameDirection(*ActiveEdge);
 
-		for (const FOrientedEdge& OrientedEdge : Loop->GetEdges())
+		TArray<double> EdgeCuttingPointCoordinates;
 		{
-			const TSharedPtr<FTopologicalEdge>& Edge = OrientedEdge.Entity;
-			const TSharedRef<FTopologicalEdge>& ActiveEdge = Edge->GetLinkActiveEdge();
-
-			bool bSameDirection = Edge->IsSameDirection(*ActiveEdge);
-
-			TArray<double> EdgeCuttingPointCoordinates;
+			TArray<FCuttingPoint>& CuttingPoints = Edge->GetCuttingPoints();
+			if (!CuttingPoints.IsEmpty())
 			{
-				TArray<FCuttingPoint>& CuttingPoints = Edge->GetCuttingPoints();
-				if (!CuttingPoints.IsEmpty())
-				{
-					GetCuttingPointCoordinates(CuttingPoints, EdgeCuttingPointCoordinates);
-				}
+				GetCuttingPointCoordinates(CuttingPoints, EdgeCuttingPointCoordinates);
 			}
+		}
 
-			FSurfacicPolyline CuttingPolyline(true);
+		FSurfacicPolyline CuttingPolyline(true);
 
-			if (Edge->IsDegenerated())
+		if (Edge->IsDegenerated())
+		{
+			if (EdgeCuttingPointCoordinates.IsEmpty())
 			{
-				if (EdgeCuttingPointCoordinates.IsEmpty())
+				int32 CuttingPointCount = 2;
+				for (const FTopologicalEdge* TwinEdge : Edge->GetTwinEntities())
 				{
-					int32 CuttingPointCount = 2;
-					for (const FTopologicalEdge* TwinEdge : Edge->GetTwinEntities())
+					int32 TwinCuttingPointCount = TwinEdge->GetCuttingPoints().Num();
+					if (TwinCuttingPointCount > CuttingPointCount)
 					{
-						int32 TwinCuttingPointCount = TwinEdge->GetCuttingPoints().Num();
-						if (TwinCuttingPointCount > CuttingPointCount)
-						{
-							CuttingPointCount = TwinCuttingPointCount;
-						}
+						CuttingPointCount = TwinCuttingPointCount;
 					}
-					FMesherTools::FillCuttingPointCoordinates(Edge->GetBoundary(), CuttingPointCount, EdgeCuttingPointCoordinates);
 				}
+				FMesherTools::FillCuttingPointCoordinates(Edge->GetBoundary(), CuttingPointCount, EdgeCuttingPointCoordinates);
+			}
 
-				Swap(CuttingPolyline.Coordinates, EdgeCuttingPointCoordinates);
-				Edge->Approximate2DPoints(CuttingPolyline.Coordinates, CuttingPolyline.Points2D);
+			Swap(CuttingPolyline.Coordinates, EdgeCuttingPointCoordinates);
+			Edge->Approximate2DPoints(CuttingPolyline.Coordinates, CuttingPolyline.Points2D);
 
-				CuttingPolyline.Points3D.Init(ActiveEdge->GetStartBarycenter(), CuttingPolyline.Coordinates.Num());
+			CuttingPolyline.Points3D.Init(ActiveEdge->GetStartBarycenter(), CuttingPolyline.Coordinates.Num());
 
-				TArray<FPoint2D> D2Points = CuttingPolyline.Points2D;
-				const FSurfacicBoundary& Boundary = Edge->GetCurve()->GetCarrierSurface()->GetBoundary();
-				// to compute the normals, the 2D points are slightly displaced perpendicular to the curve
-				SlightlyDisplacedPolyline(D2Points, Boundary);
-				Edge->GetCurve()->GetCarrierSurface()->EvaluateNormals(D2Points, CuttingPolyline.Normals);
-			}
-			else
+			TArray<FPoint2D> D2Points = CuttingPolyline.Points2D;
+			const FSurfacicBoundary& Boundary = Edge->GetCurve()->GetCarrierSurface()->GetBoundary();
+			// to compute the normals, the 2D points are slightly displaced perpendicular to the curve
+			SlightlyDisplacedPolyline(D2Points, Boundary);
+			Edge->GetCurve()->GetCarrierSurface()->EvaluateNormals(D2Points, CuttingPolyline.Normals);
+		}
+		else
+		{
+			if (EdgeCuttingPointCoordinates.IsEmpty())
 			{
-				if (EdgeCuttingPointCoordinates.IsEmpty())
-				{
-					TArray<FPoint>& MeshVertex3D = ActiveEdge->GetMesh()->GetNodeCoordinates();
-					TArray<double> ProjectedPointCoords;
+				TArray<FPoint>& MeshVertex3D = ActiveEdge->GetMesh()->GetNodeCoordinates();
+				TArray<double> ProjectedPointCoords;
 
-					CuttingPolyline.Coordinates.Reserve(MeshVertex3D.Num() + 2);
-					Edge->ProjectTwinEdgePoints(MeshVertex3D, bSameDirection, CuttingPolyline.Coordinates);
-					CuttingPolyline.Coordinates.Insert(Edge->GetBoundary().GetMin(), 0);
-					CuttingPolyline.Coordinates.Add(Edge->GetBoundary().GetMax());
+				CuttingPolyline.Coordinates.Reserve(MeshVertex3D.Num() + 2);
+				Edge->ProjectTwinEdgePoints(MeshVertex3D, bSameDirection, CuttingPolyline.Coordinates);
+				CuttingPolyline.Coordinates.Insert(Edge->GetBoundary().GetMin(), 0);
+				CuttingPolyline.Coordinates.Add(Edge->GetBoundary().GetMax());
 
-					// check if there are coincident coordinates
-					bool bProjectionFailed = false;
-					for (int32 Index = 0; Index < CuttingPolyline.Coordinates.Num() - 1;)
-					{
-						double Diff = CuttingPolyline.Coordinates[Index++];
-						Diff -= CuttingPolyline.Coordinates[Index];
-						if (Diff > -SMALL_NUMBER)
-						{
-							bProjectionFailed = true;
-							break;
-						}
-					}
-					if(bProjectionFailed)
+				// check if there are coincident coordinates
+				bool bProjectionFailed = false;
+				for (int32 Index = 0; Index < CuttingPolyline.Coordinates.Num() - 1;)
+				{
+					double Diff = CuttingPolyline.Coordinates[Index++];
+					Diff -= CuttingPolyline.Coordinates[Index];
+					if (Diff > -SMALL_NUMBER)
 					{
-						FMesherTools::FillCuttingPointCoordinates(Edge->GetBoundary(), MeshVertex3D.Num() + 2, CuttingPolyline.Coordinates);
+						bProjectionFailed = true;
+						break;
 					}
 				}
-				else
+				if (bProjectionFailed)
 				{
-					Swap(CuttingPolyline.Coordinates, EdgeCuttingPointCoordinates);
+					FMesherTools::FillCuttingPointCoordinates(Edge->GetBoundary(), MeshVertex3D.Num() + 2, CuttingPolyline.Coordinates);
 				}
-
-				Edge->ApproximatePolyline(CuttingPolyline);
 			}
-
-			TArray<int32> EdgeVerticesIndex;
-			if (Edge->IsDegenerated())
+			else
 			{
-				EdgeVerticesIndex.Init(ActiveEdge->GetStartVertex()->GetLinkActiveEntity()->GetOrCreateMesh(MeshModel)->GetMesh(), CuttingPolyline.Coordinates.Num());
+				Swap(CuttingPolyline.Coordinates, EdgeCuttingPointCoordinates);
 			}
-			else if (Edge->IsVirtuallyMeshed())
+
+			Edge->ApproximatePolyline(CuttingPolyline);
+		}
+
+		TArray<int32> EdgeVerticesIndex;
+		if (Edge->IsDegenerated())
+		{
+			EdgeVerticesIndex.Init(ActiveEdge->GetStartVertex()->GetLinkActiveEntity()->GetOrCreateMesh(MeshModel)->GetMesh(), CuttingPolyline.Coordinates.Num());
+		}
+		else if (Edge->IsVirtuallyMeshed())
+		{
+			// @see FParametricMesher::Mesh(FTopologicalEdge& InEdge, FTopologicalFace& Face)
+			int32 NodeCount = CuttingPolyline.Coordinates.Num();
+			EdgeVerticesIndex.Reserve(NodeCount);
+			int32 MiddleNodeIndex = NodeCount / 2;
+			int32 Index = 0;
 			{
-				// @see FParametricMesher::Mesh(FTopologicalEdge& InEdge, FTopologicalFace& Face)
-				int32 NodeCount = CuttingPolyline.Coordinates.Num();
-				EdgeVerticesIndex.Reserve(NodeCount);
-				int32 MiddleNodeIndex = NodeCount / 2;
-				int32 Index = 0;
+				int32 StartVertexMeshId = Edge->GetStartVertex()->GetLinkActiveEntity()->GetOrCreateMesh(MeshModel)->GetMesh();
+				for (; Index < MiddleNodeIndex; ++Index)
 				{
-					int32 StartVertexMeshId = Edge->GetStartVertex()->GetLinkActiveEntity()->GetOrCreateMesh(MeshModel)->GetMesh();
-					for (; Index < MiddleNodeIndex; ++Index)
-					{
-						EdgeVerticesIndex.Add(StartVertexMeshId);
-					}
-				}
-				{
-					int32 EndVertexMeshId = Edge->GetEndVertex()->GetLinkActiveEntity()->GetOrCreateMesh(MeshModel)->GetMesh();
-					for (; Index < NodeCount; ++Index)
-					{
-						EdgeVerticesIndex.Add(EndVertexMeshId);
-					}
+					EdgeVerticesIndex.Add(StartVertexMeshId);
 				}
 			}
-			else
 			{
-				EdgeVerticesIndex = ((TSharedRef<FEdgeMesh>)ActiveEdge->GetOrCreateMesh(MeshModel))->EdgeVerticesIndex;
+				int32 EndVertexMeshId = Edge->GetEndVertex()->GetLinkActiveEntity()->GetOrCreateMesh(MeshModel)->GetMesh();
+				for (; Index < NodeCount; ++Index)
+				{
+					EdgeVerticesIndex.Add(EndVertexMeshId);
+				}
 			}
+		}
+		else
+		{
+			EdgeVerticesIndex = ((TSharedRef<FEdgeMesh>)ActiveEdge->GetOrCreateMesh(MeshModel))->EdgeVerticesIndex;
+		}
 
 #ifdef DEBUG_GET_BOUNDARY_MESH
-			if (bDisplay)
+		if (bDisplay)
+		{
+			F3DDebugSession _(bDisplay, FString::Printf(TEXT("Edge %d cutting points on surface"), ActiveEdge->GetId()));
+			for (const FPoint2D& Point2D : CuttingPolyline.Points2D)
 			{
-				F3DDebugSession _(bDisplay, FString::Printf(TEXT("Edge %d cutting points on surface"), ActiveEdge->GetId()));
-				for (const FPoint2D& Point2D : CuttingPolyline.Points2D)
-				{
-					DisplayPoint(Point2D);
-				}
-				//Wait();
+				DisplayPoint(Point2D);
 			}
+			//Wait();
+		}
 #endif
 
-			if (OrientedEdge.Direction != EOrientation::Front)
-			{
-				CuttingPolyline.Reverse();
-			}
+		if (OrientedEdge.Direction != EOrientation::Front)
+		{
+			CuttingPolyline.Reverse();
+		}
 
-			if (bSameDirection != (OrientedEdge.Direction == EOrientation::Front))
-			{
-				Algo::Reverse(EdgeVerticesIndex);
-			}
+		if (bSameDirection != (OrientedEdge.Direction == EOrientation::Front))
+		{
+			Algo::Reverse(EdgeVerticesIndex);
+		}
+
+		ensureCADKernel(CuttingPolyline.Size() > 1);
+
+		Loop2D.Append(CuttingPolyline.Points2D);
+		Loop2D.Pop();
+
+		Loop3D.Emplace(ActiveEdge->GetStartVertex((OrientedEdge.Direction == EOrientation::Front) == bSameDirection)->GetLinkActiveEntity()->GetBarycenter());
+		Loop3D.Append(CuttingPolyline.Points3D.GetData() + 1, CuttingPolyline.Points3D.Num() - 2);
+
+		LoopNormals.Append(CuttingPolyline.Normals);
+		LoopNormals.Pop();
+
+		LoopIds.Append(EdgeVerticesIndex);
+		LoopIds.Pop();
+	}
+
+	if (Loop2D.Num() < 3) // degenerated loop
+	{
+		FaceLoops2D[EGridSpace::Default2D].Pop();
+		FaceLoops3D.Pop();
+		NormalsOfFaceLoops.Pop();
+	}
+}
 
-			ensureCADKernel(CuttingPolyline.Size() > 1);
 
-			Loop2D.Append(CuttingPolyline.Points2D);
-			Loop2D.Pop();
+//#define DEBUG_GET_BOUNDARY_MESH
+bool FGrid::GetMeshOfLoops()
+{
+	int32 ThinZoneNum = 0;
+	if (Face.HasThinZone())
+	{
+		ThinZoneNum = ThinZoneFinder.GetThinZones().Num();
+	}
+
+	int32 LoopCount = Face.GetLoops().Num();
+	FaceLoops2D[EGridSpace::Default2D].Reserve(LoopCount + ThinZoneNum);
 
-			Loop3D.Emplace(ActiveEdge->GetStartVertex((OrientedEdge.Direction == EOrientation::Front) == bSameDirection)->GetLinkActiveEntity()->GetBarycenter());
-			Loop3D.Append(CuttingPolyline.Points3D.GetData() + 1, CuttingPolyline.Points3D.Num() - 2);
+	FaceLoops3D.Reserve(LoopCount);
+	NormalsOfFaceLoops.Reserve(LoopCount);
+	NodeIdsOfFaceLoops.Reserve(LoopCount);
 
-			LoopNormals.Append(CuttingPolyline.Normals);
-			LoopNormals.Pop();
+#ifdef DEBUG_GET_BOUNDARY_MESH
+	F3DDebugSession _(bDisplay, ("GetLoopMesh"));
+#endif
 
-			LoopIds.Append(EdgeVerticesIndex);
-			LoopIds.Pop();
+	// Outer loops is processed first
+	for (const TSharedPtr<FTopologicalLoop>& Loop : Face.GetLoops())
+	{
+		if (Loop->IsExternal())
+		{
+			GetMeshOfLoop(*Loop);
+			break;
 		}
+	}
 
-		if (Loop2D.Num() < 3) // degenerated loop
+	for (const TSharedPtr<FTopologicalLoop>& Loop : Face.GetLoops())
+	{
+		if (!Loop->IsExternal())
 		{
-			FaceLoops2D[EGridSpace::Default2D].Pop();
-			FaceLoops3D.Pop();
-			NormalsOfFaceLoops.Pop();
-			continue;
+			GetMeshOfLoop(*Loop);
 		}
 	}
 
diff --git a/Engine/Source/Runtime/Datasmith/CADKernel/Private/CADKernel/Mesh/Structure/Grid.h b/Engine/Source/Runtime/Datasmith/CADKernel/Private/CADKernel/Mesh/Structure/Grid.h
--- a/Engine/Source/Runtime/Datasmith/CADKernel/Private/CADKernel/Mesh/Structure/Grid.h
+++ b/Engine/Source/Runtime/Datasmith/CADKernel/Private/CADKernel/Mesh/Structure/Grid.h
@@ -277,6 +277,7 @@ protected:
 	 * @ return false if the mesh of the loop is degenerated
 	 */
 	bool GetMeshOfLoops();
+	void GetMeshOfLoop(const FTopologicalLoop& Loop);
 
 public:
 
diff --git a/Engine/Source/Runtime/Datasmith/CADKernel/Private/CADKernel/Topo/TopologicalEdge.cpp b/Engine/Source/Runtime/Datasmith/CADKernel/Private/CADKernel/Topo/TopologicalEdge.cpp
--- a/Engine/Source/Runtime/Datasmith/CADKernel/Private/CADKernel/Topo/TopologicalEdge.cpp
+++ b/Engine/Source/Runtime/Datasmith/CADKernel/Private/CADKernel/Topo/TopologicalEdge.cpp
@@ -518,22 +518,41 @@ bool FTopologicalEdge::IsSameDirection(const FTopologicalEdge& Edge) const
 		return true;
 	}
 
-	TSharedPtr<const FVertexLink> vertex1Edge = GetStartVertex()->GetLink();
-	TSharedPtr<const FVertexLink> vertex2Edge = GetEndVertex()->GetLink();
+	TSharedPtr<const FVertexLink> Vertex1Edge = GetStartVertex()->GetLink();
+	TSharedPtr<const FVertexLink> Vertex2Edge = GetEndVertex()->GetLink();
 
-	if (vertex1Edge == vertex2Edge)
+	if (Vertex1Edge == Vertex2Edge)
 	{
 		if (Edge.IsDegenerated())
 		{
 			return true;
 		}
+
+		FPoint EdgeStartTangent;
+		FPoint EdgeEndTangent;
+		Edge.GetTangentsAtExtremities(EdgeStartTangent, EdgeEndTangent, true);
+
+		FPoint StartTangent;
+		FPoint EndTangent;
+		GetTangentsAtExtremities(StartTangent, EndTangent, true);
+
+		double StartAngle = StartTangent.ComputeCosinus(EdgeStartTangent);
+		double EndAngle = EndTangent.ComputeCosinus(EdgeEndTangent);
+
+		if (StartAngle >= 0 && EndAngle >= 0)
+		{
+			return true;
+		}
+		if (StartAngle <= 0 && EndAngle <= 0)
+		{
+			return false;
+		}
+
 		Edge.SetAsDegenerated();
 		return true;
-
-		// Todo: Two cycles can have connected extremities but not with the same orientation
 	}
 
-	return vertex1Edge == Edge.GetStartVertex()->GetLink();
+	return Vertex1Edge == Edge.GetStartVertex()->GetLink();
 }
 
 #ifdef CADKERNEL_DEV
diff --git a/Engine/Source/Runtime/Datasmith/CADKernel/Private/CADKernel/Topo/TopologicalFace.cpp b/Engine/Source/Runtime/Datasmith/CADKernel/Private/CADKernel/Topo/TopologicalFace.cpp
--- a/Engine/Source/Runtime/Datasmith/CADKernel/Private/CADKernel/Topo/TopologicalFace.cpp
+++ b/Engine/Source/Runtime/Datasmith/CADKernel/Private/CADKernel/Topo/TopologicalFace.cpp
@@ -217,7 +217,8 @@ void FTopologicalFace::ApplyNaturalLoops(const FSurfacicBoundary& Boundaries)
 	TArray<EOrientation> Orientations;
 	Orientations.Init(EOrientation::Front, Edges.Num());
 
-	TSharedPtr<FTopologicalLoop> Loop = FTopologicalLoop::Make(Edges, Orientations, CarrierSurface->Get3DTolerance());
+	const bool bIsExternalLoop = true;
+	TSharedPtr<FTopologicalLoop> Loop = FTopologicalLoop::Make(Edges, Orientations, bIsExternalLoop, CarrierSurface->Get3DTolerance());
 	AddLoop(Loop);
 }
 
@@ -240,10 +241,6 @@ void FTopologicalFace::AddLoops(const TArray<TSharedPtr<FTopologicalLoop>>& InLo
 void FTopologicalFace::AddLoop(const TSharedPtr<FTopologicalLoop>& InLoop)
 {
 	InLoop->SetSurface(this);
-	if (Loops.Num() > 0)
-	{
-		InLoop->SetAsInnerBoundary();
-	}
 	Loops.Add(InLoop);
 }
 
diff --git a/Engine/Source/Runtime/Datasmith/CADKernel/Private/CADKernel/Topo/TopologicalLoop.cpp b/Engine/Source/Runtime/Datasmith/CADKernel/Private/CADKernel/Topo/TopologicalLoop.cpp
--- a/Engine/Source/Runtime/Datasmith/CADKernel/Private/CADKernel/Topo/TopologicalLoop.cpp
+++ b/Engine/Source/Runtime/Datasmith/CADKernel/Private/CADKernel/Topo/TopologicalLoop.cpp
@@ -17,9 +17,9 @@
 namespace CADKernel
 {
 
-TSharedPtr<FTopologicalLoop> FTopologicalLoop::Make(const TArray<TSharedPtr<FTopologicalEdge>>& InEdges, const TArray<EOrientation>& InEdgeDirections, double GeometricTolerance)
+TSharedPtr<FTopologicalLoop> FTopologicalLoop::Make(const TArray<TSharedPtr<FTopologicalEdge>>& InEdges, const TArray<EOrientation>& InEdgeDirections, const bool bIsExternalLoop, double GeometricTolerance)
 {
-	TSharedRef<FTopologicalLoop> LoopRef = FEntity::MakeShared<FTopologicalLoop>(InEdges, InEdgeDirections);
+	TSharedRef<FTopologicalLoop> LoopRef = FEntity::MakeShared<FTopologicalLoop>(InEdges, InEdgeDirections, bIsExternalLoop);
 	FTopologicalLoop& Loop = *LoopRef;
 
 	Loop.EnsureLogicalClosing(GeometricTolerance);
@@ -57,9 +57,9 @@ TSharedPtr<FTopologicalLoop> FTopologicalLoop::Make(const TArray<TSharedPtr<FTop
 	return LoopRef;
 }
 
-FTopologicalLoop::FTopologicalLoop(const TArray<TSharedPtr<FTopologicalEdge>>& InEdges, const TArray<EOrientation>& InEdgeDirections)
+FTopologicalLoop::FTopologicalLoop(const TArray<TSharedPtr<FTopologicalEdge>>& InEdges, const TArray<EOrientation>& InEdgeDirections, const bool bIsExternalLoop)
 	: Face(nullptr)
-	, bExternalLoop(true)
+	, bIsExternal(bIsExternalLoop)
 {
 	Edges.Reserve(InEdges.Num());
 	for (int32 Index = 0; Index < InEdges.Num(); ++Index)
@@ -466,7 +466,7 @@ bool FTopologicalLoop::Orient()
 		FMessage::Printf(Log, TEXT("WARNING: Loop Orientation of surface %d is doubtful\n"), Face->GetId());
 	}
 
-	if ((WrongOrientationCount > GoodOrientationCount) == bExternalLoop)
+	if ((WrongOrientationCount > GoodOrientationCount) == bIsExternal)
 	{
 		SwapOrientation();
 	}
@@ -959,6 +959,7 @@ FInfoEntity& FTopologicalLoop::GetInfo(FInfoEntity& Info) const
 {
 	return FEntity::GetInfo(Info)
 		.Add(TEXT("Edges"), (TArray<TOrientedEntity<FEntity>>&) Edges)
+		.Add(TEXT("IsExternal"), bIsExternal)
 		.Add(TEXT("Hosted by"), Face);
 }
 #endif
diff --git a/Engine/Source/Runtime/Datasmith/CADKernel/Public/CADKernel/Topo/TopologicalLoop.h b/Engine/Source/Runtime/Datasmith/CADKernel/Public/CADKernel/Topo/TopologicalLoop.h
--- a/Engine/Source/Runtime/Datasmith/CADKernel/Public/CADKernel/Topo/TopologicalLoop.h
+++ b/Engine/Source/Runtime/Datasmith/CADKernel/Public/CADKernel/Topo/TopologicalLoop.h
@@ -53,9 +53,9 @@ protected:
 	TArray<FOrientedEdge> Edges;
 
 	FTopologicalFace* Face;
-	bool bExternalLoop;
+	bool bIsExternal;
 
-	FTopologicalLoop(const TArray<TSharedPtr<FTopologicalEdge>>& Edges, const TArray<EOrientation>& EdgeDirections);
+	FTopologicalLoop(const TArray<TSharedPtr<FTopologicalEdge>>& Edges, const TArray<EOrientation>& EdgeDirections, const bool bIsEternalLoop);
 
 	FTopologicalLoop() = default;
 
@@ -75,7 +75,7 @@ public:
 
 	~FTopologicalLoop() = default;
 
-	static TSharedPtr<FTopologicalLoop> Make(const TArray<TSharedPtr<FTopologicalEdge>>& EdgeList, const TArray<EOrientation>& EdgeDirections, const double GeometricTolerance);
+	static TSharedPtr<FTopologicalLoop> Make(const TArray<TSharedPtr<FTopologicalEdge>>& EdgeList, const TArray<EOrientation>& EdgeDirections, const bool bIsExternalLoop, const double GeometricTolerance);
 
 	void DeleteLoopEdges();
 
@@ -84,7 +84,7 @@ public:
 		FTopologicalEntity::Serialize(Ar);
 		SerializeIdents(Ar, (TArray<TOrientedEntity<FEntity>>&) Edges);
 		SerializeIdent(Ar, &Face);
-		Ar << bExternalLoop;
+		Ar << bIsExternal;
 	}
 
 	virtual void SpawnIdent(FDatabase& Database) override
@@ -161,17 +161,17 @@ public:
 		return Face;
 	}
 
+	bool IsExternal() const 
+	{
+		return bIsExternal;
+	}
+
 	/*
 	 * @return false if the orientation is doubtful
 	 */
 	bool Orient();
 	void SwapOrientation();
 
-	void SetAsInnerBoundary()
-	{
-		bExternalLoop = false;
-	}
-
 	void ReplaceEdge(TSharedPtr<FTopologicalEdge>& OldEdge, TSharedPtr<FTopologicalEdge>& NewEdge);
 	void ReplaceEdge(TSharedPtr<FTopologicalEdge>& Edge, TArray<TSharedPtr<FTopologicalEdge>>& NewEdges);
 	void ReplaceEdges(TArray<FOrientedEdge>& Candidates, TSharedPtr<FTopologicalEdge>& NewEdge);
