Removed unnecessary check on file existence in the Datasmith XML reader following changes made in CL 19422918

diff --git a/Engine/Source/Runtime/Datasmith/DatasmithCore/Private/DatasmithSceneXmlReader.cpp b/Engine/Source/Runtime/Datasmith/DatasmithCore/Private/DatasmithSceneXmlReader.cpp
--- a/Engine/Source/Runtime/Datasmith/DatasmithCore/Private/DatasmithSceneXmlReader.cpp
+++ b/Engine/Source/Runtime/Datasmith/DatasmithCore/Private/DatasmithSceneXmlReader.cpp
@@ -1229,26 +1229,23 @@ bool FDatasmithSceneXmlReader::ParseXmlFile(TSharedRef< IDatasmithScene >& OutSc
 			ParseTextureElement(Nodes[i], Element);
 
 			FString TextureName(Element->GetName());
-			if (FPaths::FileExists(Element->GetFile()))
+			int32 TexturesCount = OutScene->GetTexturesCount();
+			bool bIsDuplicate = false;
+			for (int32 t = 0; t < TexturesCount; t++)
 			{
-				int32 TexturesCount = OutScene->GetTexturesCount();
-				bool bIsDuplicate = false;
-				for (int32 t = 0; t < TexturesCount; t++)
+				const TSharedPtr< IDatasmithTextureElement >& TextureElement = OutScene->GetTexture(t);
+				if (TextureName == TextureElement->GetName())
 				{
-					const TSharedPtr< IDatasmithTextureElement >& TextureElement = OutScene->GetTexture(t);
-					if (TextureName == TextureElement->GetName())
-					{
-						bIsDuplicate = true;
-						break;
-					}
+					bIsDuplicate = true;
+					break;
 				}
+			}
 
-				if (bIsDuplicate == false)
-				{
-					OutScene->AddTexture(Element);
+			if (bIsDuplicate == false)
+			{
+				OutScene->AddTexture(Element);
 
-					Objects.Add( Element->GetName(), Element );
-				}
+				Objects.Add(Element->GetName(), Element);
 			}
 		}
 		//READ ENVIRONMENTS
@@ -1384,8 +1381,8 @@ bool FDatasmithSceneXmlReader::ParseXmlFile(TSharedRef< IDatasmithScene >& OutSc
 #ifdef USE_LOCALE
 	if (Locale)
 	{
-		::uselocale(PreviousLocale);
-		::freelocale(Locale);
+	::uselocale(PreviousLocale);
+	::freelocale(Locale);
 	}
 #endif
 
