DirectLink: Enable LWC types serialization as doubles

diff --git a/Engine/Source/Runtime/Datasmith/DirectLink/Public/DirectLinkSerialMethods.h b/Engine/Source/Runtime/Datasmith/DirectLink/Public/DirectLinkSerialMethods.h
--- a/Engine/Source/Runtime/Datasmith/DirectLink/Public/DirectLinkSerialMethods.h
+++ b/Engine/Source/Runtime/Datasmith/DirectLink/Public/DirectLinkSerialMethods.h
@@ -10,13 +10,6 @@
 #include "Serialization/Archive.h"
 
 
-// LWC initiative changes the support of some geometry-related types from float to doubles.
-// As we want to be able to send data to apps built without LWC support, we send data with
-// the legacy float serialization.
-// When the support of pre-LWC products is no longer required, we can safely toggle this
-// macro off and send LWC types natively (as doubles).
-#define DL_SEND_LWC_TYPES_AS_F32 1
-
 namespace Reflect
 {
 
@@ -77,7 +70,7 @@ enum class ESerialMethod : uint8
 static constexpr EStoreType GetStoreType(ESerialMethod Method) { return EStoreType(uint8(Method) & _storeTypeMask); }
 
 // This list maps which serialization method should be used for each Type reflected by directlink
-template<typename T> struct TDefaultSerialMethod;
+template<typename T> struct TDefaultSerialMethod           { constexpr static ESerialMethod Value = ESerialMethod::_NotImplementedYet; };
 template<> struct TDefaultSerialMethod<bool>               { constexpr static ESerialMethod Value = ESerialMethod::Uint8_Default;      };
 template<> struct TDefaultSerialMethod<uint8>              { constexpr static ESerialMethod Value = ESerialMethod::Uint8_Default;      };
 template<> struct TDefaultSerialMethod<int32>              { constexpr static ESerialMethod Value = ESerialMethod::Int32_Default;      };
@@ -93,14 +86,8 @@ template<> struct TDefaultSerialMethod<FLinearColor>       { constexpr static ES
 template<> struct TDefaultSerialMethod<FMD5Hash>           { constexpr static ESerialMethod Value = ESerialMethod::MD5Hash_Default;    };
 template<> struct TDefaultSerialMethod<FTransform>         { constexpr static ESerialMethod Value = ESerialMethod::Transform_Default;  };
 template<> struct TDefaultSerialMethod<TArray<FTransform>> { constexpr static ESerialMethod Value = ESerialMethod::Transform_Array;    };
-#if DL_SEND_LWC_TYPES_AS_F32
-// for compatibility with pre-LWC endpoints, we send LWC-aware types as their f32 variants
-template<> struct TDefaultSerialMethod<FVector>            { constexpr static ESerialMethod Value = ESerialMethod::Vector_f32;         };
-template<> struct TDefaultSerialMethod<FQuat>              { constexpr static ESerialMethod Value = ESerialMethod::Quat_f32;           };
-#else
 template<> struct TDefaultSerialMethod<FVector>            { constexpr static ESerialMethod Value = ESerialMethod::Vector_f64;         };
-template<> struct TDefaultSerialMethod<FQuat>>             { constexpr static ESerialMethod Value = ESerialMethod::Quat_f64;           };
-#endif
+template<> struct TDefaultSerialMethod<FQuat>              { constexpr static ESerialMethod Value = ESerialMethod::Quat_f64;           };
 
 
 template<typename T> EStoreType GetStoreTypeForType() { return GetStoreType(TDefaultSerialMethod<T>::Value); }
