Datasmith: fix mesh area computation

diff --git a/Engine/Source/Runtime/Datasmith/DatasmithCore/Private/DatasmithMesh.cpp b/Engine/Source/Runtime/Datasmith/DatasmithCore/Private/DatasmithMesh.cpp
--- a/Engine/Source/Runtime/Datasmith/DatasmithCore/Private/DatasmithMesh.cpp
+++ b/Engine/Source/Runtime/Datasmith/DatasmithCore/Private/DatasmithMesh.cpp
@@ -454,7 +454,10 @@ float FDatasmithMesh::ComputeArea() const
 
 	for (int32 i = 0; i < NumFaces; i++)
 	{
-		Area += FDatasmithUtils::AreaTriangle3D( Impl->Vertices[ Impl->Indices[i + 0] ], Impl->Vertices[ Impl->Indices[i + 1] ], Impl->Vertices[ Impl->Indices[i + 2] ] );
+		Area += FDatasmithUtils::AreaTriangle3D(
+			Impl->Vertices[ Impl->Indices[3 * i + 0] ],
+			Impl->Vertices[ Impl->Indices[3 * i + 1] ],
+			Impl->Vertices[ Impl->Indices[3 * i + 2] ]);
 	}
 
 	return Area;
