Fix build warning about parameter name in documentation

diff --git a/Engine/Source/Runtime/Datasmith/DatasmithCore/Public/DatasmithUtils.h b/Engine/Source/Runtime/Datasmith/DatasmithCore/Public/DatasmithUtils.h
--- a/Engine/Source/Runtime/Datasmith/DatasmithCore/Public/DatasmithUtils.h
+++ b/Engine/Source/Runtime/Datasmith/DatasmithCore/Public/DatasmithUtils.h
@@ -146,17 +146,10 @@ public:
 	 */
 	static bool IsUVChannelValid(const FDatasmithMesh& DsMesh, const int32 Channel);
 
-	/**
-	 * Generate simple UV data at channel 0.
-	 *
-	 * @param DatasmithMesh    The DatasmithMesh in which the UV data will be created.
-	 */
-	 // static void CreateDefaultUVs(FDatasmithMesh& Mesh);
-
 	/**
 	 * Generate simple UV data at channel 0 for the base mesh and it's various LOD variants.
 	 *
-	 * @param DatasmithMesh    The DatasmithMesh in which the UV data will be created.
+	 * @param Mesh    The DatasmithMesh in which the UV data will be created.
 	 */
 	static void CreateDefaultUVsWithLOD(FDatasmithMesh& Mesh);
 
