Datasmith: Large World Coordinate support
rules of thumb:
- Actor transforms are in doubles
- mesh internal data are in floats, except UVs that are in doubles
- compiled 3dsmax, archicad, navisworks, revit, rhino, sketchup, solidworks

diff --git a/Engine/Plugins/Enterprise/DatasmithImporter/Source/DatasmithImporter/Private/DatasmithStaticMeshImporter.cpp b/Engine/Plugins/Enterprise/DatasmithImporter/Source/DatasmithImporter/Private/DatasmithStaticMeshImporter.cpp
--- a/Engine/Plugins/Enterprise/DatasmithImporter/Source/DatasmithImporter/Private/DatasmithStaticMeshImporter.cpp
+++ b/Engine/Plugins/Enterprise/DatasmithImporter/Source/DatasmithImporter/Private/DatasmithStaticMeshImporter.cpp
@@ -256,7 +256,7 @@ UStaticMesh* FDatasmithStaticMeshImporter::ImportStaticMesh(const TSharedRef< ID
 	}
 
 	// 4. Collisions
-	TArray< FVector > VertexPositions;
+	TArray< FVector3f > VertexPositions;
 	DatasmithMeshHelper::ExtractVertexPositions(Payload.CollisionMesh, VertexPositions);
 	if ( VertexPositions.Num() == 0 )
 	{
@@ -318,7 +318,7 @@ bool FDatasmithStaticMeshImporter::PreBuildStaticMesh( UStaticMesh* StaticMesh )
 			DatasmithMeshHelper::CreateDefaultUVs(MeshDescription);
 		}
 
-		if (DatasmithMeshHelper::IsMeshValid(MeshDescription, BuildSettings.BuildScale3D))
+		if (DatasmithMeshHelper::IsMeshValid(MeshDescription, FVector3f(BuildSettings.BuildScale3D)))
 		{
 			if (BuildSettings.bGenerateLightmapUVs)
 			{
@@ -361,7 +361,7 @@ void FDatasmithStaticMeshImporter::BuildStaticMeshes(const TArray< UStaticMesh*
 	UStaticMesh::BatchBuild(StaticMeshes, true, ProgressCallback);
 }
 
-void FDatasmithStaticMeshImporter::ProcessCollision(UStaticMesh* StaticMesh, const TArray< FVector >& VertexPositions)
+void FDatasmithStaticMeshImporter::ProcessCollision(UStaticMesh* StaticMesh, const TArray< FVector3f >& VertexPositions)
 {
 	// The following code is copied from StaticMeshEdit AddConvexGeomFromVertices (inaccessible outside UnrealEd)
 	if ( !StaticMesh )
@@ -382,9 +382,9 @@ void FDatasmithStaticMeshImporter::ProcessCollision(UStaticMesh* StaticMesh, con
 	FKConvexElem& ConvexElem = AggGeom.ConvexElems.AddDefaulted_GetRef();
 
 	ConvexElem.VertexData.Reserve( VertexPositions.Num() );
-	for ( const FVector& Position : VertexPositions )
+	for ( const FVector3f& Position : VertexPositions )
 	{
-		ConvexElem.VertexData.Add( Position );
+		ConvexElem.VertexData.Add( FVector(Position) );
 	}
 
 	ConvexElem.UpdateElemBox();
diff --git a/Engine/Plugins/Enterprise/DatasmithImporter/Source/DatasmithImporter/Public/DatasmithStaticMeshImporter.h b/Engine/Plugins/Enterprise/DatasmithImporter/Source/DatasmithImporter/Public/DatasmithStaticMeshImporter.h
--- a/Engine/Plugins/Enterprise/DatasmithImporter/Source/DatasmithImporter/Public/DatasmithStaticMeshImporter.h
+++ b/Engine/Plugins/Enterprise/DatasmithImporter/Source/DatasmithImporter/Public/DatasmithStaticMeshImporter.h
@@ -67,7 +67,7 @@ public:
 	static void BuildStaticMeshes(const TArray< UStaticMesh* >& StaticMeshes, TFunction<bool(UStaticMesh*)> ProgressCallback = nullptr);
 
 private:
-	static void ProcessCollision( UStaticMesh* StaticMesh, const TArray< FVector >& VertexPositions );
+	static void ProcessCollision( UStaticMesh* StaticMesh, const TArray< FVector3f >& VertexPositions );
 
 	/**
 	 * Applies the UMaterialInterfaces related to the IDatasmithMeshElement to the UStaticMesh.
diff --git a/Engine/Plugins/Enterprise/DatasmithImporter/Source/DatasmithTranslator/Private/Utility/DatasmithMeshHelper.cpp b/Engine/Plugins/Enterprise/DatasmithImporter/Source/DatasmithTranslator/Private/Utility/DatasmithMeshHelper.cpp
--- a/Engine/Plugins/Enterprise/DatasmithImporter/Source/DatasmithTranslator/Private/Utility/DatasmithMeshHelper.cpp
+++ b/Engine/Plugins/Enterprise/DatasmithImporter/Source/DatasmithTranslator/Private/Utility/DatasmithMeshHelper.cpp
@@ -26,18 +26,9 @@ namespace DatasmithMeshHelper
 		return Mesh.Triangles().Num();
 	}
 
-	void ExtractVertexPositions(const FMeshDescription& Mesh, TArray<FVector>& Positions)
+	void ExtractVertexPositions(const FMeshDescription& Mesh, TArray<FVector3f>& OutPositions)
 	{
-		FStaticMeshConstAttributes Attributes(Mesh);
-		const TVertexAttributesConstRef<FVector3f> VertexPositions = Attributes.GetVertexPositions();
-		if (VertexPositions.IsValid())
-		{
-			Positions.Reserve(VertexPositions.GetNumElements());
-			for (const FVertexID VertexID : Mesh.Vertices().GetElementIDs())
-			{
-				Positions.Add((FVector)VertexPositions[VertexID]);
-			}
-		}
+		FDatasmithMeshUtils::ExtractVertexPositions(Mesh, OutPositions);
 	}
 
 	void PrepareAttributeForStaticMesh(FMeshDescription& Mesh)
@@ -297,24 +288,24 @@ namespace DatasmithMeshHelper
 		}
 	}
 
-	bool IsMeshValid(const FMeshDescription& Mesh, FVector BuildScale)
+	bool IsMeshValid(const FMeshDescription& Mesh, FVector3f BuildScale)
 	{
 		TVertexAttributesConstRef<FVector3f> VertexPositions = Mesh.GetVertexPositions();
-		FVector RawNormalScale(BuildScale.Y*BuildScale.Z, BuildScale.X*BuildScale.Z, BuildScale.X*BuildScale.Y); // Component-wise scale
+		FVector3f RawNormalScale(BuildScale.Y*BuildScale.Z, BuildScale.X*BuildScale.Z, BuildScale.X*BuildScale.Y); // Component-wise scale
 
 		for (const FTriangleID TriangleID : Mesh.Triangles().GetElementIDs())
 		{
 			TArrayView<const FVertexID> VertexIDs = Mesh.GetTriangleVertices(TriangleID);
-			FVector Corners[3] =
+			FVector3f Corners[3] =
 			{
-				(FVector)VertexPositions[VertexIDs[0]],
-				(FVector)VertexPositions[VertexIDs[1]],
-				(FVector)VertexPositions[VertexIDs[2]]
+				VertexPositions[VertexIDs[0]],
+				VertexPositions[VertexIDs[1]],
+				VertexPositions[VertexIDs[2]]
 			};
 
-			FVector RawNormal = (Corners[1] - Corners[2]) ^ (Corners[0] - Corners[2]);
+			FVector3f RawNormal = (Corners[1] - Corners[2]) ^ (Corners[0] - Corners[2]);
 			RawNormal *= RawNormalScale;
-			double FourSquaredTriangleArea = RawNormal.SizeSquared();
+			float FourSquaredTriangleArea = RawNormal.SizeSquared();
 
 			// We support even small triangles, but this function is still useful to
 			// see if we have at least one valid triangle in the mesh
diff --git a/Engine/Plugins/Enterprise/DatasmithImporter/Source/DatasmithTranslator/Public/DatasmithPayload.h b/Engine/Plugins/Enterprise/DatasmithImporter/Source/DatasmithTranslator/Public/DatasmithPayload.h
--- a/Engine/Plugins/Enterprise/DatasmithImporter/Source/DatasmithTranslator/Public/DatasmithPayload.h
+++ b/Engine/Plugins/Enterprise/DatasmithImporter/Source/DatasmithTranslator/Public/DatasmithPayload.h
@@ -14,7 +14,7 @@ struct FDatasmithMeshElementPayload : public FDatasmithPayload
 {
 	TArray<FMeshDescription> LodMeshes;
 	FMeshDescription CollisionMesh;
-	TArray<FVector> CollisionPointCloud; // compatibility, favor the CollisionMesh member
+	TArray<FVector3f> CollisionPointCloud; // compatibility, favor the CollisionMesh member
 };
 
 struct FDatasmithLevelSequencePayload : public FDatasmithPayload
diff --git a/Engine/Plugins/Enterprise/DatasmithImporter/Source/DatasmithTranslator/Public/Utility/DatasmithMeshHelper.h b/Engine/Plugins/Enterprise/DatasmithImporter/Source/DatasmithTranslator/Public/Utility/DatasmithMeshHelper.h
--- a/Engine/Plugins/Enterprise/DatasmithImporter/Source/DatasmithTranslator/Public/Utility/DatasmithMeshHelper.h
+++ b/Engine/Plugins/Enterprise/DatasmithImporter/Source/DatasmithTranslator/Public/Utility/DatasmithMeshHelper.h
@@ -40,7 +40,7 @@ namespace DatasmithMeshHelper
 	 * @param Positions
 	 * @return DATASMITHIMPORTER_API void
 	 */
-	DATASMITHTRANSLATOR_API void ExtractVertexPositions(const FMeshDescription& Mesh, TArray<FVector>& Positions);
+	DATASMITHTRANSLATOR_API void ExtractVertexPositions(const FMeshDescription& Mesh, TArray<FVector3f>& OutPositions);
 
 	/**
 	 * Register all attributes required to correctly interact with StaticMesh.
@@ -152,5 +152,5 @@ namespace DatasmithMeshHelper
 	 * @param BuildScale
 	 * @return DATASMITHIMPORTER_API bool
 	 */
-	DATASMITHTRANSLATOR_API bool IsMeshValid(const FMeshDescription& Mesh, FVector BuildScale=FVector::OneVector);
+	DATASMITHTRANSLATOR_API bool IsMeshValid(const FMeshDescription& Mesh, FVector3f BuildScale=FVector3f::OneVector);
 }
diff --git a/Engine/Plugins/Experimental/Enterprise/DatasmithRuntime/Source/Private/DatasmithRuntimeUtils.cpp b/Engine/Plugins/Experimental/Enterprise/DatasmithRuntime/Source/Private/DatasmithRuntimeUtils.cpp
--- a/Engine/Plugins/Experimental/Enterprise/DatasmithRuntime/Source/Private/DatasmithRuntimeUtils.cpp
+++ b/Engine/Plugins/Experimental/Enterprise/DatasmithRuntime/Source/Private/DatasmithRuntimeUtils.cpp
@@ -153,7 +153,7 @@ namespace DatasmithRuntime
 		// Packing expects at least one texel per chart. This is the absolute minimum to generate valid UVs.
 		float ChartCount = FStaticMeshOperations::GetUVChartCount(Mesh, SrcLightmapIndex, ELightmapUVVersion::Latest, OverlappingCorners);
 		const int32 AbsoluteMinResolution = 1 << FMath::CeilLogTwo(FMath::Sqrt(ChartCount));
-		
+
 		return FMath::Clamp(MinLightmapResolution, AbsoluteMinResolution, 512);
 	}
 
@@ -165,7 +165,7 @@ namespace DatasmithRuntime
 			return;
 		}
 
-		TArray< FVector > VertexPositions;
+		TArray< FVector3f > VertexPositions;
 		DatasmithMeshHelper::ExtractVertexPositions(Payload.CollisionMesh, VertexPositions);
 		if (VertexPositions.Num() == 0)
 		{
@@ -188,9 +188,9 @@ namespace DatasmithRuntime
 			FKConvexElem& ConvexElem = AggGeom.ConvexElems.AddDefaulted_GetRef();
 
 			ConvexElem.VertexData.Reserve(VertexPositions.Num());
-			for (const FVector& Position : VertexPositions)
+			for (const FVector3f& Position : VertexPositions)
 			{
-				ConvexElem.VertexData.Add(Position);
+				ConvexElem.VertexData.Add(FVector(Position));
 			}
 
 			ConvexElem.UpdateElemBox();
@@ -452,7 +452,7 @@ namespace DatasmithRuntime
 	#define RCP_SQRT2 (0.70710678118654752440084436210485f)
 	#define RCP_SQRT3 (0.57735026918962576450914878050196f)
 
-	const FVector KDopDir26[26] = 
+	const FVector KDopDir26[26] =
 	{
 		FVector( 1.f, 0.f, 0.f),
 		FVector(-1.f, 0.f, 0.f),
diff --git a/Engine/Source/Developer/Datasmith/DatasmithExporter/Private/DatasmithMeshExporter.cpp b/Engine/Source/Developer/Datasmith/DatasmithExporter/Private/DatasmithMeshExporter.cpp
--- a/Engine/Source/Developer/Datasmith/DatasmithExporter/Private/DatasmithMeshExporter.cpp
+++ b/Engine/Source/Developer/Datasmith/DatasmithExporter/Private/DatasmithMeshExporter.cpp
@@ -145,7 +145,7 @@ bool FDatasmithMeshExporterImpl::DoExport(TSharedPtr< IDatasmithMeshElement >& M
 		MeshElement->SetFile( *ExportOptions.MeshFullPath );
 		MeshElement->SetFileHash( Hash );
 
-		FBox Extents = Mesh.GetExtents();
+		FBox3f Extents = Mesh.GetExtents();
 		float Width = Extents.Max[0] - Extents.Min[0];
 		float Height = Extents.Max[2] - Extents.Min[2];
 		float Depth = Extents.Max[1] - Extents.Min[1];
@@ -358,7 +358,7 @@ void FDatasmithMeshExporterLegacyImpl::PreExport( const FDatasmithMeshExporterOp
 
 void FDatasmithMeshExporterLegacyImpl::PostExport( const FDatasmithMesh& DatasmithMesh, TSharedPtr< IDatasmithMeshElement >& MeshElement )
 {
-	FBox Extents = DatasmithMesh.GetExtents();
+	FBox3f Extents = DatasmithMesh.GetExtents();
 	float Width = Extents.Max[0] - Extents.Min[0];
 	float Height = Extents.Max[2] - Extents.Min[2];
 	float Depth = Extents.Max[1] - Extents.Min[1];
diff --git a/Engine/Source/Developer/Datasmith/DatasmithFacade/Private/DatasmithFacadeActor.cpp b/Engine/Source/Developer/Datasmith/DatasmithFacade/Private/DatasmithFacadeActor.cpp
--- a/Engine/Source/Developer/Datasmith/DatasmithFacade/Private/DatasmithFacadeActor.cpp
+++ b/Engine/Source/Developer/Datasmith/DatasmithFacade/Private/DatasmithFacadeActor.cpp
@@ -27,7 +27,7 @@ FDatasmithFacadeActor::FDatasmithFacadeActor(
 {}
 
 void FDatasmithFacadeActor::SetWorldTransform(
-	const float InWorldMatrix[16],
+	const double InWorldMatrix[16],
 	bool bRowMajor
 )
 {
@@ -41,94 +41,94 @@ void FDatasmithFacadeActor::SetWorldTransform(
 }
 
 void FDatasmithFacadeActor::SetScale(
-	float X,
-	float Y,
-	float Z
+	double X,
+	double Y,
+	double Z
 )
 {
 	GetDatasmithActorElement()->SetScale(FVector(X, Y, Z));
 }
 
 void FDatasmithFacadeActor::GetScale(
-	float& OutX,
-	float& OutY,
-	float& OutZ
+	double& OutX,
+	double& OutY,
+	double& OutZ
 ) const
 {
 	FVector ScaleVector(GetDatasmithActorElement()->GetScale());
 
-	OutX = (float)ScaleVector.X;
-	OutY = (float)ScaleVector.Y;
-	OutZ = (float)ScaleVector.Z;
+	OutX = ScaleVector.X;
+	OutY = ScaleVector.Y;
+	OutZ = ScaleVector.Z;
 }
 
 void FDatasmithFacadeActor::SetRotation(
-	float Pitch,
-	float Yaw,
-	float Roll
+	double Pitch,
+	double Yaw,
+	double Roll
 )
 {
 	GetDatasmithActorElement()->SetRotation(FQuat(FRotator(Pitch, Yaw, Roll)));
 }
 
 void FDatasmithFacadeActor::GetRotation(
-	float& OutPitch,
-	float& OutYaw,
-	float& OutRoll
+	double& OutPitch,
+	double& OutYaw,
+	double& OutRoll
 ) const
 {
 	FRotator Rotator(GetDatasmithActorElement()->GetRotation().Rotator());
 
-	OutPitch = (float)Rotator.Pitch;
-	OutYaw = (float)Rotator.Yaw;
-	OutRoll = (float)Rotator.Roll;
+	OutPitch = Rotator.Pitch;
+	OutYaw = Rotator.Yaw;
+	OutRoll = Rotator.Roll;
 }
 
 void FDatasmithFacadeActor::SetRotation(
-	float X,
-	float Y,
-	float Z,
-	float W
+	double X,
+	double Y,
+	double Z,
+	double W
 )
 {
 	GetDatasmithActorElement()->SetRotation(FQuat(X, Y, Z, W));
 }
 
 void FDatasmithFacadeActor::GetRotation(
-	float& OutX,
-	float& OutY,
-	float& OutZ,
-	float& OutW
+	double& OutX,
+	double& OutY,
+	double& OutZ,
+	double& OutW
 ) const
 {
 	FQuat RotationQuat(GetDatasmithActorElement()->GetRotation());
 
-	OutX = (float)RotationQuat.X;
-	OutY = (float)RotationQuat.Y;
-	OutZ = (float)RotationQuat.Z;
-	OutW = (float)RotationQuat.W;
+	OutX = RotationQuat.X;
+	OutY = RotationQuat.Y;
+	OutZ = RotationQuat.Z;
+	OutW = RotationQuat.W;
 }
 
 void FDatasmithFacadeActor::SetTranslation(
-	float X,
-	float Y,
-	float Z
+	double X,
+	double Y,
+	double Z
 )
 {
 	GetDatasmithActorElement()->SetTranslation(FVector(X, Y, Z));
 }
 
 void FDatasmithFacadeActor::GetTranslation(
-	float& OutX,
-	float& OutY,
-	float& OutZ
+	double& OutX,
+	double& OutY,
+	double& OutZ
 ) const
 {
 	FVector TranslationVector(GetDatasmithActorElement()->GetTranslation());
 
-	OutX = (float)TranslationVector.X;
-	OutY = (float)TranslationVector.Y;
-	OutZ = (float)TranslationVector.Z;
+	OutX = TranslationVector.X;
+	OutY = TranslationVector.Y;
+	OutZ = TranslationVector.Z;
 }
 
 void FDatasmithFacadeActor::SetLayer(
@@ -331,34 +331,34 @@ void FDatasmithFacadeActor::RemoveChild(
 }
 
 FTransform FDatasmithFacadeActor::ConvertTransform(
-	const float InSourceMatrix[16],
+	const double InSourceMatrix[16],
 	bool bRowMajor
 ) const
 {
 	// We use Imath::extractAndRemoveScalingAndShear() because FMatrix::ExtractScaling() is deemed unreliable.
 
 	// Set up a scaling and rotation matrix.
-	Imath::Matrix44<float> Matrix;
+	Imath::Matrix44<double> Matrix;
 
 	if (bRowMajor)
 	{
-		Matrix = Imath::Matrix44<float>(InSourceMatrix[0], InSourceMatrix[4], InSourceMatrix[8],  0.0,
+		Matrix = Imath::Matrix44<double>(InSourceMatrix[0], InSourceMatrix[4], InSourceMatrix[8],  0.0,
 										InSourceMatrix[1], InSourceMatrix[5], InSourceMatrix[9],  0.0,
 										InSourceMatrix[2], InSourceMatrix[6], InSourceMatrix[10], 0.0,
 										0.0			   , 0.0            , 0.0				, 1.0);
 	}
 	else
 	{
-		Matrix = Imath::Matrix44<float>(InSourceMatrix[0], InSourceMatrix[1], InSourceMatrix[2],  0.0,
+		Matrix = Imath::Matrix44<double>(InSourceMatrix[0], InSourceMatrix[1], InSourceMatrix[2],  0.0,
 										InSourceMatrix[4], InSourceMatrix[5], InSourceMatrix[6],  0.0,
 										InSourceMatrix[8], InSourceMatrix[9], InSourceMatrix[10], 0.0,
 										0.0			  , 0.0              , 0.0               , 1.0);
 	}
 
 	// Remove any scaling from the matrix and get the scale vector that was initially present.
-	Imath::Vec3<float> Scale;
-	Imath::Vec3<float> Shear;
-	bool bExtracted = Imath::extractAndRemoveScalingAndShear<float>(Matrix, Scale, Shear, false);
+	Imath::Vec3<double> Scale;
+	Imath::Vec3<double> Shear;
+	bool bExtracted = Imath::extractAndRemoveScalingAndShear<double>(Matrix, Scale, Shear, false);
 
 	if (!bExtracted)
 	{
@@ -371,7 +371,7 @@ FTransform FDatasmithFacadeActor::ConvertTransform(
 	FVector TransformScale3D;
 
 	// Initialize a rotation quaternion with the rotation matrix.
-	Imath::Quat<float> Quaternion = Imath::extractQuat<float>(Matrix);
+	Imath::Quat<double> Quaternion = Imath::extractQuat<double>(Matrix);
 
 	switch (WorldCoordinateSystemType)
 	{
@@ -415,8 +415,8 @@ FTransform FDatasmithFacadeActor::ConvertTransform(
 	}
 
 	// Make sure Unreal will be able to handle the rotation quaternion.
-	float              Angle = Quaternion.angle();
-	Imath::Vec3<float> Axis  = Quaternion.axis();
+	double              Angle = Quaternion.angle();
+	Imath::Vec3<double> Axis  = Quaternion.axis();
 	FQuat TransformRotation = FQuat(FVector(Axis.x, Axis.y, Axis.z), Angle);
 
 	// Scale and convert the source translation into a Datasmith actor translation.
diff --git a/Engine/Source/Developer/Datasmith/DatasmithFacade/Private/DatasmithFacadeActorCamera.cpp b/Engine/Source/Developer/Datasmith/DatasmithFacade/Private/DatasmithFacadeActorCamera.cpp
--- a/Engine/Source/Developer/Datasmith/DatasmithFacade/Private/DatasmithFacadeActorCamera.cpp
+++ b/Engine/Source/Developer/Datasmith/DatasmithFacade/Private/DatasmithFacadeActorCamera.cpp
@@ -26,9 +26,9 @@ FDatasmithFacadeActorCamera::FDatasmithFacadeActorCamera(
 {}
 
 void FDatasmithFacadeActorCamera::SetCameraPosition(
-	float InX,
-	float InY,
-	float InZ
+	double InX,
+	double InY,
+	double InZ
 )
 {
 	// Scale and convert the camera position into a Datasmith actor world translation.
@@ -36,12 +36,12 @@ void FDatasmithFacadeActorCamera::SetCameraPosition(
 }
 
 void FDatasmithFacadeActorCamera::SetCameraRotation(
-	float InForwardX,
-	float InForwardY,
-	float InForwardZ,
-	float InUpX,
-	float InUpY,
-	float InUpZ
+	double InForwardX,
+	double InForwardY,
+	double InForwardZ,
+	double InUpX,
+	double InUpY,
+	double InUpZ
 )
 {
 	// Convert the camera orientation into a Datasmith actor world look-at rotation quaternion.
@@ -82,9 +82,10 @@ void FDatasmithFacadeActorCamera::SetFocusDistance(
 {
 	TSharedRef<IDatasmithCameraActorElement> CameraActor(GetDatasmithActorCameraElement());
 	FVector CameraPosition(CameraActor->GetTranslation());
+	FVector Target(InTargetX, InTargetY, InTargetZ);
 
-	FVector DistanceVector(InTargetX - CameraPosition.X, InTargetY - CameraPosition.Y, InTargetZ - CameraPosition.Z);
-	CameraActor->SetFocusDistance((float)DistanceVector.Size() * WorldUnitScale);	// LWC_TODO: Precision loss
+	FVector DistanceVector = Target - CameraPosition;
+	CameraActor->SetFocusDistance((float)DistanceVector.Size() * WorldUnitScale);
 }
 
 void FDatasmithFacadeActorCamera::SetFocusDistance(
diff --git a/Engine/Source/Developer/Datasmith/DatasmithFacade/Private/DatasmithFacadeElement.cpp b/Engine/Source/Developer/Datasmith/DatasmithFacade/Private/DatasmithFacadeElement.cpp
--- a/Engine/Source/Developer/Datasmith/DatasmithFacade/Private/DatasmithFacadeElement.cpp
+++ b/Engine/Source/Developer/Datasmith/DatasmithFacade/Private/DatasmithFacadeElement.cpp
@@ -7,44 +7,11 @@ FDatasmithFacadeElement::ECoordinateSystemType FDatasmithFacadeElement::WorldCoo
 
 float FDatasmithFacadeElement::WorldUnitScale = 1.0;
 
-FDatasmithFacadeElement::ConvertVertexMethod FDatasmithFacadeElement::ConvertPosition  = nullptr;
-FDatasmithFacadeElement::ConvertVectorMethod FDatasmithFacadeElement::ConvertBackPosition = nullptr;
-FDatasmithFacadeElement::ConvertVertexMethod FDatasmithFacadeElement::ConvertDirection = nullptr;
-FDatasmithFacadeElement::ConvertVectorMethod FDatasmithFacadeElement::ConvertBackDirection = nullptr;
-
 void FDatasmithFacadeElement::SetCoordinateSystemType(
 	ECoordinateSystemType InWorldCoordinateSystemType
 )
 {
 	WorldCoordinateSystemType = InWorldCoordinateSystemType;
-
-	switch (WorldCoordinateSystemType)
-	{
-		case ECoordinateSystemType::LeftHandedYup:
-		{
-			ConvertPosition  = ConvertFromPositionLeftHandedYup;
-			ConvertBackPosition = ConvertToPositionLeftHandedYup;
-			ConvertDirection = ConvertFromDirectionLeftHandedYup;
-			ConvertBackDirection = ConvertToDirectionLeftHandedYup;
-			break;
-		}
-		case ECoordinateSystemType::LeftHandedZup:
-		{
-			ConvertPosition  = ConvertFromPositionLeftHandedZup;
-			ConvertBackPosition = ConvertToPositionLeftHandedZup;
-			ConvertDirection = ConvertFromDirectionLeftHandedZup;
-			ConvertBackDirection = ConvertToDirectionLeftHandedZup;
-			break;
-		}
-		case ECoordinateSystemType::RightHandedZup:
-		{
-			ConvertPosition  = ConvertFromPositionRightHandedZup;
-			ConvertBackPosition = ConvertToPositionRightHandedZup;
-			ConvertDirection = ConvertFromDirectionRightHandedZup;
-			ConvertBackDirection = ConvertToDirectionRightHandedZup;
-			break;
-		}
-	}
 }
 
 void FDatasmithFacadeElement::SetWorldUnitScale(
@@ -90,122 +57,65 @@ const TCHAR* FDatasmithFacadeElement::GetLabel() const
 	return InternalDatasmithElement->GetLabel();
 }
 
-FVector FDatasmithFacadeElement::ConvertTranslation(
-	FVector const& InVertex
-)
+template<bool kIsForward, bool kIsDirection, typename Vec_t>
+Vec_t FDatasmithFacadeElement::Convert(const Vec_t& V)
 {
-	return ConvertPosition((float)InVertex.X, (float)InVertex.Y, (float)InVertex.Z);	// LWC_TODO: Precision loss
-}
-
-void FDatasmithFacadeElement::ExportAsset(
-	FString const& InAssetFolder
-)
-{
-	// By default, there is no Datasmith scene element asset to build and export.
-}
+	Vec_t Tmp;
+	switch (WorldCoordinateSystemType)
+	{
+		case ECoordinateSystemType::LeftHandedYup:
+			Tmp = kIsForward ? Vec_t{V.X, -V.Z, V.Y} : Vec_t{V.X, V.Z, -V.Y};
+			break;
 
-FVector FDatasmithFacadeElement::ConvertFromPositionLeftHandedYup(
-	float InX,
-	float InY,
-	float InZ
-)
-{
-	return FVector(InX * WorldUnitScale, -InZ * WorldUnitScale, InY * WorldUnitScale);
-}
+		case ECoordinateSystemType::RightHandedZup:
+			// To avoid perturbating X, which is forward in Unreal, the handedness conversion is done by flipping the side vector Y.
+			Tmp = Vec_t{V.X, -V.Y, V.Z};
+			break;
 
-FVector FDatasmithFacadeElement::ConvertToPositionLeftHandedYup(
-	const FVector& InVector
-)
-{
-	float Scale = 1.f / WorldUnitScale;
-	return FVector(InVector.X * Scale, InVector.Z * Scale, -InVector.Y * Scale);
-}
+		case ECoordinateSystemType::LeftHandedZup:
+		default:
+			Tmp = V;
+			break;
+	}
 
-FVector FDatasmithFacadeElement::ConvertFromPositionLeftHandedZup(
-	float InX,
-	float InY,
-	float InZ
-)
-{
-	return FVector(InX * WorldUnitScale, InY * WorldUnitScale, InZ * WorldUnitScale);
-}
+	// World scaling for positions conversions
+	if constexpr (!kIsDirection)
+	{
+		Tmp *= kIsForward ? 1 / WorldUnitScale : WorldUnitScale;
+	}
 
-FVector FDatasmithFacadeElement::ConvertToPositionLeftHandedZup(
-	const FVector& InVector
-)
-{
-	return InVector / WorldUnitScale;
+	return Tmp;
 }
 
-FVector FDatasmithFacadeElement::ConvertFromPositionRightHandedZup(
-	float InX,
-	float InY,
-	float InZ
-)
-{
-	// Convert the position from the source right-hand Z-up coordinate system to the Unreal left-hand Z-up coordinate system.
-	// To avoid perturbating X, which is forward in Unreal, the handedness conversion is done by flipping the side vector Y.
-	return FVector(InX * WorldUnitScale, -InY * WorldUnitScale, InZ * WorldUnitScale);
-}
+static constexpr bool kForward = true;
+static constexpr bool kBackward = false;
+static constexpr bool kDir = true;
+static constexpr bool kPos = false;
 
-FVector FDatasmithFacadeElement::ConvertToPositionRightHandedZup(
-	const FVector& InVector
-)
-{
-	// Convert the position from the source right-hand Z-up coordinate system to the Unreal left-hand Z-up coordinate system.
-	// To avoid perturbating X, which is forward in Unreal, the handedness conversion is done by flipping the side vector Y.
-	float Scale = 1.f / WorldUnitScale;
-	return FVector(InVector.X * Scale, -InVector.Y * Scale, InVector.Z * Scale);
-}
+FVector   FDatasmithFacadeElement::ConvertDirection    (const FVector& In)                  { return Convert<kForward,  kDir>(In);                       }
+FVector   FDatasmithFacadeElement::ConvertBackDirection(const FVector& In)                  { return Convert<kBackward, kDir>(In);                       }
+FVector   FDatasmithFacadeElement::ConvertPosition     (const FVector& In)                  { return Convert<kForward,  kPos>(In);                       }
+FVector   FDatasmithFacadeElement::ConvertBackPosition (const FVector& In)                  { return Convert<kBackward, kPos>(In);                       }
 
-FVector FDatasmithFacadeElement::ConvertFromDirectionLeftHandedYup(
-	float InX,
-	float InY,
-	float InZ
-)
-{
-	return FVector(InX, -InZ, InY);
-}
+FVector   FDatasmithFacadeElement::ConvertDirection    (double InX, double InY, double InZ) { return Convert<kForward,  kDir>(FVector{InX, InY, InZ});   }
+FVector   FDatasmithFacadeElement::ConvertBackDirection(double InX, double InY, double InZ) { return Convert<kBackward, kDir>(FVector{InX, InY, InZ});   }
+FVector   FDatasmithFacadeElement::ConvertPosition     (double InX, double InY, double InZ) { return Convert<kForward,  kPos>(FVector{InX, InY, InZ});   }
+FVector   FDatasmithFacadeElement::ConvertBackPosition (double InX, double InY, double InZ) { return Convert<kBackward, kPos>(FVector{InX, InY, InZ});   }
 
-FVector FDatasmithFacadeElement::ConvertToDirectionLeftHandedYup(
-	const FVector& InVector
-)
-{
-	return FVector(InVector.X, InVector.Z, -InVector.Y);
-}
+FVector3f FDatasmithFacadeElement::ConvertDirection    (const FVector3f& In)                { return Convert<kForward,  kDir>(In);                       }
+FVector3f FDatasmithFacadeElement::ConvertBackDirection(const FVector3f& In)                { return Convert<kBackward, kDir>(In);                       }
+FVector3f FDatasmithFacadeElement::ConvertPosition     (const FVector3f& In)                { return Convert<kForward,  kPos>(In);                       }
+FVector3f FDatasmithFacadeElement::ConvertBackPosition (const FVector3f& In)                { return Convert<kBackward, kPos>(In);                       }
 
-FVector FDatasmithFacadeElement::ConvertFromDirectionLeftHandedZup(
-	float InX,
-	float InY,
-	float InZ
-)
-{
-	return FVector(InX, InY, InZ);
-}
+FVector3f FDatasmithFacadeElement::ConvertDirection    (float InX, float InY, float InZ)    { return Convert<kForward,  kDir>(FVector3f{InX, InY, InZ}); }
+FVector3f FDatasmithFacadeElement::ConvertBackDirection(float InX, float InY, float InZ)    { return Convert<kBackward, kDir>(FVector3f{InX, InY, InZ}); }
+FVector3f FDatasmithFacadeElement::ConvertPosition     (float InX, float InY, float InZ)    { return Convert<kForward,  kPos>(FVector3f{InX, InY, InZ}); }
+FVector3f FDatasmithFacadeElement::ConvertBackPosition (float InX, float InY, float InZ)    { return Convert<kBackward, kPos>(FVector3f{InX, InY, InZ}); }
 
-FVector FDatasmithFacadeElement::ConvertToDirectionLeftHandedZup(
-	const FVector& InVector
-)
-{
-	return InVector;
-}
 
-FVector FDatasmithFacadeElement::ConvertFromDirectionRightHandedZup(
-	float InX,
-	float InY,
-	float InZ
+void FDatasmithFacadeElement::ExportAsset(
+	FString const& InAssetFolder
 )
 {
-	// Convert the direction from the source right-hand Z-up coordinate system to the Unreal left-hand Z-up coordinate system.
-	// To avoid perturbating X, which is forward in Unreal, the handedness conversion is done by flipping the side vector Y.
-	return FVector(InX, -InY, InZ);
+	// By default, there is no Datasmith scene element asset to build and export.
 }
-
-FVector FDatasmithFacadeElement::ConvertToDirectionRightHandedZup(
-	const FVector& InVector
-)
-{
-	// Convert the direction from the source right-hand Z-up coordinate system to the Unreal left-hand Z-up coordinate system.
-	// To avoid perturbating X, which is forward in Unreal, the handedness conversion is done by flipping the side vector Y.
-	return FVector(InVector.X, -InVector.Y, InVector.Z);
-}
\ No newline at end of file
diff --git a/Engine/Source/Developer/Datasmith/DatasmithFacade/Private/DatasmithFacadeMesh.cpp b/Engine/Source/Developer/Datasmith/DatasmithFacade/Private/DatasmithFacadeMesh.cpp
--- a/Engine/Source/Developer/Datasmith/DatasmithFacade/Private/DatasmithFacadeMesh.cpp
+++ b/Engine/Source/Developer/Datasmith/DatasmithFacade/Private/DatasmithFacadeMesh.cpp
@@ -16,8 +16,8 @@ void FDatasmithFacadeMesh::SetVertex(int32 Index, float X, float Y, float Z)
 }
 
 void FDatasmithFacadeMesh::GetVertex(int32 Index, float& OutX, float& OutY, float& OutZ) const
-{ 
-	FVector3f Position(FDatasmithFacadeElement::ConvertBackPosition(GetDatasmithMesh().GetVertex(Index)));
+{
+	FVector3f Position = FDatasmithFacadeElement::ConvertBackPosition(GetDatasmithMesh().GetVertex(Index));
 	OutX = Position.X;
 	OutY = Position.Y;
 	OutZ = Position.Z;
diff --git a/Engine/Source/Developer/Datasmith/DatasmithFacade/Private/DatasmithFacadeVariant.cpp b/Engine/Source/Developer/Datasmith/DatasmithFacade/Private/DatasmithFacadeVariant.cpp
--- a/Engine/Source/Developer/Datasmith/DatasmithFacade/Private/DatasmithFacadeVariant.cpp
+++ b/Engine/Source/Developer/Datasmith/DatasmithFacade/Private/DatasmithFacadeVariant.cpp
@@ -56,9 +56,9 @@ TSharedRef<IDatasmithActorBindingElement> FDatasmithFacadeActorBinding::GetDatas
 }
 
 void FDatasmithFacadeActorBinding::AddRelativeLocationCapture(
-	float X,
-	float Y,
-	float Z
+	double X,
+	double Y,
+	double Z
 )
 {
 	FVector Location(X, Y, Z);
@@ -69,9 +69,9 @@ void FDatasmithFacadeActorBinding::AddRelativeLocationCapture(
 }
 
 void FDatasmithFacadeActorBinding::AddRelativeRotationCapture(
-	float Pitch,
-	float Yaw,
-	float Roll
+	double Pitch,
+	double Yaw,
+	double Roll
 )
 {
 	FRotator Rotation(Pitch, Yaw, Roll);
@@ -82,10 +82,10 @@ void FDatasmithFacadeActorBinding::AddRelativeRotationCapture(
 }
 
 void FDatasmithFacadeActorBinding::AddRelativeRotationCapture(
-	float X,
-	float Y,
-	float Z,
-	float W
+	double X,
+	double Y,
+	double Z,
+	double W
 )
 {
 	FRotator Rotation(FQuat(X, Y, Z, W));
@@ -96,9 +96,9 @@ void FDatasmithFacadeActorBinding::AddRelativeRotationCapture(
 }
 
 void FDatasmithFacadeActorBinding::AddRelativeScaleCapture(
-	float X,
-	float Y,
-	float Z
+	double X,
+	double Y,
+	double Z
 )
 {
 	FVector Scale(X, Y, Z);
@@ -109,7 +109,7 @@ void FDatasmithFacadeActorBinding::AddRelativeScaleCapture(
 }
 
 void FDatasmithFacadeActorBinding::AddRelativeTransformCapture(
-	const float InMatrix[16],
+	const double InMatrix[16],
 	bool bRowMajor
 )
 {
@@ -124,9 +124,9 @@ void FDatasmithFacadeActorBinding::AddRelativeTransformCapture(
 	FVector Location = Transform.GetTranslation();
 	FQuat Rotation = Transform.GetRotation();
 
-	AddRelativeScaleCapture((float)Scale.X, (float)Scale.Y, (float)Scale.Z);	// LWC_TODO: Precision loss
-	AddRelativeLocationCapture((float)Location.X, (float)Location.Y, (float)Location.Z);
-	AddRelativeRotationCapture((float)Rotation.X, (float)Rotation.Y, (float)Rotation.Z, (float)Rotation.W);
+	AddRelativeScaleCapture(Scale.X, Scale.Y, Scale.Z);
+	AddRelativeLocationCapture(Location.X, Location.Y, Location.Z);
+	AddRelativeRotationCapture(Rotation.X, Rotation.Y, Rotation.Z, Rotation.W);
 }
 
 void FDatasmithFacadeActorBinding::AddVisibilityCapture(
diff --git a/Engine/Source/Developer/Datasmith/DatasmithFacade/Public/DatasmithFacadeActor.h b/Engine/Source/Developer/Datasmith/DatasmithFacade/Public/DatasmithFacadeActor.h
--- a/Engine/Source/Developer/Datasmith/DatasmithFacade/Public/DatasmithFacadeActor.h
+++ b/Engine/Source/Developer/Datasmith/DatasmithFacade/Public/DatasmithFacadeActor.h
@@ -34,60 +34,72 @@ public:
 
 	virtual ~FDatasmithFacadeActor() {}
 
-	/** Set the world transform of the Datasmith actor. */
 	void SetWorldTransform(
-		const float InWorldMatrix[16],
+		const double InWorldMatrix[16],
 		bool bRowMajor = false
 	);
 
+	void SetWorldTransform(
+		const float InWorldMatrix[16],
+		bool bRowMajor = false
+	)
+	{
+		double InMatrixD[16];
+		for(int I = 0; I < 16; ++I)
+		{
+			InMatrixD[I] = InWorldMatrix[I];
+		}
+		SetWorldTransform(InMatrixD, bRowMajor);
+	}
+
 	void SetScale(
-		float X,
-		float Y,
-		float Z
+		double X,
+		double Y,
+		double Z
 	);
 
 	void GetScale(
-		float& OutX,
-		float& OutY,
-		float& OutZ
+		double& OutX,
+		double& OutY,
+		double& OutZ
 	) const;
 
 	void SetRotation(
-		float Pitch,
-		float Yaw,
-		float Roll
+		double Pitch,
+		double Yaw,
+		double Roll
 	);
 
 	void GetRotation(
-		float& OutPitch,
-		float& OutYaw,
-		float& OutRoll
+		double& OutPitch,
+		double& OutYaw,
+		double& OutRoll
 	) const;
 
 	void SetRotation(
-		float X,
-		float Y,
-		float Z,
-		float W
+		double X,
+		double Y,
+		double Z,
+		double W
 	);
 
 	void GetRotation(
-		float& OutX,
-		float& OutY,
-		float& OutZ,
-		float& OutW
+		double& OutX,
+		double& OutY,
+		double& OutZ,
+		double& OutW
 	) const;
 
 	void SetTranslation(
-		float X,
-		float Y,
-		float Z
+		double X,
+		double Y,
+		double Z
 	);
 
 	void GetTranslation(
-		float& OutX,
-		float& OutY,
-		float& OutZ
+		double& OutX,
+		double& OutY,
+		double& OutZ
 	) const;
 
 	// Set the layer of the Datasmith actor.
@@ -187,7 +199,7 @@ protected:
 
 	// Convert a source matrix into a Datasmith actor transform.
 	FTransform ConvertTransform(
-		const float InSourceMatrix[16],
+		const double InSourceMatrix[16],
 		bool bRowMajor
 	) const;
 
diff --git a/Engine/Source/Developer/Datasmith/DatasmithFacade/Public/DatasmithFacadeActorCamera.h b/Engine/Source/Developer/Datasmith/DatasmithFacade/Public/DatasmithFacadeActorCamera.h
--- a/Engine/Source/Developer/Datasmith/DatasmithFacade/Public/DatasmithFacadeActorCamera.h
+++ b/Engine/Source/Developer/Datasmith/DatasmithFacade/Public/DatasmithFacadeActorCamera.h
@@ -19,19 +19,19 @@ public:
 
 	// Set the world position of the Datasmith camera actor.
 	void SetCameraPosition(
-		float InX, // camera position on the X axis
-		float InY, // camera position on the Y axis
-		float InZ  // camera position on the Z axis
+		double InX, // camera position on the X axis
+		double InY, // camera position on the Y axis
+		double InZ  // camera position on the Z axis
 	);
 
 	// Set the world rotation of the Datasmith camera actor with the camera world forward and up directions.
 	void SetCameraRotation(
-		float InForwardX, // camera forward direction on the X axis
-		float InForwardY, // camera forward direction on the Y axis
-		float InForwardZ, // camera forward direction on the Z axis
-		float InUpX,      // camera up direction on the X axis
-		float InUpY,      // camera up direction on the Y axis
-		float InUpZ       // camera up direction on the Z axis
+		double InForwardX, // camera forward direction on the X axis
+		double InForwardY, // camera forward direction on the Y axis
+		double InForwardZ, // camera forward direction on the Z axis
+		double InUpX,      // camera up direction on the X axis
+		double InUpY,      // camera up direction on the Y axis
+		double InUpZ       // camera up direction on the Z axis
 	);
 
 	// Set the sensor width of the Datasmith camera.
diff --git a/Engine/Source/Developer/Datasmith/DatasmithFacade/Public/DatasmithFacadeElement.h b/Engine/Source/Developer/Datasmith/DatasmithFacade/Public/DatasmithFacadeElement.h
--- a/Engine/Source/Developer/Datasmith/DatasmithFacade/Public/DatasmithFacadeElement.h
+++ b/Engine/Source/Developer/Datasmith/DatasmithFacade/Public/DatasmithFacadeElement.h
@@ -21,18 +21,6 @@ public:
 		RightHandedZup // right-handed Z-up coordinate system: Revit
 	};
 
-	// Type of a method that converts a vertex from world coordinates to Unreal coordinates.
-	typedef FVector (* ConvertVertexMethod)(
-		float InX, // vertex X in world coordinates
-		float InY, // vertex Y in world coordinates
-		float InZ  // vertex Z in world coordinates
-	);
-
-	// Type of a method that converts an FVector from Unreal coordinates to world coordinates.
-	typedef FVector (* ConvertVectorMethod )(
-		const FVector& InVector
-	);
-
 public:
 
 	// Set the coordinate system type of the world geometries and transforms.
@@ -40,7 +28,7 @@ public:
 		ECoordinateSystemType InWorldCoordinateSystemType // world coordinate system type
 	);
 
-	/** 
+	/**
 	 * Set the scale factor from world units to Datasmith centimeters.
 	 * If the given value is too close to 0, it is clamped to SMALL_NUMBER value.
 	 */
@@ -73,22 +61,30 @@ public:
 protected:
 #endif
 
-	// Convert a position from world unit coordinates to Unreal centimeter coordinates.
-	static ConvertVertexMethod ConvertPosition;
+	template<bool kIsForward, bool kIsDirection, typename Vec_t> static Vec_t Convert(const Vec_t& V);
 
-	// Convert a position from Unreal centimeter coordinates to world unit coordinates.
-	static ConvertVectorMethod ConvertBackPosition;
+	static FVector   ConvertDirection    (const FVector& In);
+	static FVector   ConvertBackDirection(const FVector& In);
+	static FVector   ConvertPosition     (const FVector& In);
+	static FVector   ConvertBackPosition (const FVector& In);
 
-	// Convert a direction from the world coordinate system to the Unreal coordinate system.
-	static ConvertVertexMethod ConvertDirection;
+	static FVector   ConvertDirection    (double InX, double InY, double InZ);
+	static FVector   ConvertBackDirection(double InX, double InY, double InZ);
+	static FVector   ConvertPosition     (double InX, double InY, double InZ);
+	static FVector   ConvertBackPosition (double InX, double InY, double InZ);
 
-	// Convert a direction from the Unreal coordinate system to the world coordinate system.
-	static ConvertVectorMethod ConvertBackDirection;
+	static FVector3f ConvertDirection    (const FVector3f& In);
+	static FVector3f ConvertBackDirection(const FVector3f& In);
+	static FVector3f ConvertPosition     (const FVector3f& In);
+	static FVector3f ConvertBackPosition (const FVector3f& In);
 
-	// Convert a translation from world unit coordinates to Unreal centimeter coordinates.
-	static FVector ConvertTranslation(
-		FVector const& InVertex // translation in world unit coordinates
-	);
+	static FVector3f ConvertDirection    (float InX, float InY, float InZ);
+	static FVector3f ConvertBackDirection(float InX, float InY, float InZ);
+	static FVector3f ConvertPosition     (float InX, float InY, float InZ);
+	static FVector3f ConvertBackPosition (float InX, float InY, float InZ);
+
+	static FVector   ConvertTranslation(FVector const& InVertex) { return ConvertPosition(InVertex); }
+	static FVector3f ConvertTranslation(FVector3f const& InVertex) { return ConvertPosition(InVertex); }
 
 	// Build and export the Datasmith scene element asset when required.
 	// This must be done before building a Datasmith scene element.
@@ -100,77 +96,6 @@ protected:
 
 	const TSharedRef<IDatasmithElement>& GetDatasmithElement() const { return InternalDatasmithElement;	}
 
-private:
-
-	// Scale and convert a position from left-handed Y-up coordinates to Unreal left-handed Z-up centimeter coordinates.
-	static FVector ConvertFromPositionLeftHandedYup(
-		float InX, // position X in world unit coordinates
-		float InY, // position Y in world unit coordinates
-		float InZ  // position Z in world unit coordinates
-	);
-
-	// Scale and convert a position from Unreal left-handed Z-up centimeter coordinates to left-handed Y-up coordinates.
-	static FVector ConvertToPositionLeftHandedYup(
-		const FVector& InVector
-	);
-
-	// Scale a position to Unreal left-handed Z-up centimeter coordinates.
-	static FVector ConvertFromPositionLeftHandedZup(
-		float InX, // position X in world unit coordinates
-		float InY, // position Y in world unit coordinates
-		float InZ  // position Z in world unit coordinates
-	);
-
-	// Scale a position from Unreal left-handed Z-up centimeter coordinates.
-	static FVector ConvertToPositionLeftHandedZup(
-		const FVector& InVector
-	);
-
-	// Scale and convert a position from right-handed Z-up coordinates to Unreal left-handed Z-up centimeter coordinates.
-	static FVector ConvertFromPositionRightHandedZup(
-		float InX, // position X in world unit coordinates
-		float InY, // position Y in world unit coordinates
-		float InZ  // position Z in world unit coordinates
-	);
-
-	// Scale and convert a position from Unreal left-handed Z-up centimeter coordinates to right-handed Z-up coordinates.
-	static FVector ConvertToPositionRightHandedZup(
-		const FVector& InVector
-	);
-
-	// Convert a direction from a left-handed Y-up coordinate system to Unreal left-handed Z-up coordinate system.
-	static FVector ConvertFromDirectionLeftHandedYup(
-		float InX, // direction X in world coordinates
-		float InY, // direction Y in world coordinates
-		float InZ  // direction Z in world coordinates
-	);
-
-	static FVector ConvertToDirectionLeftHandedYup(
-		const FVector& InVector
-	);
-
-	// Convert a direction to Unreal left-handed Z-up coordinate system.
-	static FVector ConvertFromDirectionLeftHandedZup(
-		float InX, // direction X in world coordinates
-		float InY, // direction Y in world coordinates
-		float InZ  // direction Z in world coordinates
-	);
-
-	static FVector ConvertToDirectionLeftHandedZup(
-		const FVector& InVector
-	);
-
-	// Convert a direction from a right-handed Z-up coordinate system to Unreal left-handed Z-up coordinate system.
-	static FVector ConvertFromDirectionRightHandedZup(
-		float InX, // direction X in world coordinates
-		float InY, // direction Y in world coordinates
-		float InZ  // direction Z in world coordinates
-	);
-
-	static FVector ConvertToDirectionRightHandedZup(
-		const FVector& InVector
-	);
-
 protected:
 
 	FDatasmithFacadeElement(
diff --git a/Engine/Source/Developer/Datasmith/DatasmithFacade/Public/DatasmithFacadeMesh.h b/Engine/Source/Developer/Datasmith/DatasmithFacade/Public/DatasmithFacadeMesh.h
--- a/Engine/Source/Developer/Datasmith/DatasmithFacade/Public/DatasmithFacadeMesh.h
+++ b/Engine/Source/Developer/Datasmith/DatasmithFacade/Public/DatasmithFacadeMesh.h
@@ -14,18 +14,18 @@ class DATASMITHFACADE_API FDatasmithFacadeMesh
 {
 public:
 
-	FDatasmithFacadeMesh() 
+	FDatasmithFacadeMesh()
 		: RootOwnerMesh(MakeShared<FDatasmithMesh>())
 		, InternalMesh(RootOwnerMesh.Get())
 	{}
 
 	void SetName(
 		const TCHAR* InName
-	) 
-	{ 
-		GetDatasmithMesh().SetName(InName); 
+	)
+	{
+		GetDatasmithMesh().SetName(InName);
 	}
-	
+
 	const TCHAR* GetName() const
 	{
 		return GetDatasmithMesh().GetName();
@@ -38,12 +38,12 @@ public:
 	void SetFacesCount(
 		int32 NumFaces
 	)
-	{ 
+	{
 		GetDatasmithMesh().SetFacesCount(NumFaces);
 	}
 
 	/** Retrieves the amount of faces */
-	int32 GetFacesCount() const 
+	int32 GetFacesCount() const
 	{
 		return GetDatasmithMesh().GetFacesCount();
 	}
@@ -88,7 +88,7 @@ public:
 		int32 Index,
 		uint32 SmoothingMask
 	)
-	{ 
+	{
 		GetDatasmithMesh().SetFaceSmoothingMask(Index, SmoothingMask);
 	}
 
@@ -100,15 +100,15 @@ public:
 	uint32 GetFaceSmoothingMask(
 		int32 Index
 	) const
-	{ 
+	{
 		return GetDatasmithMesh().GetFaceSmoothingMask(Index);
 	}
 
-	int32 GetMaterialsCount() const 
-	{ 
+	int32 GetMaterialsCount() const
+	{
 		return GetDatasmithMesh().GetMaterialsCount();
 	}
-	
+
 	bool IsMaterialIdUsed(
 		int32 MaterialId
 	) const
@@ -123,7 +123,7 @@ public:
 	void SetVerticesCount(
 		int32 NumVerts
 	)
-	{ 
+	{
 		GetDatasmithMesh().SetVerticesCount(NumVerts);
 	}
 
@@ -172,7 +172,7 @@ public:
 		float Y,
 		float Z
 	);
-	
+
 	void GetNormal(
 		int32 Index,
 		float& OutX,
@@ -191,7 +191,7 @@ public:
 	void SetUVChannelsCount(
 		int32 ChannelCount
 	)
-	{ 
+	{
 		GetDatasmithMesh().SetUVChannelsCount(ChannelCount);
 	}
 
@@ -199,7 +199,7 @@ public:
 	 * Add a UV channel at the end
 	 */
 	void AddUVChannel()
-	{ 
+	{
 		GetDatasmithMesh().AddUVChannel();
 	}
 
@@ -212,8 +212,8 @@ public:
 	}
 
 	/** Gets the amount of UV channels on this mesh */
-	int32 GetUVChannelsCount() const 
-	{ 
+	int32 GetUVChannelsCount() const
+	{
 		return GetDatasmithMesh().GetUVChannelsCount();
 	}
 
@@ -227,7 +227,7 @@ public:
 		int32 Channel,
 		int32 NumVerts
 	)
-	{ 
+	{
 		return GetDatasmithMesh().SetUVCount(Channel, NumVerts);
 	}
 
@@ -251,8 +251,8 @@ public:
 		int32 Index,
 		double U,
 		double V
-	) 
-	{ 
+	)
+	{
 		return GetDatasmithMesh().SetUV(Channel, Index, U, V);
 	}
 
@@ -261,8 +261,8 @@ public:
 	 */
 	uint32 GetHashForUVChannel(
 		int32 Channel
-	) const 
-	{ 
+	) const
+	{
 		return GetDatasmithMesh().GetHashForUVChannel(Channel);
 	}
 
@@ -273,15 +273,15 @@ public:
 	 * @param Index		The UV coordinates to retrieve
 	 */
 	void GetUV(
-		int32 Channel, 
-		int32 Index, 
-		float& OutU, 
-		float& OutV
-	) const 
-	{ 
+		int32 Channel,
+		int32 Index,
+		double& OutU,
+		double& OutV
+	) const
+	{
 		FVector2D UV(GetDatasmithMesh().GetUV(Channel, Index));
-		OutU = (float)UV.X;			// LWC_TODO: Precision loss
-		OutV = (float)UV.Y;
+		OutU = UV.X;
+		OutV = UV.Y;
 	}
 
 	/**
@@ -294,10 +294,10 @@ public:
 	 * @param Vertex3	Index of the third uv vertex that defines this face
 	 */
 	void SetFaceUV(
-		int32 Index, 
-		int32 Channel, 
-		int32 Vertex1, 
-		int32 Vertex2, 
+		int32 Index,
+		int32 Channel,
+		int32 Vertex1,
+		int32 Vertex2,
 		int32 Vertex3
 	)
 	{
@@ -314,10 +314,10 @@ public:
 	 * @param Vertex3	Index of the third uv vertex that defines this face
 	 */
 	void GetFaceUV(
-		int32 Index, 
-		int32 Channel, 
-		int32& OutVertex1, 
-		int32& OutVertex2, 
+		int32 Index,
+		int32 Channel,
+		int32& OutVertex1,
+		int32& OutVertex2,
 		int32& OutVertex3
 	) const
 	{
@@ -326,7 +326,7 @@ public:
 
 	/** Get the number of vertex color */
 	int32 GetVertexColorCount() const
-	{ 
+	{
 		return GetDatasmithMesh().GetVertexColorCount();
 	}
 
@@ -337,13 +337,13 @@ public:
 	 * @param Color		The color for the vertex
 	 */
 	void SetVertexColor(
-		int32 Index, 
-		uint8 R, 
-		uint8 G, 
-		uint8 B, 
+		int32 Index,
+		uint8 R,
+		uint8 G,
+		uint8 B,
 		uint8 A
 	)
-	{ 
+	{
 		return GetDatasmithMesh().SetVertexColor( Index, FColor(R, G, B, A) );
 	}
 
@@ -353,13 +353,13 @@ public:
 	 * @param Index		Index of the vertex color to retrieve
 	 */
 	void GetVertexColor(
-		int32 Index, 
-		uint8& OutR, 
-		uint8 OutG, 
-		uint8& OutB, 
+		int32 Index,
+		uint8& OutR,
+		uint8 OutG,
+		uint8& OutB,
 		uint8& OutA
 	) const
-	{ 
+	{
 		FColor VertexColor(GetDatasmithMesh().GetVertexColor(Index));
 		OutR = VertexColor.R;
 		OutG = VertexColor.G;
@@ -373,15 +373,15 @@ public:
 	 */
 	void SetLightmapSourceUVChannel(
 		int32 Channel
-	) 
-	{ 
+	)
+	{
 		return GetDatasmithMesh().SetLightmapSourceUVChannel(Channel);
 	}
 
 	/** Gets the UV channel that will be used for lightmap UVs generation at import */
-	int32 GetLightmapSourceUVChannel() const 
-	{ 
-		return GetDatasmithMesh().GetLightmapSourceUVChannel(); 
+	int32 GetLightmapSourceUVChannel() const
+	{
+		return GetDatasmithMesh().GetLightmapSourceUVChannel();
 	}
 
 	//--------------------------
@@ -405,9 +405,9 @@ public:
 	// Misc
 	//--------------------------
 	/** Returns the total surface area */
-	float ComputeArea() const 
-	{ 
-		return GetDatasmithMesh().ComputeArea(); 
+	float ComputeArea() const
+	{
+		return GetDatasmithMesh().ComputeArea();
 	}
 
 	///////////////////////////////
@@ -417,7 +417,7 @@ protected:
 #endif
 
 	FDatasmithFacadeMesh(
-		const TSharedRef<FDatasmithMesh>& InRootOwnerMesh, 
+		const TSharedRef<FDatasmithMesh>& InRootOwnerMesh,
 		FDatasmithMesh& InReferencedMesh
 	)
 		: RootOwnerMesh(InRootOwnerMesh)
@@ -441,7 +441,7 @@ public:
 
 	/** Get the output filename, it can be absolute or relative to the scene file */
 	const TCHAR* GetFile() const
-	{ 
+	{
 		return GetDatasmithMeshElement()->GetFile();
 	}
 
@@ -449,7 +449,7 @@ public:
 	void SetFile(
 		const TCHAR* InFile
 	)
-	{ 
+	{
 		GetDatasmithMeshElement()->SetFile(InFile);
 	}
 
@@ -473,9 +473,9 @@ public:
 	 * @param InDepth bounding box depth
 	 */
 	void SetDimensions(
-		float InArea, 
-		float InWidth, 
-		float InHeight, 
+		float InArea,
+		float InWidth,
+		float InHeight,
 		float InDepth
 	)
 	{
@@ -495,7 +495,7 @@ public:
 	}
 
 	/** Get the bounding box height */
-	float GetHeight() const 
+	float GetHeight() const
 	{
 		return GetDatasmithMeshElement()->GetHeight();
 	}
@@ -508,7 +508,7 @@ public:
 
 	/** Get the UV channel that will be used for the lightmap */
 	int32 GetLightmapCoordinateIndex() const
-	{ 
+	{
 		return GetDatasmithMeshElement()->GetLightmapCoordinateIndex();
 	}
 
@@ -522,7 +522,7 @@ public:
 	}
 
 	/** Get the source UV channel that will be used at import to generate the lightmap UVs */
-	int32 GetLightmapSourceUV() const 
+	int32 GetLightmapSourceUV() const
 	{
 		return GetDatasmithMeshElement()->GetLightmapSourceUV();
 	}
diff --git a/Engine/Source/Developer/Datasmith/DatasmithFacade/Public/DatasmithFacadeVariant.h b/Engine/Source/Developer/Datasmith/DatasmithFacade/Public/DatasmithFacadeVariant.h
--- a/Engine/Source/Developer/Datasmith/DatasmithFacade/Public/DatasmithFacadeVariant.h
+++ b/Engine/Source/Developer/Datasmith/DatasmithFacade/Public/DatasmithFacadeVariant.h
@@ -94,35 +94,48 @@ public:
 	);
 
 	void AddRelativeLocationCapture(
-		float X,
-		float Y,
-		float Z
+		double X,
+		double Y,
+		double Z
 	);
 
 	void AddRelativeRotationCapture(
-		float Pitch,
-		float Yaw,
-		float Roll
+		double Pitch,
+		double Yaw,
+		double Roll
 	);
 
 	void AddRelativeRotationCapture(
-		float X,
-		float Y,
-		float Z,
-		float W
+		double X,
+		double Y,
+		double Z,
+		double W
 	);
 
 	void AddRelativeScaleCapture(
-		float X,
-		float Y,
-		float Z
+		double X,
+		double Y,
+		double Z
 	);
 
 	void AddRelativeTransformCapture(
-		const float InMatrix[16],
+		const double InMatrix[16],
 		bool bRowMajor = false
 	);
 
+	void AddRelativeTransformCapture(
+		const float InMatrix[16],
+		bool bRowMajor = false
+	)
+	{
+		double InMatrixD[16];
+		for(int I = 0; I < 16; ++I)
+		{
+			InMatrixD[I] = InMatrix[I];
+		}
+		AddRelativeTransformCapture(InMatrixD, bRowMajor);
+	}
+
 	void AddVisibilityCapture(
 		bool bInVisibility
 	);
diff --git a/Engine/Source/Programs/Enterprise/Datasmith/DatasmithARCHICADExporter/Private/DatasmithHashTools.cpp b/Engine/Source/Programs/Enterprise/Datasmith/DatasmithARCHICADExporter/Private/DatasmithHashTools.cpp
--- a/Engine/Source/Programs/Enterprise/Datasmith/DatasmithARCHICADExporter/Private/DatasmithHashTools.cpp
+++ b/Engine/Source/Programs/Enterprise/Datasmith/DatasmithARCHICADExporter/Private/DatasmithHashTools.cpp
@@ -159,7 +159,7 @@ void FDatasmithHashTools::ComputeDatasmithMeshHash(const FDatasmithMesh& Mesh)
 	TUpdate(VerticesCount);
 	for (int32 IdxVertice = 0; IdxVertice < VerticesCount; ++IdxVertice)
 	{
-		HashFixVector(Mesh.GetVertex(IdxVertice));
+		HashFixVector(FVector(Mesh.GetVertex(IdxVertice)));
 		TUpdate(Mesh.GetVertexColor(IdxVertice));
 	}
 
@@ -193,7 +193,7 @@ void FDatasmithHashTools::ComputeDatasmithMeshHash(const FDatasmithMesh& Mesh)
 
 		for (int32 IdxComponent = 0; IdxComponent < 3; IdxComponent++)
 		{
-			HashFloatVector(Mesh.GetNormal(IdxFace * 3 + IdxComponent));
+			HashFloatVector(FVector(Mesh.GetNormal(IdxFace * 3 + IdxComponent)));
 			TUpdate(Mesh.GetFaceSmoothingMask(IdxFace * 3 + IdxComponent));
 		}
 
diff --git a/Engine/Source/Programs/Enterprise/Datasmith/DatasmithARCHICADExporter/Private/Element2StaticMesh.cpp b/Engine/Source/Programs/Enterprise/Datasmith/DatasmithARCHICADExporter/Private/Element2StaticMesh.cpp
--- a/Engine/Source/Programs/Enterprise/Datasmith/DatasmithARCHICADExporter/Private/Element2StaticMesh.cpp
+++ b/Engine/Source/Programs/Enterprise/Datasmith/DatasmithARCHICADExporter/Private/Element2StaticMesh.cpp
@@ -83,7 +83,7 @@ utf8_string FElement2StaticMesh::MeshAsString(const FDatasmithMesh& Mesh)
 	Dump += Utf8StringFormat("\tVertices Count = %d\n", VerticesCount);
 	for (int32 IdxVertice = 0; IdxVertice < VerticesCount; ++IdxVertice)
 	{
-		FVector Vertex = Mesh.GetVertex(IdxVertice);
+		FVector3f Vertex = Mesh.GetVertex(IdxVertice);
 		Dump += Utf8StringFormat("\t\tVertice[%d] = {%f, %f, %f}\n", IdxVertice, Vertex.X, Vertex.Y, Vertex.Z);
 	}
 
@@ -111,15 +111,15 @@ utf8_string FElement2StaticMesh::MeshAsString(const FDatasmithMesh& Mesh)
 		Mesh.GetFace(IdxFace, Vertex1, Vertex2, Vertex3, MaterialId);
 		Dump += Utf8StringFormat("\t\tVertex[%d] = {%d, %d, %d} Mat = %d\n", IdxFace, Vertex1, Vertex2, Vertex3,
 								 MaterialId);
-		FVector Point1 = Mesh.GetVertex(Vertex1);
-		FVector Point2 = Mesh.GetVertex(Vertex2);
-		FVector Point3 = Mesh.GetVertex(Vertex3);
+		FVector3f Point1 = Mesh.GetVertex(Vertex1);
+		FVector3f Point2 = Mesh.GetVertex(Vertex2);
+		FVector3f Point3 = Mesh.GetVertex(Vertex3);
 		Dump += Utf8StringFormat("\t\t\t\t{{%f, %f, %f}, {%f, %f, %f}, {%f, %f, %f}}\n", Point1.X, Point1.Y, Point1.Z,
 								 Point2.X, Point2.Y, Point2.Z, Point3.X, Point3.Y, Point3.Z);
 
 		for (int32 IdxComponent = 0; IdxComponent < 3; IdxComponent++)
 		{
-			FVector Normal = Mesh.GetNormal(IdxFace * 3 + IdxComponent);
+			FVector3f Normal = Mesh.GetNormal(IdxFace * 3 + IdxComponent);
 			Dump += Utf8StringFormat("\t\t\tNormal[%d][%d] = {%f, %f, %f}\n", IdxFace, IdxComponent, Normal.X, Normal.Y,
 									 Normal.Z);
 		}
@@ -156,7 +156,7 @@ utf8_string FElement2StaticMesh::MeshElementAsString(const IDatasmithMeshElement
 	Dump += Utf8StringFormat("Mesh \"%s\"\n", TCHAR_TO_UTF8(Mesh.GetName()));
 	Dump += Utf8StringFormat("\tLabel = \"%s\"\n", TCHAR_TO_UTF8(Mesh.GetLabel()));
 	Dump += Utf8StringFormat("\tFile = \"%s\"\n", TCHAR_TO_UTF8(Mesh.GetFile()));
-	const FVector Dim = Mesh.GetDimensions();
+	const FVector3f Dim = Mesh.GetDimensions();
 	Dump += Utf8StringFormat("\tDimensions = {%f, %f, %f}\n", Dim.X, Dim.Y, Dim.Z);
 	Dump += Utf8StringFormat("\tArea = %f\n", Mesh.GetArea());
 	Dump +=
@@ -253,7 +253,7 @@ void FElement2StaticMesh::AddVertex(GS::Int32 InBodyVertex, const Geometry::Vect
 	}
 
 	Geometry::Vector3D VertexWorldNormal = bIsIdentity ? VertexNormal : Matrix * VertexNormal;
-	FVector CurrentNormal(float(VertexWorldNormal.x), -float(VertexWorldNormal.y), float(VertexWorldNormal.z));
+	FVector3f CurrentNormal(float(VertexWorldNormal.x), -float(VertexWorldNormal.y), float(VertexWorldNormal.z));
 
 	// Create triangles
 	if (VertexCount == 0)
@@ -502,8 +502,8 @@ void FElement2StaticMesh::FillMesh(FDatasmithMesh* OutMesh)
 
 			for (int32 IndexComponent = 0; IndexComponent < 3; IndexComponent++)
 			{
-				const FVector& Normal = triangle.Normals[IndexComponent];
-				OutMesh->SetNormal(IndexFace * 3 + IndexComponent, (float)Normal.X, (float)Normal.Y, (float)Normal.Z);
+				const FVector3f& Normal = triangle.Normals[IndexComponent];
+				OutMesh->SetNormal(IndexFace * 3 + IndexComponent, Normal.X, Normal.Y, Normal.Z);
 			}
 
 			OutMesh->SetFaceUV(IndexFace, UVChannel, triangle.UV0, triangle.UV1, triangle.UV2);
diff --git a/Engine/Source/Programs/Enterprise/Datasmith/DatasmithARCHICADExporter/Private/Element2StaticMesh.h b/Engine/Source/Programs/Enterprise/Datasmith/DatasmithARCHICADExporter/Private/Element2StaticMesh.h
--- a/Engine/Source/Programs/Enterprise/Datasmith/DatasmithARCHICADExporter/Private/Element2StaticMesh.h
+++ b/Engine/Source/Programs/Enterprise/Datasmith/DatasmithARCHICADExporter/Private/Element2StaticMesh.h
@@ -116,7 +116,7 @@ class FElement2StaticMesh
 		int		UV0 = kInvalidIndex;
 		int		UV1 = kInvalidIndex;
 		int		UV2 = kInvalidIndex;
-		FVector Normals[3] = {};
+		FVector3f Normals[3] = {};
 
 		// Material
 		int	 LocalMatID = 0;
diff --git a/Engine/Source/Programs/Enterprise/Datasmith/DatasmithFacadeCSharp/Private/DatasmithFacadeCSharp.cpp b/Engine/Source/Programs/Enterprise/Datasmith/DatasmithFacadeCSharp/Private/DatasmithFacadeCSharp.cpp
--- a/Engine/Source/Programs/Enterprise/Datasmith/DatasmithFacadeCSharp/Private/DatasmithFacadeCSharp.cpp
+++ b/Engine/Source/Programs/Enterprise/Datasmith/DatasmithFacadeCSharp/Private/DatasmithFacadeCSharp.cpp
@@ -2867,18 +2867,18 @@ SWIGEXPORT unsigned int SWIGSTDCALL CSharp_FDatasmithFacadeMesh_GetHashForUVChan
 }
 
 
-SWIGEXPORT void SWIGSTDCALL CSharp_FDatasmithFacadeMesh_GetUV(void * jarg1, int jarg2, int jarg3, float * jarg4, float * jarg5) {
+SWIGEXPORT void SWIGSTDCALL CSharp_FDatasmithFacadeMesh_GetUV(void * jarg1, int jarg2, int jarg3, double * jarg4, double * jarg5) {
   FDatasmithFacadeMesh *arg1 = (FDatasmithFacadeMesh *) 0 ;
   int32 arg2 ;
   int32 arg3 ;
-  float *arg4 = 0 ;
-  float *arg5 = 0 ;
+  double *arg4 = 0 ;
+  double *arg5 = 0 ;
   
   arg1 = (FDatasmithFacadeMesh *)jarg1; 
   arg2 = (int32)jarg2; 
   arg3 = (int32)jarg3; 
-  arg4 = (float *)jarg4; 
-  arg5 = (float *)jarg5; 
+  arg4 = (double *)jarg4; 
+  arg5 = (double *)jarg5; 
   ((FDatasmithFacadeMesh const *)arg1)->GetUV(arg2,arg3,*arg4,*arg5);
 }
 
@@ -3455,7 +3455,29 @@ SWIGEXPORT void SWIGSTDCALL CSharp_delete_FDatasmithFacadeActor(void * jarg1) {
 }
 
 
-SWIGEXPORT void SWIGSTDCALL CSharp_FDatasmithFacadeActor_SetWorldTransform__SWIG_0(void * jarg1, float * jarg2, unsigned int jarg3) {
+SWIGEXPORT void SWIGSTDCALL CSharp_FDatasmithFacadeActor_SetWorldTransform__SWIG_0(void * jarg1, double * jarg2, unsigned int jarg3) {
+  FDatasmithFacadeActor *arg1 = (FDatasmithFacadeActor *) 0 ;
+  double *arg2 ;
+  bool arg3 ;
+  
+  arg1 = (FDatasmithFacadeActor *)jarg1; 
+  arg2 = (double *)jarg2; 
+  arg3 = jarg3 ? true : false; 
+  (arg1)->SetWorldTransform((double const (*))arg2,arg3);
+}
+
+
+SWIGEXPORT void SWIGSTDCALL CSharp_FDatasmithFacadeActor_SetWorldTransform__SWIG_1(void * jarg1, double * jarg2) {
+  FDatasmithFacadeActor *arg1 = (FDatasmithFacadeActor *) 0 ;
+  double *arg2 ;
+  
+  arg1 = (FDatasmithFacadeActor *)jarg1; 
+  arg2 = (double *)jarg2; 
+  (arg1)->SetWorldTransform((double const (*))arg2);
+}
+
+
+SWIGEXPORT void SWIGSTDCALL CSharp_FDatasmithFacadeActor_SetWorldTransform__SWIG_2(void * jarg1, float * jarg2, unsigned int jarg3) {
   FDatasmithFacadeActor *arg1 = (FDatasmithFacadeActor *) 0 ;
   float *arg2 ;
   bool arg3 ;
@@ -3467,7 +3489,7 @@ SWIGEXPORT void SWIGSTDCALL CSharp_FDatasmithFacadeActor_SetWorldTransform__SWIG
 }
 
 
-SWIGEXPORT void SWIGSTDCALL CSharp_FDatasmithFacadeActor_SetWorldTransform__SWIG_1(void * jarg1, float * jarg2) {
+SWIGEXPORT void SWIGSTDCALL CSharp_FDatasmithFacadeActor_SetWorldTransform__SWIG_3(void * jarg1, float * jarg2) {
   FDatasmithFacadeActor *arg1 = (FDatasmithFacadeActor *) 0 ;
   float *arg2 ;
   
@@ -3477,118 +3499,118 @@ SWIGEXPORT void SWIGSTDCALL CSharp_FDatasmithFacadeActor_SetWorldTransform__SWIG
 }
 
 
-SWIGEXPORT void SWIGSTDCALL CSharp_FDatasmithFacadeActor_SetScale(void * jarg1, float jarg2, float jarg3, float jarg4) {
+SWIGEXPORT void SWIGSTDCALL CSharp_FDatasmithFacadeActor_SetScale(void * jarg1, double jarg2, double jarg3, double jarg4) {
   FDatasmithFacadeActor *arg1 = (FDatasmithFacadeActor *) 0 ;
-  float arg2 ;
-  float arg3 ;
-  float arg4 ;
+  double arg2 ;
+  double arg3 ;
+  double arg4 ;
   
   arg1 = (FDatasmithFacadeActor *)jarg1; 
-  arg2 = (float)jarg2; 
-  arg3 = (float)jarg3; 
-  arg4 = (float)jarg4; 
+  arg2 = (double)jarg2; 
+  arg3 = (double)jarg3; 
+  arg4 = (double)jarg4; 
   (arg1)->SetScale(arg2,arg3,arg4);
 }
 
 
-SWIGEXPORT void SWIGSTDCALL CSharp_FDatasmithFacadeActor_GetScale(void * jarg1, float * jarg2, float * jarg3, float * jarg4) {
+SWIGEXPORT void SWIGSTDCALL CSharp_FDatasmithFacadeActor_GetScale(void * jarg1, double * jarg2, double * jarg3, double * jarg4) {
   FDatasmithFacadeActor *arg1 = (FDatasmithFacadeActor *) 0 ;
-  float *arg2 = 0 ;
-  float *arg3 = 0 ;
-  float *arg4 = 0 ;
+  double *arg2 = 0 ;
+  double *arg3 = 0 ;
+  double *arg4 = 0 ;
   
   arg1 = (FDatasmithFacadeActor *)jarg1; 
-  arg2 = (float *)jarg2; 
-  arg3 = (float *)jarg3; 
-  arg4 = (float *)jarg4; 
+  arg2 = (double *)jarg2; 
+  arg3 = (double *)jarg3; 
+  arg4 = (double *)jarg4; 
   ((FDatasmithFacadeActor const *)arg1)->GetScale(*arg2,*arg3,*arg4);
 }
 
 
-SWIGEXPORT void SWIGSTDCALL CSharp_FDatasmithFacadeActor_SetRotation__SWIG_0(void * jarg1, float jarg2, float jarg3, float jarg4) {
+SWIGEXPORT void SWIGSTDCALL CSharp_FDatasmithFacadeActor_SetRotation__SWIG_0(void * jarg1, double jarg2, double jarg3, double jarg4) {
   FDatasmithFacadeActor *arg1 = (FDatasmithFacadeActor *) 0 ;
-  float arg2 ;
-  float arg3 ;
-  float arg4 ;
+  double arg2 ;
+  double arg3 ;
+  double arg4 ;
   
   arg1 = (FDatasmithFacadeActor *)jarg1; 
-  arg2 = (float)jarg2; 
-  arg3 = (float)jarg3; 
-  arg4 = (float)jarg4; 
+  arg2 = (double)jarg2; 
+  arg3 = (double)jarg3; 
+  arg4 = (double)jarg4; 
   (arg1)->SetRotation(arg2,arg3,arg4);
 }
 
 
-SWIGEXPORT void SWIGSTDCALL CSharp_FDatasmithFacadeActor_GetRotation__SWIG_0(void * jarg1, float * jarg2, float * jarg3, float * jarg4) {
+SWIGEXPORT void SWIGSTDCALL CSharp_FDatasmithFacadeActor_GetRotation__SWIG_0(void * jarg1, double * jarg2, double * jarg3, double * jarg4) {
   FDatasmithFacadeActor *arg1 = (FDatasmithFacadeActor *) 0 ;
-  float *arg2 = 0 ;
-  float *arg3 = 0 ;
-  float *arg4 = 0 ;
+  double *arg2 = 0 ;
+  double *arg3 = 0 ;
+  double *arg4 = 0 ;
   
   arg1 = (FDatasmithFacadeActor *)jarg1; 
-  arg2 = (float *)jarg2; 
-  arg3 = (float *)jarg3; 
-  arg4 = (float *)jarg4; 
+  arg2 = (double *)jarg2; 
+  arg3 = (double *)jarg3; 
+  arg4 = (double *)jarg4; 
   ((FDatasmithFacadeActor const *)arg1)->GetRotation(*arg2,*arg3,*arg4);
 }
 
 
-SWIGEXPORT void SWIGSTDCALL CSharp_FDatasmithFacadeActor_SetRotation__SWIG_1(void * jarg1, float jarg2, float jarg3, float jarg4, float jarg5) {
+SWIGEXPORT void SWIGSTDCALL CSharp_FDatasmithFacadeActor_SetRotation__SWIG_1(void * jarg1, double jarg2, double jarg3, double jarg4, double jarg5) {
   FDatasmithFacadeActor *arg1 = (FDatasmithFacadeActor *) 0 ;
-  float arg2 ;
-  float arg3 ;
-  float arg4 ;
-  float arg5 ;
+  double arg2 ;
+  double arg3 ;
+  double arg4 ;
+  double arg5 ;
   
   arg1 = (FDatasmithFacadeActor *)jarg1; 
-  arg2 = (float)jarg2; 
-  arg3 = (float)jarg3; 
-  arg4 = (float)jarg4; 
-  arg5 = (float)jarg5; 
+  arg2 = (double)jarg2; 
+  arg3 = (double)jarg3; 
+  arg4 = (double)jarg4; 
+  arg5 = (double)jarg5; 
   (arg1)->SetRotation(arg2,arg3,arg4,arg5);
 }
 
 
-SWIGEXPORT void SWIGSTDCALL CSharp_FDatasmithFacadeActor_GetRotation__SWIG_1(void * jarg1, float * jarg2, float * jarg3, float * jarg4, float * jarg5) {
+SWIGEXPORT void SWIGSTDCALL CSharp_FDatasmithFacadeActor_GetRotation__SWIG_1(void * jarg1, double * jarg2, double * jarg3, double * jarg4, double * jarg5) {
   FDatasmithFacadeActor *arg1 = (FDatasmithFacadeActor *) 0 ;
-  float *arg2 = 0 ;
-  float *arg3 = 0 ;
-  float *arg4 = 0 ;
-  float *arg5 = 0 ;
+  double *arg2 = 0 ;
+  double *arg3 = 0 ;
+  double *arg4 = 0 ;
+  double *arg5 = 0 ;
   
   arg1 = (FDatasmithFacadeActor *)jarg1; 
-  arg2 = (float *)jarg2; 
-  arg3 = (float *)jarg3; 
-  arg4 = (float *)jarg4; 
-  arg5 = (float *)jarg5; 
+  arg2 = (double *)jarg2; 
+  arg3 = (double *)jarg3; 
+  arg4 = (double *)jarg4; 
+  arg5 = (double *)jarg5; 
   ((FDatasmithFacadeActor const *)arg1)->GetRotation(*arg2,*arg3,*arg4,*arg5);
 }
 
 
-SWIGEXPORT void SWIGSTDCALL CSharp_FDatasmithFacadeActor_SetTranslation(void * jarg1, float jarg2, float jarg3, float jarg4) {
+SWIGEXPORT void SWIGSTDCALL CSharp_FDatasmithFacadeActor_SetTranslation(void * jarg1, double jarg2, double jarg3, double jarg4) {
   FDatasmithFacadeActor *arg1 = (FDatasmithFacadeActor *) 0 ;
-  float arg2 ;
-  float arg3 ;
-  float arg4 ;
+  double arg2 ;
+  double arg3 ;
+  double arg4 ;
   
   arg1 = (FDatasmithFacadeActor *)jarg1; 
-  arg2 = (float)jarg2; 
-  arg3 = (float)jarg3; 
-  arg4 = (float)jarg4; 
+  arg2 = (double)jarg2; 
+  arg3 = (double)jarg3; 
+  arg4 = (double)jarg4; 
   (arg1)->SetTranslation(arg2,arg3,arg4);
 }
 
 
-SWIGEXPORT void SWIGSTDCALL CSharp_FDatasmithFacadeActor_GetTranslation(void * jarg1, float * jarg2, float * jarg3, float * jarg4) {
+SWIGEXPORT void SWIGSTDCALL CSharp_FDatasmithFacadeActor_GetTranslation(void * jarg1, double * jarg2, double * jarg3, double * jarg4) {
   FDatasmithFacadeActor *arg1 = (FDatasmithFacadeActor *) 0 ;
-  float *arg2 = 0 ;
-  float *arg3 = 0 ;
-  float *arg4 = 0 ;
+  double *arg2 = 0 ;
+  double *arg3 = 0 ;
+  double *arg4 = 0 ;
   
   arg1 = (FDatasmithFacadeActor *)jarg1; 
-  arg2 = (float *)jarg2; 
-  arg3 = (float *)jarg3; 
-  arg4 = (float *)jarg4; 
+  arg2 = (double *)jarg2; 
+  arg3 = (double *)jarg3; 
+  arg4 = (double *)jarg4; 
   ((FDatasmithFacadeActor const *)arg1)->GetTranslation(*arg2,*arg3,*arg4);
 }
 
@@ -3815,36 +3837,36 @@ SWIGEXPORT void SWIGSTDCALL CSharp_delete_FDatasmithFacadeActorCamera(void * jar
 }
 
 
-SWIGEXPORT void SWIGSTDCALL CSharp_FDatasmithFacadeActorCamera_SetCameraPosition(void * jarg1, float jarg2, float jarg3, float jarg4) {
+SWIGEXPORT void SWIGSTDCALL CSharp_FDatasmithFacadeActorCamera_SetCameraPosition(void * jarg1, double jarg2, double jarg3, double jarg4) {
   FDatasmithFacadeActorCamera *arg1 = (FDatasmithFacadeActorCamera *) 0 ;
-  float arg2 ;
-  float arg3 ;
-  float arg4 ;
+  double arg2 ;
+  double arg3 ;
+  double arg4 ;
   
   arg1 = (FDatasmithFacadeActorCamera *)jarg1; 
-  arg2 = (float)jarg2; 
-  arg3 = (float)jarg3; 
-  arg4 = (float)jarg4; 
+  arg2 = (double)jarg2; 
+  arg3 = (double)jarg3; 
+  arg4 = (double)jarg4; 
   (arg1)->SetCameraPosition(arg2,arg3,arg4);
 }
 
 
-SWIGEXPORT void SWIGSTDCALL CSharp_FDatasmithFacadeActorCamera_SetCameraRotation(void * jarg1, float jarg2, float jarg3, float jarg4, float jarg5, float jarg6, float jarg7) {
+SWIGEXPORT void SWIGSTDCALL CSharp_FDatasmithFacadeActorCamera_SetCameraRotation(void * jarg1, double jarg2, double jarg3, double jarg4, double jarg5, double jarg6, double jarg7) {
   FDatasmithFacadeActorCamera *arg1 = (FDatasmithFacadeActorCamera *) 0 ;
-  float arg2 ;
-  float arg3 ;
-  float arg4 ;
-  float arg5 ;
-  float arg6 ;
-  float arg7 ;
+  double arg2 ;
+  double arg3 ;
+  double arg4 ;
+  double arg5 ;
+  double arg6 ;
+  double arg7 ;
   
   arg1 = (FDatasmithFacadeActorCamera *)jarg1; 
-  arg2 = (float)jarg2; 
-  arg3 = (float)jarg3; 
-  arg4 = (float)jarg4; 
-  arg5 = (float)jarg5; 
-  arg6 = (float)jarg6; 
-  arg7 = (float)jarg7; 
+  arg2 = (double)jarg2; 
+  arg3 = (double)jarg3; 
+  arg4 = (double)jarg4; 
+  arg5 = (double)jarg5; 
+  arg6 = (double)jarg6; 
+  arg7 = (double)jarg7; 
   (arg1)->SetCameraRotation(arg2,arg3,arg4,arg5,arg6,arg7);
 }
 
@@ -6929,65 +6951,87 @@ SWIGEXPORT void SWIGSTDCALL CSharp_FDatasmithFacadeActorBinding_RemovePropertyCa
 }
 
 
-SWIGEXPORT void SWIGSTDCALL CSharp_FDatasmithFacadeActorBinding_AddRelativeLocationCapture(void * jarg1, float jarg2, float jarg3, float jarg4) {
+SWIGEXPORT void SWIGSTDCALL CSharp_FDatasmithFacadeActorBinding_AddRelativeLocationCapture(void * jarg1, double jarg2, double jarg3, double jarg4) {
   FDatasmithFacadeActorBinding *arg1 = (FDatasmithFacadeActorBinding *) 0 ;
-  float arg2 ;
-  float arg3 ;
-  float arg4 ;
+  double arg2 ;
+  double arg3 ;
+  double arg4 ;
   
   arg1 = (FDatasmithFacadeActorBinding *)jarg1; 
-  arg2 = (float)jarg2; 
-  arg3 = (float)jarg3; 
-  arg4 = (float)jarg4; 
+  arg2 = (double)jarg2; 
+  arg3 = (double)jarg3; 
+  arg4 = (double)jarg4; 
   (arg1)->AddRelativeLocationCapture(arg2,arg3,arg4);
 }
 
 
-SWIGEXPORT void SWIGSTDCALL CSharp_FDatasmithFacadeActorBinding_AddRelativeRotationCapture__SWIG_0(void * jarg1, float jarg2, float jarg3, float jarg4) {
+SWIGEXPORT void SWIGSTDCALL CSharp_FDatasmithFacadeActorBinding_AddRelativeRotationCapture__SWIG_0(void * jarg1, double jarg2, double jarg3, double jarg4) {
   FDatasmithFacadeActorBinding *arg1 = (FDatasmithFacadeActorBinding *) 0 ;
-  float arg2 ;
-  float arg3 ;
-  float arg4 ;
+  double arg2 ;
+  double arg3 ;
+  double arg4 ;
   
   arg1 = (FDatasmithFacadeActorBinding *)jarg1; 
-  arg2 = (float)jarg2; 
-  arg3 = (float)jarg3; 
-  arg4 = (float)jarg4; 
+  arg2 = (double)jarg2; 
+  arg3 = (double)jarg3; 
+  arg4 = (double)jarg4; 
   (arg1)->AddRelativeRotationCapture(arg2,arg3,arg4);
 }
 
 
-SWIGEXPORT void SWIGSTDCALL CSharp_FDatasmithFacadeActorBinding_AddRelativeRotationCapture__SWIG_1(void * jarg1, float jarg2, float jarg3, float jarg4, float jarg5) {
+SWIGEXPORT void SWIGSTDCALL CSharp_FDatasmithFacadeActorBinding_AddRelativeRotationCapture__SWIG_1(void * jarg1, double jarg2, double jarg3, double jarg4, double jarg5) {
   FDatasmithFacadeActorBinding *arg1 = (FDatasmithFacadeActorBinding *) 0 ;
-  float arg2 ;
-  float arg3 ;
-  float arg4 ;
-  float arg5 ;
+  double arg2 ;
+  double arg3 ;
+  double arg4 ;
+  double arg5 ;
   
   arg1 = (FDatasmithFacadeActorBinding *)jarg1; 
-  arg2 = (float)jarg2; 
-  arg3 = (float)jarg3; 
-  arg4 = (float)jarg4; 
-  arg5 = (float)jarg5; 
+  arg2 = (double)jarg2; 
+  arg3 = (double)jarg3; 
+  arg4 = (double)jarg4; 
+  arg5 = (double)jarg5; 
   (arg1)->AddRelativeRotationCapture(arg2,arg3,arg4,arg5);
 }
 
 
-SWIGEXPORT void SWIGSTDCALL CSharp_FDatasmithFacadeActorBinding_AddRelativeScaleCapture(void * jarg1, float jarg2, float jarg3, float jarg4) {
+SWIGEXPORT void SWIGSTDCALL CSharp_FDatasmithFacadeActorBinding_AddRelativeScaleCapture(void * jarg1, double jarg2, double jarg3, double jarg4) {
   FDatasmithFacadeActorBinding *arg1 = (FDatasmithFacadeActorBinding *) 0 ;
-  float arg2 ;
-  float arg3 ;
-  float arg4 ;
+  double arg2 ;
+  double arg3 ;
+  double arg4 ;
   
   arg1 = (FDatasmithFacadeActorBinding *)jarg1; 
-  arg2 = (float)jarg2; 
-  arg3 = (float)jarg3; 
-  arg4 = (float)jarg4; 
+  arg2 = (double)jarg2; 
+  arg3 = (double)jarg3; 
+  arg4 = (double)jarg4; 
   (arg1)->AddRelativeScaleCapture(arg2,arg3,arg4);
 }
 
 
-SWIGEXPORT void SWIGSTDCALL CSharp_FDatasmithFacadeActorBinding_AddRelativeTransformCapture__SWIG_0(void * jarg1, float * jarg2, unsigned int jarg3) {
+SWIGEXPORT void SWIGSTDCALL CSharp_FDatasmithFacadeActorBinding_AddRelativeTransformCapture__SWIG_0(void * jarg1, double * jarg2, unsigned int jarg3) {
+  FDatasmithFacadeActorBinding *arg1 = (FDatasmithFacadeActorBinding *) 0 ;
+  double *arg2 ;
+  bool arg3 ;
+  
+  arg1 = (FDatasmithFacadeActorBinding *)jarg1; 
+  arg2 = (double *)jarg2; 
+  arg3 = jarg3 ? true : false; 
+  (arg1)->AddRelativeTransformCapture((double const (*))arg2,arg3);
+}
+
+
+SWIGEXPORT void SWIGSTDCALL CSharp_FDatasmithFacadeActorBinding_AddRelativeTransformCapture__SWIG_1(void * jarg1, double * jarg2) {
+  FDatasmithFacadeActorBinding *arg1 = (FDatasmithFacadeActorBinding *) 0 ;
+  double *arg2 ;
+  
+  arg1 = (FDatasmithFacadeActorBinding *)jarg1; 
+  arg2 = (double *)jarg2; 
+  (arg1)->AddRelativeTransformCapture((double const (*))arg2);
+}
+
+
+SWIGEXPORT void SWIGSTDCALL CSharp_FDatasmithFacadeActorBinding_AddRelativeTransformCapture__SWIG_2(void * jarg1, float * jarg2, unsigned int jarg3) {
   FDatasmithFacadeActorBinding *arg1 = (FDatasmithFacadeActorBinding *) 0 ;
   float *arg2 ;
   bool arg3 ;
@@ -6999,7 +7043,7 @@ SWIGEXPORT void SWIGSTDCALL CSharp_FDatasmithFacadeActorBinding_AddRelativeTrans
 }
 
 
-SWIGEXPORT void SWIGSTDCALL CSharp_FDatasmithFacadeActorBinding_AddRelativeTransformCapture__SWIG_1(void * jarg1, float * jarg2) {
+SWIGEXPORT void SWIGSTDCALL CSharp_FDatasmithFacadeActorBinding_AddRelativeTransformCapture__SWIG_3(void * jarg1, float * jarg2) {
   FDatasmithFacadeActorBinding *arg1 = (FDatasmithFacadeActorBinding *) 0 ;
   float *arg2 ;
   
diff --git a/Engine/Source/Programs/Enterprise/Datasmith/DatasmithFacadeCSharp/Public/DatasmithFacadeCSharpPINVOKE.cs b/Engine/Source/Programs/Enterprise/Datasmith/DatasmithFacadeCSharp/Public/DatasmithFacadeCSharpPINVOKE.cs
--- a/Engine/Source/Programs/Enterprise/Datasmith/DatasmithFacadeCSharp/Public/DatasmithFacadeCSharpPINVOKE.cs
+++ b/Engine/Source/Programs/Enterprise/Datasmith/DatasmithFacadeCSharp/Public/DatasmithFacadeCSharpPINVOKE.cs
@@ -864,7 +864,7 @@ class DatasmithFacadeCSharpPINVOKE {
   public static extern uint FDatasmithFacadeMesh_GetHashForUVChannel(global::System.Runtime.InteropServices.HandleRef jarg1, int jarg2);
 
   [global::System.Runtime.InteropServices.DllImport("DatasmithFacadeCSharp", EntryPoint="CSharp_FDatasmithFacadeMesh_GetUV")]
-  public static extern void FDatasmithFacadeMesh_GetUV(global::System.Runtime.InteropServices.HandleRef jarg1, int jarg2, int jarg3, out float jarg4, out float jarg5);
+  public static extern void FDatasmithFacadeMesh_GetUV(global::System.Runtime.InteropServices.HandleRef jarg1, int jarg2, int jarg3, out double jarg4, out double jarg5);
 
   [global::System.Runtime.InteropServices.DllImport("DatasmithFacadeCSharp", EntryPoint="CSharp_FDatasmithFacadeMesh_SetFaceUV")]
   public static extern void FDatasmithFacadeMesh_SetFaceUV(global::System.Runtime.InteropServices.HandleRef jarg1, int jarg2, int jarg3, int jarg4, int jarg5, int jarg6);
@@ -1008,34 +1008,40 @@ class DatasmithFacadeCSharpPINVOKE {
   public static extern void delete_FDatasmithFacadeActor(global::System.Runtime.InteropServices.HandleRef jarg1);
 
   [global::System.Runtime.InteropServices.DllImport("DatasmithFacadeCSharp", EntryPoint="CSharp_FDatasmithFacadeActor_SetWorldTransform__SWIG_0")]
-  public static extern void FDatasmithFacadeActor_SetWorldTransform__SWIG_0(global::System.Runtime.InteropServices.HandleRef jarg1, [global::System.Runtime.InteropServices.In, global::System.Runtime.InteropServices.MarshalAs(global::System.Runtime.InteropServices.UnmanagedType.LPArray)] float[] jarg2, bool jarg3);
+  public static extern void FDatasmithFacadeActor_SetWorldTransform__SWIG_0(global::System.Runtime.InteropServices.HandleRef jarg1, [global::System.Runtime.InteropServices.In, global::System.Runtime.InteropServices.MarshalAs(global::System.Runtime.InteropServices.UnmanagedType.LPArray)] double[] jarg2, bool jarg3);
 
   [global::System.Runtime.InteropServices.DllImport("DatasmithFacadeCSharp", EntryPoint="CSharp_FDatasmithFacadeActor_SetWorldTransform__SWIG_1")]
-  public static extern void FDatasmithFacadeActor_SetWorldTransform__SWIG_1(global::System.Runtime.InteropServices.HandleRef jarg1, [global::System.Runtime.InteropServices.In, global::System.Runtime.InteropServices.MarshalAs(global::System.Runtime.InteropServices.UnmanagedType.LPArray)] float[] jarg2);
+  public static extern void FDatasmithFacadeActor_SetWorldTransform__SWIG_1(global::System.Runtime.InteropServices.HandleRef jarg1, [global::System.Runtime.InteropServices.In, global::System.Runtime.InteropServices.MarshalAs(global::System.Runtime.InteropServices.UnmanagedType.LPArray)] double[] jarg2);
+
+  [global::System.Runtime.InteropServices.DllImport("DatasmithFacadeCSharp", EntryPoint="CSharp_FDatasmithFacadeActor_SetWorldTransform__SWIG_2")]
+  public static extern void FDatasmithFacadeActor_SetWorldTransform__SWIG_2(global::System.Runtime.InteropServices.HandleRef jarg1, [global::System.Runtime.InteropServices.In, global::System.Runtime.InteropServices.MarshalAs(global::System.Runtime.InteropServices.UnmanagedType.LPArray)] float[] jarg2, bool jarg3);
+
+  [global::System.Runtime.InteropServices.DllImport("DatasmithFacadeCSharp", EntryPoint="CSharp_FDatasmithFacadeActor_SetWorldTransform__SWIG_3")]
+  public static extern void FDatasmithFacadeActor_SetWorldTransform__SWIG_3(global::System.Runtime.InteropServices.HandleRef jarg1, [global::System.Runtime.InteropServices.In, global::System.Runtime.InteropServices.MarshalAs(global::System.Runtime.InteropServices.UnmanagedType.LPArray)] float[] jarg2);
 
   [global::System.Runtime.InteropServices.DllImport("DatasmithFacadeCSharp", EntryPoint="CSharp_FDatasmithFacadeActor_SetScale")]
-  public static extern void FDatasmithFacadeActor_SetScale(global::System.Runtime.InteropServices.HandleRef jarg1, float jarg2, float jarg3, float jarg4);
+  public static extern void FDatasmithFacadeActor_SetScale(global::System.Runtime.InteropServices.HandleRef jarg1, double jarg2, double jarg3, double jarg4);
 
   [global::System.Runtime.InteropServices.DllImport("DatasmithFacadeCSharp", EntryPoint="CSharp_FDatasmithFacadeActor_GetScale")]
-  public static extern void FDatasmithFacadeActor_GetScale(global::System.Runtime.InteropServices.HandleRef jarg1, out float jarg2, out float jarg3, out float jarg4);
+  public static extern void FDatasmithFacadeActor_GetScale(global::System.Runtime.InteropServices.HandleRef jarg1, out double jarg2, out double jarg3, out double jarg4);
 
   [global::System.Runtime.InteropServices.DllImport("DatasmithFacadeCSharp", EntryPoint="CSharp_FDatasmithFacadeActor_SetRotation__SWIG_0")]
-  public static extern void FDatasmithFacadeActor_SetRotation__SWIG_0(global::System.Runtime.InteropServices.HandleRef jarg1, float jarg2, float jarg3, float jarg4);
+  public static extern void FDatasmithFacadeActor_SetRotation__SWIG_0(global::System.Runtime.InteropServices.HandleRef jarg1, double jarg2, double jarg3, double jarg4);
 
   [global::System.Runtime.InteropServices.DllImport("DatasmithFacadeCSharp", EntryPoint="CSharp_FDatasmithFacadeActor_GetRotation__SWIG_0")]
-  public static extern void FDatasmithFacadeActor_GetRotation__SWIG_0(global::System.Runtime.InteropServices.HandleRef jarg1, out float jarg2, out float jarg3, out float jarg4);
+  public static extern void FDatasmithFacadeActor_GetRotation__SWIG_0(global::System.Runtime.InteropServices.HandleRef jarg1, out double jarg2, out double jarg3, out double jarg4);
 
   [global::System.Runtime.InteropServices.DllImport("DatasmithFacadeCSharp", EntryPoint="CSharp_FDatasmithFacadeActor_SetRotation__SWIG_1")]
-  public static extern void FDatasmithFacadeActor_SetRotation__SWIG_1(global::System.Runtime.InteropServices.HandleRef jarg1, float jarg2, float jarg3, float jarg4, float jarg5);
+  public static extern void FDatasmithFacadeActor_SetRotation__SWIG_1(global::System.Runtime.InteropServices.HandleRef jarg1, double jarg2, double jarg3, double jarg4, double jarg5);
 
   [global::System.Runtime.InteropServices.DllImport("DatasmithFacadeCSharp", EntryPoint="CSharp_FDatasmithFacadeActor_GetRotation__SWIG_1")]
-  public static extern void FDatasmithFacadeActor_GetRotation__SWIG_1(global::System.Runtime.InteropServices.HandleRef jarg1, out float jarg2, out float jarg3, out float jarg4, out float jarg5);
+  public static extern void FDatasmithFacadeActor_GetRotation__SWIG_1(global::System.Runtime.InteropServices.HandleRef jarg1, out double jarg2, out double jarg3, out double jarg4, out double jarg5);
 
   [global::System.Runtime.InteropServices.DllImport("DatasmithFacadeCSharp", EntryPoint="CSharp_FDatasmithFacadeActor_SetTranslation")]
-  public static extern void FDatasmithFacadeActor_SetTranslation(global::System.Runtime.InteropServices.HandleRef jarg1, float jarg2, float jarg3, float jarg4);
+  public static extern void FDatasmithFacadeActor_SetTranslation(global::System.Runtime.InteropServices.HandleRef jarg1, double jarg2, double jarg3, double jarg4);
 
   [global::System.Runtime.InteropServices.DllImport("DatasmithFacadeCSharp", EntryPoint="CSharp_FDatasmithFacadeActor_GetTranslation")]
-  public static extern void FDatasmithFacadeActor_GetTranslation(global::System.Runtime.InteropServices.HandleRef jarg1, out float jarg2, out float jarg3, out float jarg4);
+  public static extern void FDatasmithFacadeActor_GetTranslation(global::System.Runtime.InteropServices.HandleRef jarg1, out double jarg2, out double jarg3, out double jarg4);
 
   [global::System.Runtime.InteropServices.DllImport("DatasmithFacadeCSharp", EntryPoint="CSharp_FDatasmithFacadeActor_SetLayer")]
   public static extern void FDatasmithFacadeActor_SetLayer(global::System.Runtime.InteropServices.HandleRef jarg1, [global::System.Runtime.InteropServices.MarshalAs(global::System.Runtime.InteropServices.UnmanagedType.LPWStr)]string jarg2);
@@ -1098,10 +1104,10 @@ class DatasmithFacadeCSharpPINVOKE {
   public static extern void delete_FDatasmithFacadeActorCamera(global::System.Runtime.InteropServices.HandleRef jarg1);
 
   [global::System.Runtime.InteropServices.DllImport("DatasmithFacadeCSharp", EntryPoint="CSharp_FDatasmithFacadeActorCamera_SetCameraPosition")]
-  public static extern void FDatasmithFacadeActorCamera_SetCameraPosition(global::System.Runtime.InteropServices.HandleRef jarg1, float jarg2, float jarg3, float jarg4);
+  public static extern void FDatasmithFacadeActorCamera_SetCameraPosition(global::System.Runtime.InteropServices.HandleRef jarg1, double jarg2, double jarg3, double jarg4);
 
   [global::System.Runtime.InteropServices.DllImport("DatasmithFacadeCSharp", EntryPoint="CSharp_FDatasmithFacadeActorCamera_SetCameraRotation")]
-  public static extern void FDatasmithFacadeActorCamera_SetCameraRotation(global::System.Runtime.InteropServices.HandleRef jarg1, float jarg2, float jarg3, float jarg4, float jarg5, float jarg6, float jarg7);
+  public static extern void FDatasmithFacadeActorCamera_SetCameraRotation(global::System.Runtime.InteropServices.HandleRef jarg1, double jarg2, double jarg3, double jarg4, double jarg5, double jarg6, double jarg7);
 
   [global::System.Runtime.InteropServices.DllImport("DatasmithFacadeCSharp", EntryPoint="CSharp_FDatasmithFacadeActorCamera_SetSensorWidth")]
   public static extern void FDatasmithFacadeActorCamera_SetSensorWidth(global::System.Runtime.InteropServices.HandleRef jarg1, float jarg2);
@@ -1392,7 +1398,7 @@ class DatasmithFacadeCSharpPINVOKE {
   public static extern int FDatasmithFacadeTransformAnimation_GetCurveInterpMode(global::System.Runtime.InteropServices.HandleRef jarg1, int jarg2);
 
   [global::System.Runtime.InteropServices.DllImport("DatasmithFacadeCSharp", EntryPoint="CSharp_FDatasmithFacadeTransformAnimation_GetFrame")]
-  public static extern void FDatasmithFacadeTransformAnimation_GetFrame(global::System.Runtime.InteropServices.HandleRef jarg1, int jarg2, int jarg3, out int jarg4, ref double jarg5, ref double jarg6, ref double jarg7);
+  public static extern void FDatasmithFacadeTransformAnimation_GetFrame(global::System.Runtime.InteropServices.HandleRef jarg1, int jarg2, int jarg3, out int jarg4, out double jarg5, out double jarg6, out double jarg7);
 
   [global::System.Runtime.InteropServices.DllImport("DatasmithFacadeCSharp", EntryPoint="CSharp_FDatasmithFacadeTransformAnimation_RemoveFrame")]
   public static extern void FDatasmithFacadeTransformAnimation_RemoveFrame(global::System.Runtime.InteropServices.HandleRef jarg1, int jarg2, int jarg3);
@@ -1917,22 +1923,28 @@ class DatasmithFacadeCSharpPINVOKE {
   public static extern void FDatasmithFacadeActorBinding_RemovePropertyCapture(global::System.Runtime.InteropServices.HandleRef jarg1, global::System.Runtime.InteropServices.HandleRef jarg2);
 
   [global::System.Runtime.InteropServices.DllImport("DatasmithFacadeCSharp", EntryPoint="CSharp_FDatasmithFacadeActorBinding_AddRelativeLocationCapture")]
-  public static extern void FDatasmithFacadeActorBinding_AddRelativeLocationCapture(global::System.Runtime.InteropServices.HandleRef jarg1, float jarg2, float jarg3, float jarg4);
+  public static extern void FDatasmithFacadeActorBinding_AddRelativeLocationCapture(global::System.Runtime.InteropServices.HandleRef jarg1, double jarg2, double jarg3, double jarg4);
 
   [global::System.Runtime.InteropServices.DllImport("DatasmithFacadeCSharp", EntryPoint="CSharp_FDatasmithFacadeActorBinding_AddRelativeRotationCapture__SWIG_0")]
-  public static extern void FDatasmithFacadeActorBinding_AddRelativeRotationCapture__SWIG_0(global::System.Runtime.InteropServices.HandleRef jarg1, float jarg2, float jarg3, float jarg4);
+  public static extern void FDatasmithFacadeActorBinding_AddRelativeRotationCapture__SWIG_0(global::System.Runtime.InteropServices.HandleRef jarg1, double jarg2, double jarg3, double jarg4);
 
   [global::System.Runtime.InteropServices.DllImport("DatasmithFacadeCSharp", EntryPoint="CSharp_FDatasmithFacadeActorBinding_AddRelativeRotationCapture__SWIG_1")]
-  public static extern void FDatasmithFacadeActorBinding_AddRelativeRotationCapture__SWIG_1(global::System.Runtime.InteropServices.HandleRef jarg1, float jarg2, float jarg3, float jarg4, float jarg5);
+  public static extern void FDatasmithFacadeActorBinding_AddRelativeRotationCapture__SWIG_1(global::System.Runtime.InteropServices.HandleRef jarg1, double jarg2, double jarg3, double jarg4, double jarg5);
 
   [global::System.Runtime.InteropServices.DllImport("DatasmithFacadeCSharp", EntryPoint="CSharp_FDatasmithFacadeActorBinding_AddRelativeScaleCapture")]
-  public static extern void FDatasmithFacadeActorBinding_AddRelativeScaleCapture(global::System.Runtime.InteropServices.HandleRef jarg1, float jarg2, float jarg3, float jarg4);
+  public static extern void FDatasmithFacadeActorBinding_AddRelativeScaleCapture(global::System.Runtime.InteropServices.HandleRef jarg1, double jarg2, double jarg3, double jarg4);
 
   [global::System.Runtime.InteropServices.DllImport("DatasmithFacadeCSharp", EntryPoint="CSharp_FDatasmithFacadeActorBinding_AddRelativeTransformCapture__SWIG_0")]
-  public static extern void FDatasmithFacadeActorBinding_AddRelativeTransformCapture__SWIG_0(global::System.Runtime.InteropServices.HandleRef jarg1, [global::System.Runtime.InteropServices.In, global::System.Runtime.InteropServices.MarshalAs(global::System.Runtime.InteropServices.UnmanagedType.LPArray)] float[] jarg2, bool jarg3);
+  public static extern void FDatasmithFacadeActorBinding_AddRelativeTransformCapture__SWIG_0(global::System.Runtime.InteropServices.HandleRef jarg1, [global::System.Runtime.InteropServices.In, global::System.Runtime.InteropServices.MarshalAs(global::System.Runtime.InteropServices.UnmanagedType.LPArray)] double[] jarg2, bool jarg3);
 
   [global::System.Runtime.InteropServices.DllImport("DatasmithFacadeCSharp", EntryPoint="CSharp_FDatasmithFacadeActorBinding_AddRelativeTransformCapture__SWIG_1")]
-  public static extern void FDatasmithFacadeActorBinding_AddRelativeTransformCapture__SWIG_1(global::System.Runtime.InteropServices.HandleRef jarg1, [global::System.Runtime.InteropServices.In, global::System.Runtime.InteropServices.MarshalAs(global::System.Runtime.InteropServices.UnmanagedType.LPArray)] float[] jarg2);
+  public static extern void FDatasmithFacadeActorBinding_AddRelativeTransformCapture__SWIG_1(global::System.Runtime.InteropServices.HandleRef jarg1, [global::System.Runtime.InteropServices.In, global::System.Runtime.InteropServices.MarshalAs(global::System.Runtime.InteropServices.UnmanagedType.LPArray)] double[] jarg2);
+
+  [global::System.Runtime.InteropServices.DllImport("DatasmithFacadeCSharp", EntryPoint="CSharp_FDatasmithFacadeActorBinding_AddRelativeTransformCapture__SWIG_2")]
+  public static extern void FDatasmithFacadeActorBinding_AddRelativeTransformCapture__SWIG_2(global::System.Runtime.InteropServices.HandleRef jarg1, [global::System.Runtime.InteropServices.In, global::System.Runtime.InteropServices.MarshalAs(global::System.Runtime.InteropServices.UnmanagedType.LPArray)] float[] jarg2, bool jarg3);
+
+  [global::System.Runtime.InteropServices.DllImport("DatasmithFacadeCSharp", EntryPoint="CSharp_FDatasmithFacadeActorBinding_AddRelativeTransformCapture__SWIG_3")]
+  public static extern void FDatasmithFacadeActorBinding_AddRelativeTransformCapture__SWIG_3(global::System.Runtime.InteropServices.HandleRef jarg1, [global::System.Runtime.InteropServices.In, global::System.Runtime.InteropServices.MarshalAs(global::System.Runtime.InteropServices.UnmanagedType.LPArray)] float[] jarg2);
 
   [global::System.Runtime.InteropServices.DllImport("DatasmithFacadeCSharp", EntryPoint="CSharp_FDatasmithFacadeActorBinding_AddVisibilityCapture")]
   public static extern void FDatasmithFacadeActorBinding_AddVisibilityCapture(global::System.Runtime.InteropServices.HandleRef jarg1, bool jarg2);
@@ -2001,7 +2013,7 @@ class DatasmithFacadeCSharpPINVOKE {
   public static extern void delete_FDatasmithFacadeActorDecal(global::System.Runtime.InteropServices.HandleRef jarg1);
 
   [global::System.Runtime.InteropServices.DllImport("DatasmithFacadeCSharp", EntryPoint="CSharp_FDatasmithFacadeActorDecal_GetDimensions")]
-  public static extern void FDatasmithFacadeActorDecal_GetDimensions(global::System.Runtime.InteropServices.HandleRef jarg1, ref double jarg2, ref double jarg3, ref double jarg4);
+  public static extern void FDatasmithFacadeActorDecal_GetDimensions(global::System.Runtime.InteropServices.HandleRef jarg1, out double jarg2, out double jarg3, out double jarg4);
 
   [global::System.Runtime.InteropServices.DllImport("DatasmithFacadeCSharp", EntryPoint="CSharp_FDatasmithFacadeActorDecal_SetDimensions")]
   public static extern void FDatasmithFacadeActorDecal_SetDimensions(global::System.Runtime.InteropServices.HandleRef jarg1, double jarg2, double jarg3, double jarg4);
diff --git a/Engine/Source/Programs/Enterprise/Datasmith/DatasmithFacadeCSharp/Public/FDatasmithFacadeActor.cs b/Engine/Source/Programs/Enterprise/Datasmith/DatasmithFacadeCSharp/Public/FDatasmithFacadeActor.cs
--- a/Engine/Source/Programs/Enterprise/Datasmith/DatasmithFacadeCSharp/Public/FDatasmithFacadeActor.cs
+++ b/Engine/Source/Programs/Enterprise/Datasmith/DatasmithFacadeCSharp/Public/FDatasmithFacadeActor.cs
@@ -38,43 +38,51 @@ public class FDatasmithFacadeActor : FDatasmithFacadeElement {
   public FDatasmithFacadeActor(string InElementName) : this(DatasmithFacadeCSharpPINVOKE.new_FDatasmithFacadeActor(InElementName), true) {
   }
 
-  public void SetWorldTransform(float[] InWorldMatrix, bool bRowMajor) {
+  public void SetWorldTransform(double[] InWorldMatrix, bool bRowMajor) {
     DatasmithFacadeCSharpPINVOKE.FDatasmithFacadeActor_SetWorldTransform__SWIG_0(swigCPtr, InWorldMatrix, bRowMajor);
   }
 
-  public void SetWorldTransform(float[] InWorldMatrix) {
+  public void SetWorldTransform(double[] InWorldMatrix) {
     DatasmithFacadeCSharpPINVOKE.FDatasmithFacadeActor_SetWorldTransform__SWIG_1(swigCPtr, InWorldMatrix);
   }
 
-  public void SetScale(float X, float Y, float Z) {
+  public void SetWorldTransform(float[] InWorldMatrix, bool bRowMajor) {
+    DatasmithFacadeCSharpPINVOKE.FDatasmithFacadeActor_SetWorldTransform__SWIG_2(swigCPtr, InWorldMatrix, bRowMajor);
+  }
+
+  public void SetWorldTransform(float[] InWorldMatrix) {
+    DatasmithFacadeCSharpPINVOKE.FDatasmithFacadeActor_SetWorldTransform__SWIG_3(swigCPtr, InWorldMatrix);
+  }
+
+  public void SetScale(double X, double Y, double Z) {
     DatasmithFacadeCSharpPINVOKE.FDatasmithFacadeActor_SetScale(swigCPtr, X, Y, Z);
   }
 
-  public void GetScale(out float OutX, out float OutY, out float OutZ) {
+  public void GetScale(out double OutX, out double OutY, out double OutZ) {
     DatasmithFacadeCSharpPINVOKE.FDatasmithFacadeActor_GetScale(swigCPtr, out OutX, out OutY, out OutZ);
   }
 
-  public void SetRotation(float Pitch, float Yaw, float Roll) {
+  public void SetRotation(double Pitch, double Yaw, double Roll) {
     DatasmithFacadeCSharpPINVOKE.FDatasmithFacadeActor_SetRotation__SWIG_0(swigCPtr, Pitch, Yaw, Roll);
   }
 
-  public void GetRotation(out float OutPitch, out float OutYaw, out float OutRoll) {
+  public void GetRotation(out double OutPitch, out double OutYaw, out double OutRoll) {
     DatasmithFacadeCSharpPINVOKE.FDatasmithFacadeActor_GetRotation__SWIG_0(swigCPtr, out OutPitch, out OutYaw, out OutRoll);
   }
 
-  public void SetRotation(float X, float Y, float Z, float W) {
+  public void SetRotation(double X, double Y, double Z, double W) {
     DatasmithFacadeCSharpPINVOKE.FDatasmithFacadeActor_SetRotation__SWIG_1(swigCPtr, X, Y, Z, W);
   }
 
-  public void GetRotation(out float OutX, out float OutY, out float OutZ, out float OutW) {
+  public void GetRotation(out double OutX, out double OutY, out double OutZ, out double OutW) {
     DatasmithFacadeCSharpPINVOKE.FDatasmithFacadeActor_GetRotation__SWIG_1(swigCPtr, out OutX, out OutY, out OutZ, out OutW);
   }
 
-  public void SetTranslation(float X, float Y, float Z) {
+  public void SetTranslation(double X, double Y, double Z) {
     DatasmithFacadeCSharpPINVOKE.FDatasmithFacadeActor_SetTranslation(swigCPtr, X, Y, Z);
   }
 
-  public void GetTranslation(out float OutX, out float OutY, out float OutZ) {
+  public void GetTranslation(out double OutX, out double OutY, out double OutZ) {
     DatasmithFacadeCSharpPINVOKE.FDatasmithFacadeActor_GetTranslation(swigCPtr, out OutX, out OutY, out OutZ);
   }
 
diff --git a/Engine/Source/Programs/Enterprise/Datasmith/DatasmithFacadeCSharp/Public/FDatasmithFacadeActorBinding.cs b/Engine/Source/Programs/Enterprise/Datasmith/DatasmithFacadeCSharp/Public/FDatasmithFacadeActorBinding.cs
--- a/Engine/Source/Programs/Enterprise/Datasmith/DatasmithFacadeCSharp/Public/FDatasmithFacadeActorBinding.cs
+++ b/Engine/Source/Programs/Enterprise/Datasmith/DatasmithFacadeCSharp/Public/FDatasmithFacadeActorBinding.cs
@@ -63,30 +63,38 @@ public class FDatasmithFacadeActorBinding : FDatasmithFacadeElement {
     DatasmithFacadeCSharpPINVOKE.FDatasmithFacadeActorBinding_RemovePropertyCapture(swigCPtr, FDatasmithFacadePropertyCapture.getCPtr(InPropertyCapturePtr));
   }
 
-  public void AddRelativeLocationCapture(float X, float Y, float Z) {
+  public void AddRelativeLocationCapture(double X, double Y, double Z) {
     DatasmithFacadeCSharpPINVOKE.FDatasmithFacadeActorBinding_AddRelativeLocationCapture(swigCPtr, X, Y, Z);
   }
 
-  public void AddRelativeRotationCapture(float Pitch, float Yaw, float Roll) {
+  public void AddRelativeRotationCapture(double Pitch, double Yaw, double Roll) {
     DatasmithFacadeCSharpPINVOKE.FDatasmithFacadeActorBinding_AddRelativeRotationCapture__SWIG_0(swigCPtr, Pitch, Yaw, Roll);
   }
 
-  public void AddRelativeRotationCapture(float X, float Y, float Z, float W) {
+  public void AddRelativeRotationCapture(double X, double Y, double Z, double W) {
     DatasmithFacadeCSharpPINVOKE.FDatasmithFacadeActorBinding_AddRelativeRotationCapture__SWIG_1(swigCPtr, X, Y, Z, W);
   }
 
-  public void AddRelativeScaleCapture(float X, float Y, float Z) {
+  public void AddRelativeScaleCapture(double X, double Y, double Z) {
     DatasmithFacadeCSharpPINVOKE.FDatasmithFacadeActorBinding_AddRelativeScaleCapture(swigCPtr, X, Y, Z);
   }
 
-  public void AddRelativeTransformCapture(float[] InMatrix, bool bRowMajor) {
+  public void AddRelativeTransformCapture(double[] InMatrix, bool bRowMajor) {
     DatasmithFacadeCSharpPINVOKE.FDatasmithFacadeActorBinding_AddRelativeTransformCapture__SWIG_0(swigCPtr, InMatrix, bRowMajor);
   }
 
-  public void AddRelativeTransformCapture(float[] InMatrix) {
+  public void AddRelativeTransformCapture(double[] InMatrix) {
     DatasmithFacadeCSharpPINVOKE.FDatasmithFacadeActorBinding_AddRelativeTransformCapture__SWIG_1(swigCPtr, InMatrix);
   }
 
+  public void AddRelativeTransformCapture(float[] InMatrix, bool bRowMajor) {
+    DatasmithFacadeCSharpPINVOKE.FDatasmithFacadeActorBinding_AddRelativeTransformCapture__SWIG_2(swigCPtr, InMatrix, bRowMajor);
+  }
+
+  public void AddRelativeTransformCapture(float[] InMatrix) {
+    DatasmithFacadeCSharpPINVOKE.FDatasmithFacadeActorBinding_AddRelativeTransformCapture__SWIG_3(swigCPtr, InMatrix);
+  }
+
   public void AddVisibilityCapture(bool bInVisibility) {
     DatasmithFacadeCSharpPINVOKE.FDatasmithFacadeActorBinding_AddVisibilityCapture(swigCPtr, bInVisibility);
   }
diff --git a/Engine/Source/Programs/Enterprise/Datasmith/DatasmithFacadeCSharp/Public/FDatasmithFacadeActorCamera.cs b/Engine/Source/Programs/Enterprise/Datasmith/DatasmithFacadeCSharp/Public/FDatasmithFacadeActorCamera.cs
--- a/Engine/Source/Programs/Enterprise/Datasmith/DatasmithFacadeCSharp/Public/FDatasmithFacadeActorCamera.cs
+++ b/Engine/Source/Programs/Enterprise/Datasmith/DatasmithFacadeCSharp/Public/FDatasmithFacadeActorCamera.cs
@@ -38,11 +38,11 @@ public class FDatasmithFacadeActorCamera : FDatasmithFacadeActor {
   public FDatasmithFacadeActorCamera(string InElementName) : this(DatasmithFacadeCSharpPINVOKE.new_FDatasmithFacadeActorCamera(InElementName), true) {
   }
 
-  public void SetCameraPosition(float InX, float InY, float InZ) {
+  public void SetCameraPosition(double InX, double InY, double InZ) {
     DatasmithFacadeCSharpPINVOKE.FDatasmithFacadeActorCamera_SetCameraPosition(swigCPtr, InX, InY, InZ);
   }
 
-  public void SetCameraRotation(float InForwardX, float InForwardY, float InForwardZ, float InUpX, float InUpY, float InUpZ) {
+  public void SetCameraRotation(double InForwardX, double InForwardY, double InForwardZ, double InUpX, double InUpY, double InUpZ) {
     DatasmithFacadeCSharpPINVOKE.FDatasmithFacadeActorCamera_SetCameraRotation(swigCPtr, InForwardX, InForwardY, InForwardZ, InUpX, InUpY, InUpZ);
   }
 
diff --git a/Engine/Source/Programs/Enterprise/Datasmith/DatasmithFacadeCSharp/Public/FDatasmithFacadeActorDecal.cs b/Engine/Source/Programs/Enterprise/Datasmith/DatasmithFacadeCSharp/Public/FDatasmithFacadeActorDecal.cs
--- a/Engine/Source/Programs/Enterprise/Datasmith/DatasmithFacadeCSharp/Public/FDatasmithFacadeActorDecal.cs
+++ b/Engine/Source/Programs/Enterprise/Datasmith/DatasmithFacadeCSharp/Public/FDatasmithFacadeActorDecal.cs
@@ -38,8 +38,8 @@ public class FDatasmithFacadeActorDecal : FDatasmithFacadeActor {
   public FDatasmithFacadeActorDecal(string InElementName) : this(DatasmithFacadeCSharpPINVOKE.new_FDatasmithFacadeActorDecal(InElementName), true) {
   }
 
-  public void GetDimensions(ref double OutX, ref double OutY, ref double OutZ) {
-    DatasmithFacadeCSharpPINVOKE.FDatasmithFacadeActorDecal_GetDimensions(swigCPtr, ref OutX, ref OutY, ref OutZ);
+  public void GetDimensions(out double OutX, out double OutY, out double OutZ) {
+    DatasmithFacadeCSharpPINVOKE.FDatasmithFacadeActorDecal_GetDimensions(swigCPtr, out OutX, out OutY, out OutZ);
   }
 
   public void SetDimensions(double InX, double InY, double InZ) {
diff --git a/Engine/Source/Programs/Enterprise/Datasmith/DatasmithFacadeCSharp/Public/FDatasmithFacadeMesh.cs b/Engine/Source/Programs/Enterprise/Datasmith/DatasmithFacadeCSharp/Public/FDatasmithFacadeMesh.cs
--- a/Engine/Source/Programs/Enterprise/Datasmith/DatasmithFacadeCSharp/Public/FDatasmithFacadeMesh.cs
+++ b/Engine/Source/Programs/Enterprise/Datasmith/DatasmithFacadeCSharp/Public/FDatasmithFacadeMesh.cs
@@ -157,7 +157,7 @@ public class FDatasmithFacadeMesh : global::System.IDisposable {
     return ret;
   }
 
-  public void GetUV(int Channel, int Index, out float OutU, out float OutV) {
+  public void GetUV(int Channel, int Index, out double OutU, out double OutV) {
     DatasmithFacadeCSharpPINVOKE.FDatasmithFacadeMesh_GetUV(swigCPtr, Channel, Index, out OutU, out OutV);
   }
 
diff --git a/Engine/Source/Programs/Enterprise/Datasmith/DatasmithFacadeCSharp/Public/FDatasmithFacadeTransformAnimation.cs b/Engine/Source/Programs/Enterprise/Datasmith/DatasmithFacadeCSharp/Public/FDatasmithFacadeTransformAnimation.cs
--- a/Engine/Source/Programs/Enterprise/Datasmith/DatasmithFacadeCSharp/Public/FDatasmithFacadeTransformAnimation.cs
+++ b/Engine/Source/Programs/Enterprise/Datasmith/DatasmithFacadeCSharp/Public/FDatasmithFacadeTransformAnimation.cs
@@ -56,8 +56,8 @@ public class FDatasmithFacadeTransformAnimation : FDatasmithFacadeBaseAnimation
     return ret;
   }
 
-  public void GetFrame(EDatasmithFacadeAnimationTransformType TransformType, int Index, out int OutFrameNumber, ref double OutX, ref double OutY, ref double OutZ) {
-    DatasmithFacadeCSharpPINVOKE.FDatasmithFacadeTransformAnimation_GetFrame(swigCPtr, (int)TransformType, Index, out OutFrameNumber, ref OutX, ref OutY, ref OutZ);
+  public void GetFrame(EDatasmithFacadeAnimationTransformType TransformType, int Index, out int OutFrameNumber, out double OutX, out double OutY, out double OutZ) {
+    DatasmithFacadeCSharpPINVOKE.FDatasmithFacadeTransformAnimation_GetFrame(swigCPtr, (int)TransformType, Index, out OutFrameNumber, out OutX, out OutY, out OutZ);
   }
 
   public void RemoveFrame(EDatasmithFacadeAnimationTransformType TransformType, int Index) {
diff --git a/Engine/Source/Programs/Enterprise/Datasmith/DatasmithMaxExporter/Private/DatasmithMaxExporterWithDirectLink.cpp b/Engine/Source/Programs/Enterprise/Datasmith/DatasmithMaxExporter/Private/DatasmithMaxExporterWithDirectLink.cpp
--- a/Engine/Source/Programs/Enterprise/Datasmith/DatasmithMaxExporter/Private/DatasmithMaxExporterWithDirectLink.cpp
+++ b/Engine/Source/Programs/Enterprise/Datasmith/DatasmithMaxExporter/Private/DatasmithMaxExporterWithDirectLink.cpp
@@ -69,7 +69,7 @@ typedef Texmap* FTexmapKey;
 class FDatasmith3dsMaxScene
 {
 public:
-	FDatasmith3dsMaxScene() 
+	FDatasmith3dsMaxScene()
 	{
 		ResetScene();
 	}
@@ -341,7 +341,7 @@ private:
 	TMap<TPair<uint32, TPair<uint32, uint32>>, MAXClass*> KnownMaxDesc;
 	// Same for the lookup_MAXSuperClass.
 	TMap<uint32, MAXSuperClass*> KnownMaxSuperClass;
-	
+
 };
 
 
@@ -440,7 +440,7 @@ public:
 
 	FBeginRefEnumProc BeginProc;
 	FEndRefEnumProc EndProc;
-	
+
 	TSet<INode*> NodesPrepared;
 };
 
@@ -816,7 +816,7 @@ public:
 		{
 			PROGRESS_STAGE("Refresh collisions") // Update set of nodes used for collision
 			PROGRESS_STAGE_COUNTER(InvalidatedNodeTrackers.Num());
-			
+
 			TSet<FNodeTracker*> NodesWithChangedCollisionStatus; // Need to invalidate these nodes to make them renderable or to keep them from renderable depending on collision status
 			for (FNodeTracker* NodeTracker : InvalidatedNodeTrackers)
 			{
@@ -835,7 +835,7 @@ public:
 				UpdateNode(*NodeTracker);
 			}
 		}
-		
+
 		{
 			PROGRESS_STAGE("Process invalidated instances")
 			PROGRESS_STAGE_COUNTER(InvalidatedInstances.Num());
@@ -944,10 +944,10 @@ public:
 		for (TPair<FNodeKey, FNodeTrackerHandle> NodeKeyAndNodeTracker: NodeTrackers)
 		{
 			FNodeTracker* NodeTracker = NodeKeyAndNodeTracker.Value.GetNodeTracker();
-			
+
 			if (NodeTracker->DatasmithActorElement)
 			{
-				
+
 				if (Lights.Contains(NodeTracker))
 				{
 					const TSharedPtr<IDatasmithLightActorElement> LightElement = StaticCastSharedPtr< IDatasmithLightActorElement >(NodeTracker->DatasmithActorElement);
@@ -972,7 +972,7 @@ public:
 	{
 		LogDebugNode(TEXT("AddNode"), Node);
 		FNodeTracker* NodeTracker = NodeTrackers.Emplace(NodeKey, FNodeTrackerHandle(NodeKey, Node)).GetNodeTracker();
-		
+
 		NodeTrackersNames.FindOrAdd(NodeTracker->Name).Add(NodeTracker);
 		InvalidatedNodeTrackers.Add(NodeTracker);
 
@@ -981,7 +981,7 @@ public:
 
 	virtual void RemoveMaterial(const TSharedPtr<IDatasmithBaseMaterialElement>& DatasmithMaterial) override
 	{
-		ExportedScene.DatasmithSceneRef->RemoveMaterial(DatasmithMaterial);		
+		ExportedScene.DatasmithSceneRef->RemoveMaterial(DatasmithMaterial);
 	}
 
 	virtual void RemoveTexture(const TSharedPtr<IDatasmithTextureElement>& DatasmithTextureElement) override
@@ -995,7 +995,7 @@ public:
 
 		if (NodeTracker.bDeleted)
 		{
-			// Change events sometimes received for nodes that are already deleted 
+			// Change events sometimes received for nodes that are already deleted
 			// skipping processing of node subtree because INode pointer may already be invalid
 			// Test case: create container, add node to it. Close it, open it, close again, then sync
 			// Change event will be received for NodeKey that returns NULL from NodeEventNamespace::GetNodeByKey
@@ -1118,7 +1118,7 @@ public:
 		Helpers.Remove(&NodeTracker);
 		Cameras.Remove(&NodeTracker);
 		Lights.Remove(&NodeTracker);
-		
+
 		{
 			// remove static meshes used by the RailClone
 			TUniquePtr<FRailClonesConverted> RailClonesConverted;
@@ -1275,7 +1275,7 @@ public:
 	void UpdateNode(FNodeTracker& NodeTracker)
 	{
 
-		// Forget anything that this node was before update: place in datasmith hierarchy, datasmith objects, instances connection. Updating may change anything 
+		// Forget anything that this node was before update: place in datasmith hierarchy, datasmith objects, instances connection. Updating may change anything
 		RemoveFromConverted(NodeTracker);
 		ConvertNodeObject(NodeTracker);
 	}
@@ -1419,7 +1419,7 @@ public:
 		{
 			return nullptr; // Max node not known, assume it doesn't exist
 		}
-		
+
 		return ParentNodeTrackerHandle->GetNodeTracker();
 	}
 
@@ -1517,7 +1517,7 @@ public:
 		FString UniqueName = FString::FromInt(NodeTracker.Node->GetHandle());
 		FString Label = FString((const TCHAR*)NodeTracker.Node->GetName());
 
-		// Create and setup mesh actor if there's a mesh 
+		// Create and setup mesh actor if there's a mesh
 		if (bHasMesh)
 		{
 			FString MeshActorName = UniqueName;
@@ -1573,7 +1573,7 @@ public:
 			NodeTracker.DatasmithActorElement->SetLayer(*NodeTracker.Layer->Name);
 		}
 
-		// Apply material 
+		// Apply material
 		if (DatasmithMeshElement)
 		{
 			if (Mtl* Material = NodeTracker.Node->GetMtl())
@@ -1602,7 +1602,7 @@ public:
 					if (Instances.Material != Material)
 					{
 						TSharedRef<IDatasmithMeshActorElement> DatasmithMeshActorRef = NodeTracker.DatasmithMeshActor.ToSharedRef();
-						FDatasmithMaxSceneExporter::ParseMaterialForMeshActor(Material, DatasmithMeshActorRef, Instances.SupportedChannels, NodeTracker.DatasmithMeshActor->GetTranslation());
+						FDatasmithMaxSceneExporter::ParseMaterialForMeshActor(Material, DatasmithMeshActorRef, Instances.SupportedChannels, FVector3f(NodeTracker.DatasmithMeshActor->GetTranslation()));
 					}
 				}
 			}
@@ -1775,7 +1775,7 @@ public:
 		SetupActor(NodeTracker);
 
 		//Cylinder shaped lights don't have the same default orientations, so we recalculate their transform and add the shape information.
-		if (LightElement->IsA(EDatasmithElementType::AreaLight) 
+		if (LightElement->IsA(EDatasmithElementType::AreaLight)
 			&& StaticCastSharedPtr<IDatasmithAreaLightElement>(LightElement)->GetLightShape() == EDatasmithLightShape::Cylinder)
 		{
 			FVector Translation, Scale;
@@ -1809,7 +1809,7 @@ public:
 		FString MeshName = FString::FromInt(NodeTracker.Node->GetHandle()) + TEXT("_") + FString::FromInt(MeshIndex);
 
 		// note: when export Mesh goes to other place due to parallellizing it's result would be unknown here so MeshIndex handling will change(i.e. increment for any mesh)
-				
+
 		TSharedPtr<IDatasmithMeshElement> DatasmithMeshElement;
 		TSet<uint16> SupportedChannels;
 
@@ -1829,11 +1829,11 @@ public:
 			DatasmithMeshElement->SetLabel(*MeshLabel);
 
 			FDatasmithConverter Converter;
-				
+
 			// todo: override material
 			TSharedPtr< IDatasmithActorElement > InversedHISMActor;
 			// todo: ExportHierarchicalInstanceStaticMeshActor CustomMeshNode only used for Material - can be simplified, Material anyway is dealt with outside too
-			TSharedRef<IDatasmithActorElement> HismActorElement = FDatasmithMaxSceneExporter::ExportHierarchicalInstanceStaticMeshActor( 
+			TSharedRef<IDatasmithActorElement> HismActorElement = FDatasmithMaxSceneExporter::ExportHierarchicalInstanceStaticMeshActor(
 				ExportedScene.GetDatasmithScene(), NodeTracker.Node, GeometryNode, *MeshLabel, SupportedChannels,
 				Material, &Transforms, *MeshName, Converter.UnitToCentimeter, EStaticMeshExportMode::Default, InversedHISMActor);
 			NodeTracker.DatasmithActorElement->AddChild(HismActorElement, EDatasmithActorAttachmentRule::KeepWorldTransform);
@@ -1855,7 +1855,7 @@ public:
 			return bResult;
 		}
 
-		// AnimHandle is unique and never reused for new objects 
+		// AnimHandle is unique and never reused for new objects
 		// todo: reset instances and nodes when one node of an instance changes??? Check how it should be done actually, dependencies - nodes, object, invalidation place(Update, Event) etc...
 		AnimHandle Handle = Animatable::GetHandleByAnim(Obj);
 
@@ -1911,8 +1911,8 @@ public:
 	virtual void NodeAdded(INode* Node) override
 	{
 		// Node sometimes is null. 'Added' NodeEvent might come after node was actually deleted(immediately after creation)
-		// e.g.[mxs]: b = box(); delete b 
-		// NodeEvents are delayed(not executed in the same stack frame as command that causes them) so they come later. 
+		// e.g.[mxs]: b = box(); delete b
+		// NodeEvents are delayed(not executed in the same stack frame as command that causes them) so they come later.
 		if (!Node)
 		{
 			return;
@@ -1998,7 +1998,7 @@ public:
 		if (FNodeTrackerHandle* NodeTracker = NodeTrackers.Find(NodeKey))
 		{
 			// todo: investigate why NodeEventNamespace::GetNodeByKey may stil return NULL
-			// testcase - add XRef Material - this will immediately have this 
+			// testcase - add XRef Material - this will immediately have this
 			// even though NOTIFY_SCENE_ADDED_NODE was called for node and NOTIFY_SCENE_PRE_DELETED_NODE wasn't!
 			INode* Node = NodeEventNamespace::GetNodeByKey(NodeKey);
 			if (Node)
@@ -2024,7 +2024,7 @@ public:
 
 	virtual void NodeHideChanged(FNodeKey NodeKey) override
 	{
-		// todo: invalidate visibility only - note to handle this not enought add/remove 
+		// todo: invalidate visibility only - note to handle this not enought add/remove
 		// actor. make sure to invalidate instances(in case geometry usage changed - like hidden node with multimat), materials
 
 		InvalidateNode(NodeKey);
@@ -2032,7 +2032,7 @@ public:
 
 	virtual void NodePropertiesChanged(FNodeKey NodeKey) override
 	{
-		// todo: invalidate visibility only - note to handle this not enought add/remove 
+		// todo: invalidate visibility only - note to handle this not enought add/remove
 		// actor. make sure to invalidate instances(in case geometry usage changed - like hidden node with multimat), materials
 
 		InvalidateNode(NodeKey);
@@ -2042,7 +2042,7 @@ public:
 	{
 		InvalidateNode(NodeKey);
 	}
-	
+
 	///////////////////////////////////////////////
 
 	const FExportOptions& Options;
@@ -2058,7 +2058,7 @@ public:
 	TSet<FNodeTracker*> InvalidatedNodeTrackers; // Nodes that need to be rebuilt
 	TMap<FNodeTracker*, TSharedPtr<IDatasmithMetaDataElement>> NodeDatasmithMetadata; // All scene nodes
 
-	TMap<FNodeTracker*, TSet<FNodeTracker*>> CollisionNodes; // Nodes used as collision meshes for other nodes, counted by each user 
+	TMap<FNodeTracker*, TSet<FNodeTracker*>> CollisionNodes; // Nodes used as collision meshes for other nodes, counted by each user
 
 	FMaterialsCollectionTracker MaterialsCollectionTracker;
 
@@ -2073,7 +2073,7 @@ public:
 	TMap<FLayerTracker*, TSet<FNodeTracker*>> NodesPerLayer;
 
 	FNodesPreparer NodesPreparer;
-	
+
 	struct FRailClonesConverted
 	{
 		TArray<TSharedPtr<IDatasmithMeshElement>> Meshes; // Meshes created for this railclones object
@@ -2100,7 +2100,7 @@ public:
 	virtual void SetOutputPath(const TCHAR* Path) override
 	{
 		OutputPath = Path;
-		ExportedScene.SetOutputPath(*OutputPath);		
+		ExportedScene.SetOutputPath(*OutputPath);
 	}
 
 	virtual void SetName(const TCHAR* Name) override
@@ -2120,17 +2120,17 @@ public:
 
 	virtual void InitializeDirectLinkForScene() override
 	{
-		if (DirectLinkImpl) 
+		if (DirectLinkImpl)
 		{
 			return;
 		}
 
 		InitializeScene();
 
-		// XXX: PreExport needs to be called before DirectLink instance is constructed - 
+		// XXX: PreExport needs to be called before DirectLink instance is constructed -
 		// Reason - it calls initialization of FTaskGraphInterface. Callstack:
 		// PreExport:
-		//  - FDatasmithExporterManager::Initialize 
+		//  - FDatasmithExporterManager::Initialize
 		//	-- DatasmithGameThread::InitializeInCurrentThread
 		//  --- GEngineLoop.PreInit
 		//  ---- PreInitPreStartupScreen
@@ -2147,7 +2147,7 @@ public:
 
 	virtual void UpdateDirectLinkScene() override
 	{
-		if (!DirectLinkImpl) 
+		if (!DirectLinkImpl)
 		{
 			// InitializeDirectLinkForScene wasn't called yet. This rarely happens when Sync is pressed right before event like PostSceneReset(for New All UI command) was handled
 			// Very quickly! Unfortunately code needs to wait for PostSceneReset to get proper scene name there(no earlier event signals that name is available)
@@ -2161,7 +2161,7 @@ public:
 
 	static VOID AutoSyncTimerProc(HWND, UINT, UINT_PTR TimerIdentifier, DWORD)
 	{
-		reinterpret_cast<FExporter*>(TimerIdentifier)->UpdateAutoSync(); 
+		reinterpret_cast<FExporter*>(TimerIdentifier)->UpdateAutoSync();
 	}
 
 	// Update is user was idle for some time
@@ -2355,7 +2355,7 @@ void ShutdownExporter()
 
 IExporter* GetExporter()
 {
-	return Exporter.Get();	
+	return Exporter.Get();
 }
 
 IPersistentExportOptions& GetPersistentExportOptions()
@@ -2363,7 +2363,7 @@ IPersistentExportOptions& GetPersistentExportOptions()
 	return PersistentExportOptions;
 }
 
-void FExporter::Shutdown() 
+void FExporter::Shutdown()
 {
 	Exporter.Reset();
 	FDatasmithDirectLink::Shutdown();
diff --git a/Engine/Source/Programs/Enterprise/Datasmith/DatasmithMaxExporter/Private/DatasmithMaxSceneExporter.cpp b/Engine/Source/Programs/Enterprise/Datasmith/DatasmithMaxExporter/Private/DatasmithMaxSceneExporter.cpp
--- a/Engine/Source/Programs/Enterprise/Datasmith/DatasmithMaxExporter/Private/DatasmithMaxSceneExporter.cpp
+++ b/Engine/Source/Programs/Enterprise/Datasmith/DatasmithMaxExporter/Private/DatasmithMaxSceneExporter.cpp
@@ -203,7 +203,7 @@ int FDatasmithMaxSceneExporter::GetSeedFromMaterial(Mtl* Material)
 	return Seed;
 }
 
-FString FDatasmithMaxSceneExporter::GetRandomSubMaterial(Mtl* Material, FVector RandomSeed)
+FString FDatasmithMaxSceneExporter::GetRandomSubMaterial(Mtl* Material, FVector3f RandomSeed)
 {
 	uint32 RandomIndex = int(FMath::Abs(RandomSeed.X - RandomSeed.Y + RandomSeed.Z)) + GetSeedFromMaterial(Material);
 
@@ -306,7 +306,7 @@ void FDatasmithMaxSceneExporter::ExportMeshActor(TSharedRef< IDatasmithScene > D
 
 	if (ExportMode == EStaticMeshExportMode::Default && StaticMeshMtl != Node->GetMtl())
 	{
-		ParseMaterialForMeshActor( Node->GetMtl(), MeshActor, SupportedChannels, MeshActor->GetTranslation() );
+		ParseMaterialForMeshActor( Node->GetMtl(), MeshActor, SupportedChannels, FVector3f(MeshActor->GetTranslation()) );
 	}
 
 	if ( bPivotIsBakedInGeometry || bNeedPivotComponent )
@@ -339,14 +339,14 @@ void FDatasmithMaxSceneExporter::ExportMeshActor(TSharedRef< IDatasmithScene > D
 TSharedRef< IDatasmithActorElement > FDatasmithMaxSceneExporter::ExportHierarchicalInstanceStaticMeshActor(TSharedRef< IDatasmithScene > DatasmithScene, INode* Node, INode* CustomMeshNode, const TCHAR* Label, TSet<uint16>& SupportedChannels, Mtl* StaticMeshMtl, const TArray<Matrix3>* Instances,
 	const TCHAR* MeshName, float UnitMultiplier, const EStaticMeshExportMode& ExportMode, TSharedPtr< IDatasmithActorElement >& OutInversedHISMActor)
 {
-	check( Node && Instances && MeshName);
+	check(Node && Instances && MeshName);
 
 	TSharedPtr< IDatasmithMeshActorElement > MeshActor;
 	FVector Pos, Scale;
 	FQuat Rotation;
 
-	const FVector RandomSeed(FMath::Rand(), FMath::Rand(), FMath::Rand());
-	auto FinalizeHISMActor = [&](TSharedPtr< IDatasmithMeshActorElement >& FinalizingActor, FVector Seed, const TCHAR* ActorLabel)
+	const FVector3f RandomSeed(FMath::Rand(), FMath::Rand(), FMath::Rand());
+	auto FinalizeHISMActor = [&](TSharedPtr< IDatasmithMeshActorElement >& FinalizingActor, FVector3f Seed, const TCHAR* ActorLabel)
 	{
 		FinalizingActor->SetStaticMeshPathName(MeshName);
 
@@ -464,7 +464,7 @@ TSharedRef< IDatasmithActorElement > FDatasmithMaxSceneExporter::ExportHierarchi
 	return MeshActor.ToSharedRef();
 }
 
-void FDatasmithMaxSceneExporter::ParseMaterialForMeshActor(Mtl* Material, TSharedRef< IDatasmithMeshActorElement >& MeshActor, TSet<uint16>& SupportedChannels, FVector RandomSeed)
+void FDatasmithMaxSceneExporter::ParseMaterialForMeshActor(Mtl* Material, TSharedRef< IDatasmithMeshActorElement >& MeshActor, TSet<uint16>& SupportedChannels, const FVector3f& RandomSeed)
 {
 	if (FDatasmithMaxMatHelper::GetMaterialClass(Material) == EDSMaterialType::XRefMat)
 	{
diff --git a/Engine/Source/Programs/Enterprise/Datasmith/DatasmithMaxExporter/Public/DatasmithMaxSceneExporter.h b/Engine/Source/Programs/Enterprise/Datasmith/DatasmithMaxExporter/Public/DatasmithMaxSceneExporter.h
--- a/Engine/Source/Programs/Enterprise/Datasmith/DatasmithMaxExporter/Public/DatasmithMaxSceneExporter.h
+++ b/Engine/Source/Programs/Enterprise/Datasmith/DatasmithMaxExporter/Public/DatasmithMaxSceneExporter.h
@@ -65,7 +65,7 @@ public:
 
 	static FString GetActualPath(const TCHAR* OriginalPath);
 
-	static FString GetRandomSubMaterial(Mtl* Material, FVector RandomSeed);
+	static FString GetRandomSubMaterial(Mtl* Material, FVector3f RandomSeed);
 
 	/**
 	 * Extract the Node to Object transform in UE coordinates
@@ -78,7 +78,7 @@ public:
 	static FTransform GetPivotTransform(INode* Node, float UnitMultiplier);
 	static void MaxToUnrealCoordinates(Matrix3 Matrix, FVector& Translation, FQuat& Rotation, FVector& Scale, float UnitMultiplier, const FMaxLightCoordinateConversionParams& LightParams = FMaxLightCoordinateConversionParams());
 
-	static void ParseMaterialForMeshActor(Mtl* Material, TSharedRef< IDatasmithMeshActorElement >& MeshActor, TSet<uint16>& SupportedChannels, FVector RandomSeed);
+	static void ParseMaterialForMeshActor(Mtl* Material, TSharedRef< IDatasmithMeshActorElement >& MeshActor, TSet<uint16>& SupportedChannels, const FVector3f& RandomSeed);
 	
 	static TSharedPtr< IDatasmithLightActorElement > CreateLightElementForNode(INode* Node, const TCHAR* Name);
 	static bool ParseLight(INode* Node, TSharedRef< IDatasmithLightActorElement > LightElement, TSharedRef< IDatasmithScene > DatasmithScene);
diff --git a/Engine/Source/Programs/Enterprise/Datasmith/DatasmithRevitExporter/Private/DatasmithRevitDocumentData.cs b/Engine/Source/Programs/Enterprise/Datasmith/DatasmithRevitExporter/Private/DatasmithRevitDocumentData.cs
--- a/Engine/Source/Programs/Enterprise/Datasmith/DatasmithRevitExporter/Private/DatasmithRevitDocumentData.cs
+++ b/Engine/Source/Programs/Enterprise/Datasmith/DatasmithRevitExporter/Private/DatasmithRevitDocumentData.cs
@@ -157,7 +157,7 @@ namespace DatasmithRevitExporter
 			{
 				InToActor.SetLabel(InFromActor.GetLabel());
 
-				float X, Y, Z, W;
+				double X, Y, Z, W;
 				InFromActor.GetTranslation(out X, out Y, out Z);
 				InToActor.SetTranslation(X, Y, Z);
 				InFromActor.GetScale(out X, out Y, out Z);
@@ -202,9 +202,9 @@ namespace DatasmithRevitExporter
 					}
 				}
 
-				bool bIsCached = 
-					ThisElement != null && 
-					ThisElement.IsValidObject && 
+				bool bIsCached =
+					ThisElement != null &&
+					ThisElement.IsValidObject &&
 					(DocumentData.DirectLink?.IsElementCached(ThisElement) ?? false);
 
 				// Check if actor type has changed for this element (f.e. static mesh actor -> regular actor),
@@ -314,7 +314,7 @@ namespace DatasmithRevitExporter
 					if (!InOutWorldTransform.IsIdentity)
 					{
 						InOutWorldTransform = InOutWorldTransform * PivotTransform;
-					} 
+					}
 					else
 					{
 						InOutWorldTransform = PivotTransform;
@@ -378,8 +378,8 @@ namespace DatasmithRevitExporter
 				{
 					Translation = (InElement as StructuralConnectionHandler).GetOrigin();
 				}
-				else if (InElement.GetType() == typeof(Floor) 
-					|| InElement.GetType() == typeof(Ceiling) 
+				else if (InElement.GetType() == typeof(Floor)
+					|| InElement.GetType() == typeof(Ceiling)
 					|| InElement.GetType() == typeof(RoofBase)
 					|| InElement.GetType().IsSubclassOf(typeof(RoofBase)))
 				{
@@ -415,14 +415,14 @@ namespace DatasmithRevitExporter
 				if (InElement.GetType() == typeof(Wall))
 				{
 					// In rare cases, wall may not support orientation.
-					// If this happens, we need to use the direction of its Curve property and 
+					// If this happens, we need to use the direction of its Curve property and
 					// derive orientation from there.
 					try
 					{
 						BasisY = (InElement as Wall).Orientation.Normalize();
 						BasisX = BasisY.CrossProduct(XYZ.BasisZ).Normalize();
 						BasisZ = BasisX.CrossProduct(BasisY).Normalize();
-					} 
+					}
 					catch
 					{
 						if (InElement.Location.GetType() == typeof(LocationCurve))
@@ -433,15 +433,15 @@ namespace DatasmithRevitExporter
 							{
 								BasisX = (CurveLocation.Curve as Line).Direction;
 								ComputeBasis(BasisX, ref BasisY, ref BasisZ);
-							} 
+							}
 							else if (CurveLocation.Curve.IsBound)
 							{
 								Transform Derivatives = CurveLocation.Curve.ComputeDerivatives(0f, true);
 								BasisX = Derivatives.BasisX.Normalize();
 								BasisY = Derivatives.BasisY.Normalize();
 								BasisZ = Derivatives.BasisZ.Normalize();
-							} 
-							else 
+							}
+							else
 							{
 								BasisX = XYZ.BasisX;
 								BasisY = XYZ.BasisY;
@@ -493,7 +493,7 @@ namespace DatasmithRevitExporter
 					BasisX = (InElement as FlexPipe).StartTangent;
 					ComputeBasis(BasisX, ref BasisY, ref BasisZ);
 				}
-				else if (InElement.GetType() == typeof(Floor) 
+				else if (InElement.GetType() == typeof(Floor)
 					|| InElement.GetType() == typeof(Ceiling)
 					|| InElement.GetType() == typeof(RoofBase)
 					|| InElement.GetType().IsSubclassOf(typeof(RoofBase)))
@@ -714,7 +714,7 @@ namespace DatasmithRevitExporter
 					FDatasmithFacadeActorMesh RPCMeshActor = new FDatasmithFacadeActorMesh(HashedActorName);
 					RPCMeshActor.SetMesh(RPCMesh.GetName());
 					FacadeActor = RPCMeshActor;
-					
+
 					OutDatasmithMesh = RPCMesh;
 					OutDatasmithMeshElement = new FDatasmithFacadeMeshElement(HashedName);
 					OutDatasmithMeshElement.SetLabel(GetActorLabel());
@@ -775,8 +775,8 @@ namespace DatasmithRevitExporter
 			}
 
 			public void AddChildActor(
-				FDatasmithFacadeActor ChildActor, 
-				FDatasmithFacadeMetaData MetaData, 
+				FDatasmithFacadeActor ChildActor,
+				FDatasmithFacadeMetaData MetaData,
 				bool bOptimizeHierarchy
 			)
 			{
@@ -922,14 +922,14 @@ namespace DatasmithRevitExporter
 
 			private string GenerateUniqueInstanceName()
 			{
-				// GenerateUniqueInstanceName is being called when generating a name for instance. 
-				// After the call, the intance is added as a child to its parent. 
+				// GenerateUniqueInstanceName is being called when generating a name for instance.
+				// After the call, the intance is added as a child to its parent.
 				// Next time the method gets called for the next instance, ChildElements.Count will be different/incremented.
 
-				// To add uniqueness to the generated name, we construct a string with child counts from 
+				// To add uniqueness to the generated name, we construct a string with child counts from
 				// current parent instance, up to the root:
 				// Elem->Instance->Instance->Instace can produce something like: "1:5:3" for example.
-				// However, this is not enough because elsewhere we might encounter the same sequence in terms of child counts, 
+				// However, this is not enough because elsewhere we might encounter the same sequence in terms of child counts,
 				// but adding the CurrentElement unique id ensures we get unique name string in the end.
 
 				StringBuilder ChildCounts = new StringBuilder();
@@ -1063,7 +1063,7 @@ namespace DatasmithRevitExporter
 			}
 		}
 
-		public Dictionary<string, Tuple<FDatasmithFacadeMeshElement, Task<bool>>> 
+		public Dictionary<string, Tuple<FDatasmithFacadeMeshElement, Task<bool>>>
 														MeshMap = new Dictionary<string, Tuple<FDatasmithFacadeMeshElement, Task<bool>>>();
 		public Dictionary<ElementId, FBaseElementData>	ActorMap = new Dictionary<ElementId, FDocumentData.FBaseElementData>();
 		public Dictionary<string, FMaterialData>		MaterialDataMap = null;
@@ -1180,7 +1180,7 @@ namespace DatasmithRevitExporter
 				}
 				ElementData = ActorMap[InElement.Id] as FElementData;
 			}
-			
+
 			if (ElementData == null)
 			{
 				if (DirectLink?.IsElementCached(InElement) ?? false)
@@ -1201,13 +1201,13 @@ namespace DatasmithRevitExporter
 							if (ElementData.ChildElements.Count > 0)
 							{
 								List<FBaseElementData> ChildrenToRemove = new List<FBaseElementData>();
-					
+
 								for(int ChildIndex = 0; ChildIndex < ElementData.ChildElements.Count; ++ChildIndex)
 								{
 									FBaseElementData ChildElement = ElementData.ChildElements[ChildIndex];
 
-									bool bIsFamilyIntance = 
-										((ChildElement as FElementData) == null) && 
+									bool bIsFamilyIntance =
+										((ChildElement as FElementData) == null) &&
 										ChildElement.ElementActor.IsComponent();
 
 									if (bIsFamilyIntance)
@@ -1302,7 +1302,7 @@ namespace DatasmithRevitExporter
 			FDatasmithFacadeActor CloneActor = new FDatasmithFacadeActor(SourceActor.GetName());
 			CloneActor.SetLabel(SourceActor.GetLabel());
 
-			float X, Y, Z, W;
+			double X, Y, Z, W;
 			SourceActor.GetTranslation(out X, out Y, out Z);
 			CloneActor.SetTranslation(X, Y, Z);
 			SourceActor.GetScale(out X, out Y, out Z);
@@ -1830,13 +1830,13 @@ namespace DatasmithRevitExporter
 			{
 				ElementData = new FBaseElementData(SiteLocationActor, null, this);
 				// Prevent the Datasmith placeholder actor from being removed by optimization.
-				ElementData.bOptimizeHierarchy = false; 
+				ElementData.bOptimizeHierarchy = false;
 			}
 			else
 			{
 				ElementData.ElementMetaData = SiteLocationMetaData;
 			}
-			
+
 			ActorMap[InSiteLocation.Id] = ElementData;
 
 			DirectLink?.CacheElement(CurrentDocument, InSiteLocation, ElementData);
@@ -1982,7 +1982,7 @@ namespace DatasmithRevitExporter
 					// Export the DatasmithMesh in a task while we parse the rest of the document.
 					// The task result indicates if the export was successful and if the associated FDatasmithFacadeMeshElement can be added to the scene.
 					MeshMap[MeshName] = new Tuple<FDatasmithFacadeMeshElement, Task<bool>>(InMeshElement, Task.Run<bool>(
-						() => 
+						() =>
 						{
 							using (FDatasmithFacadeMesh DatasmithMesh = ParsePolymesh(InPolymesh, MeshName))
 							{
@@ -2015,7 +2015,7 @@ namespace DatasmithRevitExporter
 			for (int UVIndex = 0; UVIndex < InPolymesh.UVs.Count; ++UVIndex)
 			{
 				UV CurrentUV = InPolymesh.UVs[UVIndex];
-				DatasmithMesh.SetUV(UVChannelIndex, UVIndex, (float)CurrentUV.U, (float)CurrentUV.V);
+				DatasmithMesh.SetUV(UVChannelIndex, UVIndex, CurrentUV.U, CurrentUV.V);
 			}
 
 			// Add the triangle vertex indexes to the Datasmith mesh.
@@ -2031,7 +2031,7 @@ namespace DatasmithRevitExporter
 				XYZ Normal = InPolymesh.Normals[NormalIndex];
 				DatasmithMesh.SetNormal(NormalIndex, (float)Normal.X, (float)Normal.Y, (float)Normal.Z);
 			}
-		
+
 			return DatasmithMesh;
 		}
 
@@ -2162,9 +2162,9 @@ namespace DatasmithRevitExporter
 			}
 
 			// DirectLink: if host is hidden, go up the hierarchy (NOTE this does not apply for linked documents)
-			if (DirectLink != null && 
-				HostElement != null && 
-				CurrentDocument.ActiveView != null && 
+			if (DirectLink != null &&
+				HostElement != null &&
+				CurrentDocument.ActiveView != null &&
 				HostElement.IsHidden(CurrentDocument.ActiveView))
 			{
 				return GetHostElement(HostElement.Id);
@@ -2248,7 +2248,7 @@ namespace DatasmithRevitExporter
 
 				FBaseElementData ElementData = ActorMap[ElemId];
 				FBaseElementData ParentElementData = ActorMap[ParentElementId];
-				
+
 				if (!ParentElementData.ChildElements.Contains(ElementData))
 				{
 					ParentElementData.ChildElements.Add(ElementData);
@@ -2294,24 +2294,24 @@ namespace DatasmithRevitExporter
 				}
 			}
 
-			float[] worldMatrix = new float[16];
-
-			worldMatrix[0] = (float)transformBasisX.X;
-			worldMatrix[1] = (float)transformBasisX.Y;
-			worldMatrix[2] = (float)transformBasisX.Z;
-			worldMatrix[3] = 0.0F;
-			worldMatrix[4] = (float)transformBasisY.X;
-			worldMatrix[5] = (float)transformBasisY.Y;
-			worldMatrix[6] = (float)transformBasisY.Z;
-			worldMatrix[7] = 0.0F;
-			worldMatrix[8] = (float)transformBasisZ.X;
-			worldMatrix[9] = (float)transformBasisZ.Y;
-			worldMatrix[10] = (float)transformBasisZ.Z;
-			worldMatrix[11] = 0.0F;
-			worldMatrix[12] = (float)transformOrigin.X;
-			worldMatrix[13] = (float)transformOrigin.Y;
-			worldMatrix[14] = (float)transformOrigin.Z;
-			worldMatrix[15] = 1.0F;
+			double[] worldMatrix = new double[16];
+
+			worldMatrix[0] = transformBasisX.X;
+			worldMatrix[1] = transformBasisX.Y;
+			worldMatrix[2] = transformBasisX.Z;
+			worldMatrix[3] = 0.0;
+			worldMatrix[4] = transformBasisY.X;
+			worldMatrix[5] = transformBasisY.Y;
+			worldMatrix[6] = transformBasisY.Z;
+			worldMatrix[7] = 0.0;
+			worldMatrix[8] = transformBasisZ.X;
+			worldMatrix[9] = transformBasisZ.Y;
+			worldMatrix[10] = transformBasisZ.Z;
+			worldMatrix[11] = 0.0;
+			worldMatrix[12] = transformOrigin.X;
+			worldMatrix[13] = transformOrigin.Y;
+			worldMatrix[14] = transformOrigin.Z;
+			worldMatrix[15] = 1.0;
 
 			// Set the world transform of the Datasmith actor.
 			IOActor.SetWorldTransform(worldMatrix);
diff --git a/Engine/Source/Runtime/Datasmith/DatasmithCore/Private/DatasmithMesh.cpp b/Engine/Source/Runtime/Datasmith/DatasmithCore/Private/DatasmithMesh.cpp
--- a/Engine/Source/Runtime/Datasmith/DatasmithCore/Private/DatasmithMesh.cpp
+++ b/Engine/Source/Runtime/Datasmith/DatasmithCore/Private/DatasmithMesh.cpp
@@ -18,10 +18,10 @@ public:
 
 	FString Name;
 
-	TArray< FVector > Vertices;
+	TArray< FVector3f > Vertices;
 	TArray< uint32 > Indices;
 
-	TArray< FVector > Normals;
+	TArray< FVector3f > Normals;
 	TArray< int32 > MaterialIndices;
 
 	TArray< TArray< FVector2D > > UVs;
@@ -35,7 +35,7 @@ public:
 
 	int32 LightmapUVChannel;
 
-	FBox Extents;
+	FBox3f Extents;
 
 private:
 	TSet< int32 > IdsInUse;
@@ -114,7 +114,7 @@ void FDatasmithMesh::SetFacesCount(int32 NumFaces)
 	Impl->MaterialIndices.Init(0, NumFaces);
 
 	Impl->Normals.Empty(NumFaces * 3);
-	Impl->Normals.Init( FVector(ForceInitToZero), NumFaces * 3 );
+	Impl->Normals.Init( FVector3f::ZeroVector, NumFaces * 3 );
 
 	Impl->FaceSmoothingMasks.Empty( NumFaces * 3 );
 	Impl->FaceSmoothingMasks.AddZeroed( NumFaces * 3 );
@@ -172,7 +172,7 @@ bool FDatasmithMesh::IsMaterialIdUsed(int32 MaterialId) const
 void FDatasmithMesh::SetVerticesCount(int32 NumVerts)
 {
 	Impl->Vertices.Empty(NumVerts);
-	Impl->Vertices.Init( FVector(ForceInit), NumVerts );
+	Impl->Vertices.Init( FVector3f::ZeroVector, NumVerts );
 }
 
 int32 FDatasmithMesh::GetVerticesCount() const
@@ -184,7 +184,7 @@ void FDatasmithMesh::SetVertex(int32 Index, float X, float Y, float Z)
 {
 	if ( Impl->Vertices.IsValidIndex( Index ) )
 	{
-		FVector Vertex(X, Y, Z);
+		FVector3f Vertex(X, Y, Z);
 
 		Impl->Vertices[Index] = Vertex;
 
@@ -192,7 +192,7 @@ void FDatasmithMesh::SetVertex(int32 Index, float X, float Y, float Z)
 	}
 }
 
-FVector FDatasmithMesh::GetVertex(int32 Index) const
+FVector3f FDatasmithMesh::GetVertex(int32 Index) const
 {
 	if ( Impl->Vertices.IsValidIndex( Index ) )
 	{
@@ -200,7 +200,7 @@ FVector FDatasmithMesh::GetVertex(int32 Index) const
 	}
 	else
 	{
-		return FVector::ZeroVector;
+		return FVector3f::ZeroVector;
 	}
 }
 
@@ -208,11 +208,11 @@ void FDatasmithMesh::SetNormal(int32 Index, float X, float Y, float Z)
 {
 	if ( Impl->Normals.IsValidIndex( Index ) )
 	{
-		Impl->Normals[Index] = FVector(X, Y, Z).GetSafeNormal();
+		Impl->Normals[Index] = FVector3f(X, Y, Z).GetSafeNormal();
 	}
 }
 
-FVector FDatasmithMesh::GetNormal(int32 Index) const
+FVector3f FDatasmithMesh::GetNormal(int32 Index) const
 {
 	if ( Impl->Normals.IsValidIndex( Index ) )
 	{
@@ -220,7 +220,7 @@ FVector FDatasmithMesh::GetNormal(int32 Index) const
 	}
 	else
 	{
-		return FVector::ZeroVector;
+		return FVector3f::ZeroVector;
 	}
 }
 
@@ -463,7 +463,7 @@ float FDatasmithMesh::ComputeArea() const
 	return Area;
 }
 
-FBox FDatasmithMesh::GetExtents() const
+FBox3f FDatasmithMesh::GetExtents() const
 {
 	return Impl->Extents;
 }
diff --git a/Engine/Source/Runtime/Datasmith/DatasmithCore/Private/DatasmithSceneElementsImpl.h b/Engine/Source/Runtime/Datasmith/DatasmithCore/Private/DatasmithSceneElementsImpl.h
--- a/Engine/Source/Runtime/Datasmith/DatasmithCore/Private/DatasmithSceneElementsImpl.h
+++ b/Engine/Source/Runtime/Datasmith/DatasmithCore/Private/DatasmithSceneElementsImpl.h
@@ -109,7 +109,7 @@ public:
 	FDatasmithActorElementImpl(const TCHAR* InName, EDatasmithElementType InType);
 
 	virtual FVector GetTranslation() const override { return Translation.Get(); }
-	virtual void SetTranslation(float InX, float InY, float InZ, bool bKeepChildrenRelative) override { SetTranslation( FVector( InX, InY, InZ ), bKeepChildrenRelative ); }
+	virtual void SetTranslation(double InX, double InY, double InZ, bool bKeepChildrenRelative) override { SetTranslation( FVector( InX, InY, InZ ), bKeepChildrenRelative ); }
 	virtual void SetTranslation(const FVector& Value, bool bKeepChildrenRelative) override
 	{
 		if (bKeepChildrenRelative)
@@ -124,7 +124,7 @@ public:
 	}
 
 	virtual FVector GetScale() const override { return Scale.Get(); }
-	virtual void SetScale(float InX, float InY, float InZ, bool bKeepChildrenRelative) override { SetScale( FVector( InX, InY, InZ ), bKeepChildrenRelative ); }
+	virtual void SetScale(double InX, double InY, double InZ, bool bKeepChildrenRelative) override { SetScale( FVector( InX, InY, InZ ), bKeepChildrenRelative ); }
 	virtual void SetScale(const FVector& Value, bool bKeepChildrenRelative) override
 	{
 		if (bKeepChildrenRelative)
@@ -139,7 +139,7 @@ public:
 	}
 
 	virtual FQuat GetRotation() const override { return Rotation; }
-	virtual void SetRotation(float InX, float InY, float InZ, float InW, bool bKeepChildrenRelative) override { SetRotation( FQuat( InX, InY, InZ, InW ), bKeepChildrenRelative ); }
+	virtual void SetRotation(double InX, double InY, double InZ, double InW, bool bKeepChildrenRelative) override { SetRotation( FQuat( InX, InY, InZ, InW ), bKeepChildrenRelative ); }
 	virtual void SetRotation(const FQuat& Value, bool bKeepChildrenRelative) override
 	{
 		if (bKeepChildrenRelative)
@@ -349,8 +349,8 @@ public:
 	virtual FMD5Hash GetFileHash() const override { return FileHash; }
 	virtual void SetFileHash(FMD5Hash Hash) override { FileHash = Hash; }
 
-	virtual void SetDimensions(const float InArea, const float InWidth, const float InHeight, const float InDepth) override { Area = InArea; Width = InWidth; Height = InHeight; Depth = InDepth;};
-	virtual FVector GetDimensions() const override { return FVector{ Width, Height, Depth }; }
+	virtual void SetDimensions(float InArea, float InWidth, float InHeight, float InDepth) override { Area = InArea; Width = InWidth; Height = InHeight; Depth = InDepth;};
+	virtual FVector3f GetDimensions() const override { return { Width, Height, Depth }; }
 
 	virtual float GetArea() const override { return Area; }
 	virtual float GetWidth() const override { return Width; }
diff --git a/Engine/Source/Runtime/Datasmith/DatasmithCore/Private/DatasmithSceneXmlReader.cpp b/Engine/Source/Runtime/Datasmith/DatasmithCore/Private/DatasmithSceneXmlReader.cpp
--- a/Engine/Source/Runtime/Datasmith/DatasmithCore/Private/DatasmithSceneXmlReader.cpp
+++ b/Engine/Source/Runtime/Datasmith/DatasmithCore/Private/DatasmithSceneXmlReader.cpp
@@ -27,85 +27,63 @@
 FDatasmithSceneXmlReader::FDatasmithSceneXmlReader() = default;
 FDatasmithSceneXmlReader::~FDatasmithSceneXmlReader() = default;
 
-namespace DatasmithSceneXmlReaderImpl
-{
-	const TCHAR* ActorTags[] = { DATASMITH_ACTORNAME, DATASMITH_ACTORMESHNAME, DATASMITH_CAMERANAME, DATASMITH_LIGHTNAME,
-		DATASMITH_CUSTOMACTORNAME, DATASMITH_LANDSCAPENAME, DATASMITH_POSTPROCESSVOLUME, DATASMITH_ACTORHIERARCHICALINSTANCEDMESHNAME, DATASMITH_DECALACTORNAME };
-
-	template< typename T >
-	T ValueFromString( const FString& InString ) = delete;
 
-	template<>
-	float ValueFromString< float >( const FString& InString )
-	{
-		return FCString::Atof( *InString );
-	}
-
-	template<>
-	double ValueFromString< double >( const FString& InString )
-	{
-		return FCString::Atod( *InString );
-	}
+const TCHAR* ActorTagsView[] = { DATASMITH_ACTORNAME, DATASMITH_ACTORMESHNAME, DATASMITH_CAMERANAME, DATASMITH_LIGHTNAME,
+DATASMITH_CUSTOMACTORNAME, DATASMITH_LANDSCAPENAME, DATASMITH_POSTPROCESSVOLUME, DATASMITH_ACTORHIERARCHICALINSTANCEDMESHNAME, DATASMITH_DECALACTORNAME };
 
-	template<>
-	FVector ValueFromString< FVector >( const FString& InString )
-	{
-		FVector Value;
-		Value.InitFromString( InString );
+template<>
+float FDatasmithSceneXmlReader::ValueFromString< float >( const FString& InString ) const
+{
+	return FCString::Atof( *InString );
+}
 
-		return Value;
-	}
+template<>
+double FDatasmithSceneXmlReader::ValueFromString< double >( const FString& InString ) const
+{
+	return FCString::Atod( *InString );
+}
 
-	template<>
-	FColor ValueFromString< FColor >( const FString& InString )
-	{
-		FColor Value;
-		Value.InitFromString( InString );
+template<>
+FVector FDatasmithSceneXmlReader::ValueFromString< FVector >( const FString& InString ) const
+{
+	FVector Value;
+	Value.InitFromString( InString );
 
-		return Value;
-	}
+	return Value;
+}
 
-	template<>
-	FLinearColor ValueFromString< FLinearColor >( const FString& InString )
-	{
-		FLinearColor Value;
-		Value.InitFromString( InString );
+template<>
+FColor FDatasmithSceneXmlReader::ValueFromString< FColor >( const FString& InString ) const
+{
+	FColor Value;
+	Value.InitFromString( InString );
 
-		return Value;
-	}
+	return Value;
+}
 
-	template<>
-	int32 ValueFromString< int32 >( const FString& InString )
-	{
-		return FCString::Atoi( *InString );
-	}
+template<>
+FLinearColor FDatasmithSceneXmlReader::ValueFromString< FLinearColor >( const FString& InString ) const
+{
+	FLinearColor Value;
+	Value.InitFromString( InString );
 
-	template<>
-	bool ValueFromString< bool >( const FString& InString )
-	{
-		return InString.ToBool();
-	}
+	return Value;
+}
 
-	UE_NODISCARD FString UnsanitizeXMLText(const FString& InString)
-	{
-		FString OutString = InString;
-		OutString.ReplaceInline( TEXT("&apos;"), TEXT("'")  );
-		OutString.ReplaceInline( TEXT("&quot;"), TEXT("\"") );
-		OutString.ReplaceInline( TEXT("&gt;"),   TEXT(">")  );
-		OutString.ReplaceInline( TEXT("&lt;"),   TEXT("<")  );
-		OutString.ReplaceInline( TEXT("&amp;"),  TEXT("&")  );
-		return OutString;
-	}
+template<>
+int32 FDatasmithSceneXmlReader::ValueFromString< int32 >( const FString& InString ) const
+{
+	return FCString::Atoi( *InString );
 }
 
-template<typename T>
-T FDatasmithSceneXmlReader::ValueFromString( const FString& InString ) const
+template<>
+bool FDatasmithSceneXmlReader::ValueFromString< bool >( const FString& InString ) const
 {
-	return DatasmithSceneXmlReaderImpl::ValueFromString<T>(InString);
+	return InString.ToBool();
 }
 
 
-void FDatasmithSceneXmlReader::PatchUpVersion(TSharedRef< IDatasmithScene >& OutScene)
+void FDatasmithSceneXmlReader::PatchUpVersion(TSharedRef< IDatasmithScene >& OutScene) const
 {
 	//@todo parse version string in proper version object
 	// Handle legacy behavior, when materials from the first actor using a mesh applied its materials to it.
@@ -136,6 +114,82 @@ void FDatasmithSceneXmlReader::PatchUpVersion(TSharedRef< IDatasmithScene >& Out
 	}
 }
 
+UE_NODISCARD FString FDatasmithSceneXmlReader::UnsanitizeXMLText(const FString& InString) const
+{
+	FString OutString = InString;
+	OutString.ReplaceInline( TEXT("&apos;"), TEXT("'")  );
+	OutString.ReplaceInline( TEXT("&quot;"), TEXT("\"") );
+	OutString.ReplaceInline( TEXT("&gt;"),   TEXT(">")  );
+	OutString.ReplaceInline( TEXT("&lt;"),   TEXT("<")  );
+	OutString.ReplaceInline( TEXT("&amp;"),  TEXT("&")  );
+	return OutString;
+}
+
+FVector FDatasmithSceneXmlReader::VectorFromNode(FXmlNode* InNode, const TCHAR* XName, const TCHAR* YName, const TCHAR* ZName) const
+{
+	return {
+		ValueFromString<double>(InNode->GetAttribute(XName)),
+		ValueFromString<double>(InNode->GetAttribute(YName)),
+		ValueFromString<double>(InNode->GetAttribute(ZName))
+	};
+}
+
+FQuat FDatasmithSceneXmlReader::QuatFromHexString(const FString& HexString) const
+{
+	if (HexString.Len() / 2 >= 4 * sizeof(double))
+	{
+		double Double[4];
+		FString::ToHexBlob( HexString, (uint8*)Double, sizeof(Double) );
+		return FQuat( Double[0], Double[1], Double[2], Double[3] );
+	}
+	else if (HexString.Len() / 2 >= 4 * sizeof(float))
+	{
+		float Floats[4];
+		FString::ToHexBlob( HexString, (uint8*)Floats, sizeof(Floats) );
+		return FQuat( Floats[0], Floats[1], Floats[2], Floats[3] );
+	}
+
+	return FQuat::Identity;
+}
+
+FQuat FDatasmithSceneXmlReader::QuatFromNode(FXmlNode* InNode) const
+{
+	FString RotationBlob = InNode->GetAttribute(TEXT("qhex64"));
+	if (RotationBlob.IsEmpty())
+	{
+		RotationBlob = InNode->GetAttribute(TEXT("qhex"));
+	}
+
+	if (!RotationBlob.IsEmpty())
+	{
+		return QuatFromHexString(RotationBlob);
+	}
+
+	return FQuat(
+		ValueFromString<double>(InNode->GetAttribute(TEXT("qx"))),
+		ValueFromString<double>(InNode->GetAttribute(TEXT("qy"))),
+		ValueFromString<double>(InNode->GetAttribute(TEXT("qz"))),
+		ValueFromString<double>(InNode->GetAttribute(TEXT("qw")))
+	);
+}
+
+FTransform FDatasmithSceneXmlReader::ParseTransform(FXmlNode* InNode) const
+{
+	return {
+		QuatFromNode(InNode),
+		VectorFromNode(InNode, TEXT("tx"), TEXT("ty"), TEXT("tz")),
+		VectorFromNode(InNode, TEXT("sx"), TEXT("sy"), TEXT("sz"))
+	};
+}
+
+void FDatasmithSceneXmlReader::ParseTransform(FXmlNode* InNode, TSharedPtr< IDatasmithActorElement >& OutElement) const
+{
+	FTransform Tmp = ParseTransform(InNode);
+	OutElement->SetTranslation(Tmp.GetTranslation());
+	OutElement->SetRotation(Tmp.GetRotation());
+	OutElement->SetScale(Tmp.GetScale3D());
+}
+
 void FDatasmithSceneXmlReader::ParseElement(FXmlNode* InNode, TSharedRef<IDatasmithElement> OutElement) const
 {
 	OutElement->SetLabel( *InNode->GetAttribute( TEXT("label") ) );
@@ -459,58 +513,6 @@ void FDatasmithSceneXmlReader::ParseTexture(FXmlNode* InNode, FString& OutTextur
 	}
 }
 
-void FDatasmithSceneXmlReader::ParseTransform(FXmlNode* InNode, TSharedPtr< IDatasmithActorElement >& OutElement) const
-{
-	OutElement->SetTranslation(ValueFromString<float>(InNode->GetAttribute(TEXT("tx"))), ValueFromString<float>(InNode->GetAttribute(TEXT("ty"))),
-		ValueFromString<float>(InNode->GetAttribute(TEXT("tz"))));
-
-	FString RotationBlob(InNode->GetAttribute(TEXT("qhex")));
-	if (!RotationBlob.IsEmpty())
-	{
-		OutElement->SetRotation( QuatFromHexString( RotationBlob ) );
-	}
-	else
-	{
-		OutElement->SetRotation(ValueFromString<float>(InNode->GetAttribute(TEXT("qx"))), ValueFromString<float>(InNode->GetAttribute(TEXT("qy"))),
-			ValueFromString<float>(InNode->GetAttribute(TEXT("qz"))), ValueFromString<float>(InNode->GetAttribute(TEXT("qw"))));
-	}
-
-	OutElement->SetScale(ValueFromString<float>(InNode->GetAttribute(TEXT("sx"))), ValueFromString<float>(InNode->GetAttribute(TEXT("sy"))), ValueFromString<float>(InNode->GetAttribute(TEXT("sz"))));
-}
-
-FTransform FDatasmithSceneXmlReader::ParseTransform(FXmlNode* InNode) const
-{
-	FVector Translation(ValueFromString<float>(InNode->GetAttribute(TEXT("tx"))), ValueFromString<float>(InNode->GetAttribute(TEXT("ty"))),
-		ValueFromString<float>(InNode->GetAttribute(TEXT("tz"))));
-
-	FQuat Quaternion;
-
-	FString RotationBlob(InNode->GetAttribute(TEXT("qhex")));
-	if (!RotationBlob.IsEmpty())
-	{
-		Quaternion = QuatFromHexString(RotationBlob);
-	}
-	else
-	{
-		Quaternion = FQuat(ValueFromString<float>(InNode->GetAttribute(TEXT("qx"))), ValueFromString<float>(InNode->GetAttribute(TEXT("qy"))),
-			ValueFromString<float>(InNode->GetAttribute(TEXT("qz"))), ValueFromString<float>(InNode->GetAttribute(TEXT("qw"))));
-	}
-
-	FVector Scale(ValueFromString<float>(InNode->GetAttribute(TEXT("sx"))), ValueFromString<float>(InNode->GetAttribute(TEXT("sy"))), ValueFromString<float>(InNode->GetAttribute(TEXT("sz"))));
-
-	return FTransform(Quaternion, Translation, Scale);
-}
-
-
-FQuat FDatasmithSceneXmlReader::QuatFromHexString(const FString& HexString) const
-{
-	float Floats[4];
-	FString::ToHexBlob( HexString, (uint8*)Floats, sizeof(Floats) );
-
-	FQuat QuatResult( Floats[0], Floats[1], Floats[2], Floats[3] );
-
-	return QuatResult;
-}
 
 void FDatasmithSceneXmlReader::ParsePostProcess(FXmlNode *InNode, const TSharedPtr< IDatasmithPostProcessElement >& Element) const
 {
@@ -518,7 +520,6 @@ void FDatasmithSceneXmlReader::ParsePostProcess(FXmlNode *InNode, const TSharedP
 
 	const TArray<FXmlNode*>& CompNodes = InNode->GetChildrenNodes();
 
-	FLinearColor Color;
 	for (int j = 0; j < CompNodes.Num(); j++)
 	{
 		if (CompNodes[j]->GetTag() == DATASMITH_POSTPRODUCTIONTEMP)
@@ -529,26 +530,16 @@ void FDatasmithSceneXmlReader::ParsePostProcess(FXmlNode *InNode, const TSharedP
 		{
 			Element->SetVignette(ValueFromString<float>(CompNodes[j]->GetAttribute(TEXT("value"))));
 		}
-		// #ueent_wip: broken PostProcess serialization.
-		// - create missing tags
-		// - update Writer
-// 		else if (CompNodes[j]->GetTag() == DATASMITH_POSTPRODUCTIONVIGNETTE)
-// 		{
-// 			Element->SetDof(ValueFromString<float>(CompNodes[j]->GetAttribute(TEXT("value"))));
-// 		}
-// 		else if (CompNodes[j]->GetTag() == DATASMITH_POSTPRODUCTIONVIGNETTE)
-// 		{
-// 			Element->SetMotionBlur(ValueFromString<float>(CompNodes[j]->GetAttribute(TEXT("value"))));
-// 		}
-		else if (CompNodes[j]->GetTag() == DATASMITH_POSTPRODUCTIONSATURATION)
-		{
-			Element->SetSaturation(ValueFromString<float>(CompNodes[j]->GetAttribute(TEXT("value"))));
-		}
 		else if (CompNodes[j]->GetTag() == DATASMITH_POSTPRODUCTIONCOLOR)
 		{
+			FLinearColor Color;
 			ParseColor(CompNodes[j], Color);
 			Element->SetColorFilter(Color);
 		}
+		else if (CompNodes[j]->GetTag() == DATASMITH_POSTPRODUCTIONSATURATION)
+		{
+			Element->SetSaturation(ValueFromString<float>(CompNodes[j]->GetAttribute(TEXT("value"))));
+		}
 		else if (CompNodes[j]->GetTag() == DATASMITH_POSTPRODUCTIONCAMERAISO)
 		{
 			Element->SetCameraISO(ValueFromString<float>(CompNodes[j]->GetAttribute(TEXT("value"))));
@@ -757,8 +748,6 @@ void FDatasmithSceneXmlReader::ParseActor(FXmlNode* InNode, TSharedPtr<IDatasmit
 			// Recursively parse the children, can be of any supported actor type
 			for (FXmlNode* ChildActorNode : ChildNode->GetChildrenNodes())
 			{
-				TArrayView< const TCHAR* > ActorTagsView( DatasmithSceneXmlReaderImpl::ActorTags );
-
 				for ( const TCHAR* ActorTag : ActorTagsView )
 				{
 					if ( ChildActorNode->GetTag() == ActorTag )
@@ -915,7 +904,7 @@ void FDatasmithSceneXmlReader::ParseLight(FXmlNode* InNode, TSharedPtr<IDatasmit
 		}
 		else if (ChildNode->GetTag() == DATASMITH_LIGHTIESROTATION)
 		{
-			OutElement->SetIesRotation( QuatFromHexString( ChildNode->GetAttribute( TEXT("qhex") ) ) );
+			OutElement->SetIesRotation( QuatFromNode( ChildNode ) );
 		}
 		else if (ChildNode->GetTag() == DATASMITH_LIGHTINTENSITYNAME)
 		{
@@ -1121,7 +1110,6 @@ bool FDatasmithSceneXmlReader::ParseBuffer(const FString& XmlBuffer, TSharedRef<
 
 bool FDatasmithSceneXmlReader::ParseXmlFile(TSharedRef< IDatasmithScene >& OutScene, bool bInAppend)
 {
-	using namespace DatasmithSceneXmlReaderImpl;
 	if (!XmlFile.IsValid())
 	{
 		return false;
@@ -1352,8 +1340,6 @@ bool FDatasmithSceneXmlReader::ParseXmlFile(TSharedRef< IDatasmithScene >& OutSc
 		//READ ACTORS
 		else
 		{
-			TArrayView< const TCHAR* > ActorTagsView( DatasmithSceneXmlReaderImpl::ActorTags );
-
 			for ( const TCHAR* ActorTag : ActorTagsView )
 			{
 				if ( Nodes[i]->GetTag() == ActorTag )
@@ -1646,8 +1632,9 @@ void FDatasmithSceneXmlReader::ParseDecalMaterial(FXmlNode* InNode, TSharedPtr<
 	}
 }
 
+
 template< typename ExpressionInputType >
-void ParseExpressionInput(const FXmlNode* InNode, TSharedPtr< IDatasmithUEPbrMaterialElement >& OutElement, ExpressionInputType& ExpressionInput)
+void FDatasmithSceneXmlReader::ParseExpressionInput(const FXmlNode* InNode, TSharedPtr< IDatasmithUEPbrMaterialElement >& OutElement, ExpressionInputType& ExpressionInput) const
 {
 	if ( !InNode )
 	{
@@ -1662,9 +1649,9 @@ void ParseExpressionInput(const FXmlNode* InNode, TSharedPtr< IDatasmithUEPbrMat
 		// From 4.23 Expressions are serialized as <Input Name="0" expression="5" OutputIndex="0"/>
 		// So if the Name is used and that backward compatibility is desired, the Node Tag can be used instead of the "Name" Attribute.
 
-		int32 ExpressionIndex = DatasmithSceneXmlReaderImpl::ValueFromString<int32>( ExpressionIndexAttribute );
+		int32 ExpressionIndex = ValueFromString<int32>( ExpressionIndexAttribute );
 
-		int32 OutputIndex = DatasmithSceneXmlReaderImpl::ValueFromString<int32>( InNode->GetAttribute( TEXT("OutputIndex") ) );
+		int32 OutputIndex = ValueFromString<int32>( InNode->GetAttribute( TEXT("OutputIndex") ) );
 
 		IDatasmithMaterialExpression* Expression = OutElement->GetExpression( ExpressionIndex );
 
@@ -1711,9 +1698,9 @@ void FDatasmithSceneXmlReader::ParseUEPbrMaterial(FXmlNode* InNode, TSharedPtr<
 				if ( Expression )
 				{
 					IDatasmithMaterialExpressionTextureCoordinate* TextureCoordinateExpression = static_cast< IDatasmithMaterialExpressionTextureCoordinate* >( Expression );
-					TextureCoordinateExpression->SetCoordinateIndex( DatasmithSceneXmlReaderImpl::ValueFromString< int32 >( ChildNode->GetAttribute( TEXT("Index") ) ) );
-					TextureCoordinateExpression->SetUTiling( DatasmithSceneXmlReaderImpl::ValueFromString< float >( ChildNode->GetAttribute( TEXT("UTiling") ) ) );
-					TextureCoordinateExpression->SetVTiling( DatasmithSceneXmlReaderImpl::ValueFromString< float >( ChildNode->GetAttribute( TEXT("VTiling") ) ) );
+					TextureCoordinateExpression->SetCoordinateIndex( ValueFromString< int32 >( ChildNode->GetAttribute( TEXT("Index") ) ) );
+					TextureCoordinateExpression->SetUTiling( ValueFromString< float >( ChildNode->GetAttribute( TEXT("UTiling") ) ) );
+					TextureCoordinateExpression->SetVTiling( ValueFromString< float >( ChildNode->GetAttribute( TEXT("VTiling") ) ) );
 				}
 			}
 			else if ( ChildNode->GetTag() == TEXT("FlattenNormal") )
@@ -1735,7 +1722,7 @@ void FDatasmithSceneXmlReader::ParseUEPbrMaterial(FXmlNode* InNode, TSharedPtr<
 
 					IDatasmithMaterialExpressionBool* ConstantBool = static_cast< IDatasmithMaterialExpressionBool* >( Expression );
 
-					ConstantBool->GetBool() = DatasmithSceneXmlReaderImpl::ValueFromString< bool >( ChildNode->GetAttribute( TEXT("Constant") ) );
+					ConstantBool->GetBool() = ValueFromString< bool >( ChildNode->GetAttribute( TEXT("Constant") ) );
 				}
 			}
 			else if ( ChildNode->GetTag() == TEXT("Color") )
@@ -1748,7 +1735,7 @@ void FDatasmithSceneXmlReader::ParseUEPbrMaterial(FXmlNode* InNode, TSharedPtr<
 
 					IDatasmithMaterialExpressionColor* ConstantColor = static_cast< IDatasmithMaterialExpressionColor* >( Expression );
 
-					ConstantColor->GetColor() = DatasmithSceneXmlReaderImpl::ValueFromString< FLinearColor >( ChildNode->GetAttribute( TEXT("Constant") ) );
+					ConstantColor->GetColor() = ValueFromString< FLinearColor >( ChildNode->GetAttribute( TEXT("Constant") ) );
 				}
 			}
 			else if ( ChildNode->GetTag() == TEXT("Scalar") )
@@ -1761,7 +1748,7 @@ void FDatasmithSceneXmlReader::ParseUEPbrMaterial(FXmlNode* InNode, TSharedPtr<
 
 					IDatasmithMaterialExpressionScalar* ConstantScalar = static_cast< IDatasmithMaterialExpressionScalar* >( Expression );
 
-					ConstantScalar->GetScalar() = DatasmithSceneXmlReaderImpl::ValueFromString< float >( ChildNode->GetAttribute( TEXT("Constant") ) );
+					ConstantScalar->GetScalar() = ValueFromString< float >( ChildNode->GetAttribute( TEXT("Constant") ) );
 				}
 			}
 			else if ( ChildNode->GetTag() == TEXT("FunctionCall") )
@@ -1788,23 +1775,23 @@ void FDatasmithSceneXmlReader::ParseUEPbrMaterial(FXmlNode* InNode, TSharedPtr<
 						if (CustomChildNode->GetTag() == TEXT("Code"))
 						{
 							const FString& SanitizedContent = CustomChildNode->GetContent();
-							const FString& Content = DatasmithSceneXmlReaderImpl::UnsanitizeXMLText(SanitizedContent);
+							const FString& Content = UnsanitizeXMLText(SanitizedContent);
 							Expression->SetCode(*Content);
 						}
 						else if (CustomChildNode->GetTag() == TEXT("Include"))
 						{
-							const FString& Path = DatasmithSceneXmlReaderImpl::UnsanitizeXMLText(CustomChildNode->GetAttribute(TEXT("path")));
+							const FString& Path = UnsanitizeXMLText(CustomChildNode->GetAttribute(TEXT("path")));
 							Expression->AddIncludeFilePath(*Path);
 						}
 						else if (CustomChildNode->GetTag() == TEXT("Define"))
 						{
-							const FString& Value = DatasmithSceneXmlReaderImpl::UnsanitizeXMLText(CustomChildNode->GetAttribute(TEXT("value")));
+							const FString& Value = UnsanitizeXMLText(CustomChildNode->GetAttribute(TEXT("value")));
 							Expression->AddAdditionalDefine(*Value);
 						}
 						else if (CustomChildNode->GetTag() == TEXT("Arg"))
 						{
-							int32 Index = DatasmithSceneXmlReaderImpl::ValueFromString<int32>(CustomChildNode->GetAttribute(TEXT("index")));
-							const FString& Value = DatasmithSceneXmlReaderImpl::UnsanitizeXMLText(CustomChildNode->GetAttribute(TEXT("name")));
+							int32 Index = ValueFromString<int32>(CustomChildNode->GetAttribute(TEXT("index")));
+							const FString& Value = UnsanitizeXMLText(CustomChildNode->GetAttribute(TEXT("name")));
 							Expression->SetArgumentName(Index, *Value);
 						}
 					}
@@ -1812,7 +1799,7 @@ void FDatasmithSceneXmlReader::ParseUEPbrMaterial(FXmlNode* InNode, TSharedPtr<
 					const FString& Description = ChildNode->GetAttribute( TEXT("Description") );
 					Expression->SetDescription( *Description );
 
-					int32 OutputType = DatasmithSceneXmlReaderImpl::ValueFromString< int32 >( ChildNode->GetAttribute( TEXT("OutputType") ) );;
+					int32 OutputType = ValueFromString< int32 >( ChildNode->GetAttribute( TEXT("OutputType") ) );;
 					Expression->SetOutputType(EDatasmithShaderDataType(OutputType));
 				}
 			}
@@ -1869,7 +1856,7 @@ void FDatasmithSceneXmlReader::ParseUEPbrMaterial(FXmlNode* InNode, TSharedPtr<
 					for ( const FXmlNode* InputChildNode : ChildNode->GetChildrenNodes() )
 					{
 						const FString& NameAttribute = InputChildNode->GetAttribute(TEXT("Name"));
-						int32 InputIndex = DatasmithSceneXmlReaderImpl::ValueFromString< int32 >( NameAttribute.IsEmpty() ? InputChildNode->GetTag() : NameAttribute );
+						int32 InputIndex = ValueFromString< int32 >( NameAttribute.IsEmpty() ? InputChildNode->GetTag() : NameAttribute );
 
 						if (IDatasmithExpressionInput* Input = GenericExpression->GetInput( InputIndex ))
 						{
@@ -1885,7 +1872,7 @@ void FDatasmithSceneXmlReader::ParseUEPbrMaterial(FXmlNode* InNode, TSharedPtr<
 
 	const TArray<FXmlNode*>& ChildrenNodes = InNode->GetChildrenNodes();
 
-	auto TryConnectMaterialInput = [&ChildrenNodes, &OutElement](IDatasmithExpressionInput& Input)
+	auto TryConnectMaterialInput = [&](IDatasmithExpressionInput& Input)
 	{
 		const TCHAR* InputName = Input.GetName();
 		for (FXmlNode* XmlNode : ChildrenNodes)
@@ -1915,27 +1902,27 @@ void FDatasmithSceneXmlReader::ParseUEPbrMaterial(FXmlNode* InNode, TSharedPtr<
 	{
 		if ( ChildNode->GetTag() == DATASMITH_USEMATERIALATTRIBUTESNAME )
 		{
-			OutElement->SetUseMaterialAttributes( DatasmithSceneXmlReaderImpl::ValueFromString< bool >( ChildNode->GetAttribute( TEXT("enabled") ) ) );
+			OutElement->SetUseMaterialAttributes( ValueFromString< bool >( ChildNode->GetAttribute( TEXT("enabled") ) ) );
 		}
 		else if ( ChildNode->GetTag() == DATASMITH_TWOSIDEDVALUENAME )
 		{
-			OutElement->SetTwoSided( DatasmithSceneXmlReaderImpl::ValueFromString< bool >( ChildNode->GetAttribute( TEXT("enabled") ) ) );
+			OutElement->SetTwoSided( ValueFromString< bool >( ChildNode->GetAttribute( TEXT("enabled") ) ) );
 		}
 		else if ( ChildNode->GetTag() == DATASMITH_BLENDMODE )
 		{
-			OutElement->SetBlendMode( DatasmithSceneXmlReaderImpl::ValueFromString< int >( ChildNode->GetAttribute( TEXT("value") ) ) );
+			OutElement->SetBlendMode( ValueFromString< int >( ChildNode->GetAttribute( TEXT("value") ) ) );
 		}
 		else if ( ChildNode->GetTag() == DATASMITH_OPACITYMASKCLIPVALUE )
 		{
-			OutElement->SetOpacityMaskClipValue( DatasmithSceneXmlReaderImpl::ValueFromString< float >( ChildNode->GetAttribute( TEXT("value") ) ) );
+			OutElement->SetOpacityMaskClipValue( ValueFromString< float >( ChildNode->GetAttribute( TEXT("value") ) ) );
 		}
 		else if (ChildNode->GetTag() == DATASMITH_TRANSLUCENCYLIGHTINGMODE)
 		{
-			OutElement->SetTranslucencyLightingMode(DatasmithSceneXmlReaderImpl::ValueFromString< int >(ChildNode->GetAttribute(TEXT("value"))));
+			OutElement->SetTranslucencyLightingMode(ValueFromString< int >(ChildNode->GetAttribute(TEXT("value"))));
 		}
 		else if ( ChildNode->GetTag() == DATASMITH_FUNCTIONLYVALUENAME )
 		{
-			OutElement->SetMaterialFunctionOnly( DatasmithSceneXmlReaderImpl::ValueFromString< bool >( ChildNode->GetAttribute( TEXT("enabled") ) ) );
+			OutElement->SetMaterialFunctionOnly( ValueFromString< bool >( ChildNode->GetAttribute( TEXT("enabled") ) ) );
 		}
 		else if ( ChildNode->GetTag() == DATASMITH_SHADINGMODEL )
 		{
diff --git a/Engine/Source/Runtime/Datasmith/DatasmithCore/Private/DatasmithSceneXmlWriter.cpp b/Engine/Source/Runtime/Datasmith/DatasmithCore/Private/DatasmithSceneXmlWriter.cpp
--- a/Engine/Source/Runtime/Datasmith/DatasmithCore/Private/DatasmithSceneXmlWriter.cpp
+++ b/Engine/Source/Runtime/Datasmith/DatasmithCore/Private/DatasmithSceneXmlWriter.cpp
@@ -25,6 +25,7 @@ public:
 	static void WriteBool(FArchive& Archive, int32 Indent, const TCHAR* Prefix, bool bValue);
 	static void WriteRGB(FArchive& Archive, int32 Indent, const TCHAR* Prefix, FLinearColor Color);
 	static void WriteValue(FArchive& Archive, int32 Indent, const TCHAR* Prefix, float Value, bool bForceWriteAlways = false, FString Desc = FString());
+	static void WriteValue(FArchive& Archive, int32 Indent, const TCHAR* Prefix, double Value, bool bForceWriteAlways, FString Desc = FString()) = delete;
 	static void AppendXMLChild(FString& XmlString, int32 Indent, const TCHAR* Tag, const TCHAR* AttributeName, const TCHAR* AttributeValue);
 
 	static void WriteTexture(FArchive& Archive, int32 Indent, const TCHAR* Prefix, const TCHAR* Name, FDatasmithTextureSampler UV);
@@ -302,13 +303,8 @@ void FDatasmithSceneXmlWriterImpl::WriteTransform(const FTransform& Transform, F
 
 FString FDatasmithSceneXmlWriterImpl::QuatToHexString(FQuat Value)
 {
-	float Floats[4];
-	Floats[0] = Value.X;
-	Floats[1] = Value.Y;
-	Floats[2] = Value.Z;
-	Floats[3] = Value.W;
-
-	FString Result = TEXT("qhex=\"") + FString::FromHexBlob((uint8*)Floats, sizeof(Floats)) + TEXT("\"");
+	double Tmp[4] = {Value.X, Value.Y, Value.Z, Value.W};
+	FString Result = TEXT("qhex64=\"") + FString::FromHexBlob((uint8*)Tmp, sizeof(Tmp)) + TEXT("\"");
 	return Result;
 }
 
@@ -1012,7 +1008,7 @@ void FDatasmithSceneXmlWriterImpl::WritePostProcessElement(const TSharedPtr< IDa
 
 	WriteIndent(Archive, Indent);
 
-	FString XmlString = TEXT("<") + FString(DATASMITH_POSTPRODUCTIONNAME) + TEXT(">") + LINE_TERMINATOR;
+	FString XmlString = TEXT("<") DATASMITH_POSTPRODUCTIONNAME TEXT(">") LINE_TERMINATOR;
 	SerializeToArchive( Archive, XmlString );
 
 	if (PostProcessElement->GetTemperature() != 6500.0f)
diff --git a/Engine/Source/Runtime/Datasmith/DatasmithCore/Private/DatasmithUtils.cpp b/Engine/Source/Runtime/Datasmith/DatasmithCore/Private/DatasmithUtils.cpp
--- a/Engine/Source/Runtime/Datasmith/DatasmithCore/Private/DatasmithUtils.cpp
+++ b/Engine/Source/Runtime/Datasmith/DatasmithCore/Private/DatasmithUtils.cpp
@@ -208,12 +208,10 @@ const TCHAR* FDatasmithUtils::GetShortAppName()
 	return TEXT("Datasmith");
 }
 
-double FDatasmithUtils::AreaTriangle3D(FVector v0, FVector v1, FVector v2)
+float FDatasmithUtils::AreaTriangle3D(const FVector3f& v0, const FVector3f& v1, const FVector3f& v2)
 {
-	FVector TriangleNormal = (v1 - v2) ^ (v0 - v2);
-	double Area = TriangleNormal.Size() * 0.5;
-
-	return Area;
+	FVector3f TriangleNormal = (v1 - v2) ^ (v0 - v2);
+	return TriangleNormal.Size() * 0.5f;
 }
 
 bool FDatasmithMeshUtils::ToRawMesh(const FDatasmithMesh& Mesh, FRawMesh& RawMesh, bool bValidateRawMesh)
@@ -229,7 +227,7 @@ bool FDatasmithMeshUtils::ToRawMesh(const FDatasmithMesh& Mesh, FRawMesh& RawMes
 
 	for ( int32 i = 0; i < Mesh.GetVerticesCount(); ++i )
 	{
-		RawMesh.VertexPositions.Add( (FVector3f)Mesh.GetVertex( i ) );
+		RawMesh.VertexPositions.Add( Mesh.GetVertex( i ) );
 	}
 
 	RawMesh.FaceMaterialIndices.Reserve( Mesh.GetFacesCount() );
@@ -255,7 +253,7 @@ bool FDatasmithMeshUtils::ToRawMesh(const FDatasmithMesh& Mesh, FRawMesh& RawMes
 
 		for ( int32 j = 0; j < 3; ++j )
 		{
-			RawMesh.WedgeTangentZ.Add( (FVector3f)Mesh.GetNormal( FaceIndex * 3 + j ) );
+			RawMesh.WedgeTangentZ.Add( Mesh.GetNormal( FaceIndex * 3 + j ) );
 		}
 	}
 
diff --git a/Engine/Source/Runtime/Datasmith/DatasmithCore/Public/DatasmithMaterialElements.h b/Engine/Source/Runtime/Datasmith/DatasmithCore/Public/DatasmithMaterialElements.h
--- a/Engine/Source/Runtime/Datasmith/DatasmithCore/Public/DatasmithMaterialElements.h
+++ b/Engine/Source/Runtime/Datasmith/DatasmithCore/Public/DatasmithMaterialElements.h
@@ -225,8 +225,9 @@ public:
 	virtual IDatasmithExpressionInput& GetClearCoatRoughness() = 0;
 	virtual IDatasmithExpressionInput& GetMaterialAttributes() = 0;
 
+	/** InBlendMode must match the values of EBlendMode from EngineTypes.h */
 	virtual int GetBlendMode() const = 0;
-	virtual void SetBlendMode( int bInBlendMode ) = 0;
+	virtual void SetBlendMode( int InBlendMode ) = 0;
 
 	virtual bool GetTwoSided() const = 0;
 	virtual void SetTwoSided( bool bTwoSided ) = 0;
diff --git a/Engine/Source/Runtime/Datasmith/DatasmithCore/Public/DatasmithMesh.h b/Engine/Source/Runtime/Datasmith/DatasmithCore/Public/DatasmithMesh.h
--- a/Engine/Source/Runtime/Datasmith/DatasmithCore/Public/DatasmithMesh.h
+++ b/Engine/Source/Runtime/Datasmith/DatasmithCore/Public/DatasmithMesh.h
@@ -74,7 +74,7 @@ public:
 	 * @z position on the z axis
 	 */
 	void SetVertex(int32 Index, float X, float Y, float Z);
-	FVector GetVertex(int32 Index) const;
+	FVector3f GetVertex(int32 Index) const;
 
 	//--------------------------
 	// Normals
@@ -88,7 +88,7 @@ public:
 	 * @z direction on the z axis
 	 */
 	void SetNormal(int32 Index, float X, float Y, float Z);
-	FVector GetNormal(int32 Index) const;
+	FVector3f GetNormal(int32 Index) const;
 
 	//--------------------------
 	// UVs
@@ -215,7 +215,7 @@ public:
 	float ComputeArea() const;
 
 	/** Returns the bounding box containing all vertices of this mesh */
-	FBox GetExtents() const;
+	FBox3f GetExtents() const;
 
 private:
 	class FDatasmithMeshImpl;
diff --git a/Engine/Source/Runtime/Datasmith/DatasmithCore/Public/DatasmithSceneXmlReader.h b/Engine/Source/Runtime/Datasmith/DatasmithCore/Public/DatasmithSceneXmlReader.h
--- a/Engine/Source/Runtime/Datasmith/DatasmithCore/Public/DatasmithSceneXmlReader.h
+++ b/Engine/Source/Runtime/Datasmith/DatasmithCore/Public/DatasmithSceneXmlReader.h
@@ -22,8 +22,16 @@ public:
 
 private:
 	bool ParseXmlFile(TSharedRef< IDatasmithScene >& OutScene, bool bInAppend = false);
+	void PatchUpVersion(TSharedRef< IDatasmithScene >& OutScene) const;
 
-	void PatchUpVersion(TSharedRef< IDatasmithScene >& OutScene);
+	UE_NODISCARD FString UnsanitizeXMLText(const FString& InString) const;
+
+	template<typename T> T ValueFromString(const FString& InString) const = delete;
+	FVector VectorFromNode(FXmlNode* InNode, const TCHAR* XName, const TCHAR* YName, const TCHAR* ZName) const;
+	FQuat QuatFromHexString(const FString& HexString) const;
+	FQuat QuatFromNode(FXmlNode* InNode) const;
+	FTransform ParseTransform(FXmlNode* InNode) const;
+	void ParseTransform(FXmlNode* InNode, TSharedPtr< IDatasmithActorElement >& OutElement) const;
 
 	void ParseElement(FXmlNode* InNode, TSharedRef<IDatasmithElement> OutElement) const;
 	void ParseLevelSequence(FXmlNode* InNode, const TSharedRef<IDatasmithLevelSequenceElement>& OutElement) const;
@@ -36,8 +44,6 @@ private:
 	void ParseMesh(FXmlNode* InNode, TSharedPtr<IDatasmithMeshElement>& OutElement) const;
 	void ParseTextureElement(FXmlNode* InNode, TSharedPtr<IDatasmithTextureElement>& OutElement) const;
 	void ParseTexture(FXmlNode* InNode, FString& OutTextureFilename, FDatasmithTextureSampler& OutTextureUV) const;
-	void ParseTransform(FXmlNode* InNode, TSharedPtr< IDatasmithActorElement >& OutElement) const;
-	FTransform ParseTransform(FXmlNode* InNode) const;
 	void ParseActor(FXmlNode* InNode, TSharedPtr<IDatasmithActorElement>& InOutElement, TSharedRef< IDatasmithScene > Scene, TMap< FString, TSharedPtr<IDatasmithActorElement> >& Actors) const;
 	void ParseMeshActor(FXmlNode* InNode, TSharedPtr<IDatasmithMeshActorElement>& OutElement, TSharedRef< IDatasmithScene > Scene) const;
 	void ParseHierarchicalInstancedStaticMeshActor(FXmlNode* InNode, TSharedPtr<IDatasmithHierarchicalInstancedStaticMeshActorElement>& OutElement, TSharedRef< IDatasmithScene > Scene) const;
@@ -61,11 +67,10 @@ private:
 	bool LoadFromFile(const FString& InFilename);
 	bool LoadFromBuffer(const FString& XmlBuffer);
 
-	FString ResolveFilePath(const FString& InAssetFile) const;
-
-	template<typename T> T ValueFromString(const FString& InString) const;
-	FQuat QuatFromHexString(const FString& HexString) const;
+	template< typename ExpressionInputType >
+	void ParseExpressionInput(const FXmlNode* InNode, TSharedPtr< IDatasmithUEPbrMaterialElement >& OutElement, ExpressionInputType& ExpressionInput) const;
 
 	TUniquePtr< FXmlFile > XmlFile;
 	FString ProjectPath;
-};
\ No newline at end of file
+};
+
diff --git a/Engine/Source/Runtime/Datasmith/DatasmithCore/Public/DatasmithUtils.h b/Engine/Source/Runtime/Datasmith/DatasmithCore/Public/DatasmithUtils.h
--- a/Engine/Source/Runtime/Datasmith/DatasmithCore/Public/DatasmithUtils.h
+++ b/Engine/Source/Runtime/Datasmith/DatasmithCore/Public/DatasmithUtils.h
@@ -50,7 +50,7 @@ public:
 	static const TCHAR* GetShortAppName();
 
 	/** Computes the area of a triangle */
-	static double AreaTriangle3D(FVector v0, FVector v1, FVector v2);
+	static float AreaTriangle3D(const FVector3f& v0, const FVector3f& v1, const FVector3f& v2);
 
 	enum class EModelCoordSystem : uint8
 	{
@@ -137,7 +137,6 @@ public:
 	static bool ToRawMesh(const FDatasmithMesh& Mesh, FRawMesh& RawMesh, bool bValidateRawMesh = true);
 	static bool ToMeshDescription(const FDatasmithMesh& DsMesh, FMeshDescription& MeshDescription);
 
-
 	/**
 	 * Validates that the given UV Channel does not contain a degenerated triangle.
 	 *
diff --git a/Engine/Source/Runtime/Datasmith/DatasmithCore/Public/IDatasmithSceneElements.h b/Engine/Source/Runtime/Datasmith/DatasmithCore/Public/IDatasmithSceneElements.h
--- a/Engine/Source/Runtime/Datasmith/DatasmithCore/Public/IDatasmithSceneElements.h
+++ b/Engine/Source/Runtime/Datasmith/DatasmithCore/Public/IDatasmithSceneElements.h
@@ -77,7 +77,7 @@ public:
 	virtual FVector GetTranslation() const = 0;
 
 	/** Set absolute translation of this entity */
-	virtual void SetTranslation(float InX, float InY, float InZ, bool bKeepChildrenRelative = true) = 0;
+	virtual void SetTranslation(double InX, double InY, double InZ, bool bKeepChildrenRelative = true) = 0;
 
 	/** Set absolute translation of this entity */
 	virtual void SetTranslation(const FVector& Value, bool bKeepChildrenRelative = true) = 0;
@@ -86,7 +86,7 @@ public:
 	virtual FVector GetScale() const = 0;
 
 	/** Set absolute scale of this entity */
-	virtual void SetScale(float InX, float InY, float InZ, bool bKeepChildrenRelative = true) = 0;
+	virtual void SetScale(double InX, double InY, double InZ, bool bKeepChildrenRelative = true) = 0;
 
 	/** Set absolute scale of this entity */
 	virtual void SetScale(const FVector& Value, bool bKeepChildrenRelative = true) = 0;
@@ -95,7 +95,7 @@ public:
 	virtual FQuat GetRotation() const = 0;
 
 	/** Set rotation (in quaternion format) of this entity */
-	virtual void SetRotation(float InX, float InY, float InZ, float InW, bool bKeepChildrenRelative = true) = 0;
+	virtual void SetRotation(double InX, double InY, double InZ, double InW, bool bKeepChildrenRelative = true) = 0;
 
 	/** Set rotation (in quaternion format) of this entity */
 	virtual void SetRotation(const FQuat& Value, bool bKeepChildrenRelative = true) = 0;
@@ -183,10 +183,10 @@ public:
 	 * @param InHeight bounding box height
 	 * @param InDepth bounding box depth
 	 */
-	virtual void SetDimensions(const float InArea, const float InWidth, const float InHeight, const float InDepth) = 0;
+	virtual void SetDimensions(float InArea, float InWidth, float InHeight, float InDepth) = 0;
 
 	/** Get the bounding box dimension of the mesh, in a vector in the form of (Width, Height, Depth )*/
-	virtual FVector GetDimensions() const = 0;
+	virtual FVector3f GetDimensions() const = 0;
 
 	/** Get the total surface area */
 	virtual float GetArea() const = 0;
@@ -252,7 +252,7 @@ public:
 	 * Adds a new material override to the Actor Element
 	 *
 	 * @param MaterialName name of the material, it should be unique
-	 * @param Id material identifier to be used with mesh sub-material indices
+	 * @param Id material identifier to be used with mesh sub-material indices. Use -1 to override all material slots.
 	 */
 	virtual void AddMaterialOverride(const TCHAR* MaterialName, int32 Id) = 0;
 
@@ -1443,7 +1443,7 @@ public:
 	/** Resets all the settings on the scene */
 	virtual void Reset() = 0;
 
-	/** Sets the name of the host application which created the scene */
+	/** Returns the name of the host application which created the scene */
 	virtual const TCHAR* GetHost() const = 0;
 
 	/**
