Datasmith: Add support for compressed mesh serialization

diff --git a/Engine/Source/Developer/Datasmith/DatasmithExporter/Private/DatasmithMeshExporter.cpp b/Engine/Source/Developer/Datasmith/DatasmithExporter/Private/DatasmithMeshExporter.cpp
--- a/Engine/Source/Developer/Datasmith/DatasmithExporter/Private/DatasmithMeshExporter.cpp
+++ b/Engine/Source/Developer/Datasmith/DatasmithExporter/Private/DatasmithMeshExporter.cpp
@@ -280,15 +280,13 @@ bool FDatasmithMeshExporterImpl::WriteMeshFile( const FDatasmithMeshExporterOpti
 
 	int32 LegacyMeshCount = 0;
 	*Archive << LegacyMeshCount; // the legacy importer expect a mesh count on the first bytes. Just in case a new file would end up parsed by the legacy code...
-	*Archive << Pack;
+	OutHash = Pack.Serialize(*Archive);
 
-	OutHash = Pack.OutHash;
 	return !Archive->IsError();
 }
 
 
-
-TSharedPtr< IDatasmithMeshElement > FDatasmithMeshExporter::ExportToUObject(const TCHAR* Filepath, const TCHAR* Filename, FDatasmithMesh& Mesh, FDatasmithMesh* CollisionMesh, EDSExportLightmapUV LightmapUV, bool NewFormat/*=false */)
+TSharedPtr< IDatasmithMeshElement > FDatasmithMeshExporter::ExportToUObject(const TCHAR* Filepath, const TCHAR* Filename, FDatasmithMesh& Mesh, FDatasmithMesh* CollisionMesh, EDSExportLightmapUV LightmapUV, bool NewFormat)
 {
 	FString FullPath( GetMeshFilePath( Filepath, Filename ) );
 
@@ -306,12 +304,20 @@ TSharedPtr< IDatasmithMeshElement > FDatasmithMeshExporter::ExportToUObject(cons
 	return ExportedMeshElement;
 }
 
-bool FDatasmithMeshExporter::ExportToUObject( TSharedPtr< IDatasmithMeshElement >& MeshElement, const TCHAR* Filepath, FDatasmithMesh& Mesh, FDatasmithMesh* CollisionMesh, EDSExportLightmapUV LightmapUV )
+
+bool FDatasmithMeshExporter::ExportToUObject( TSharedPtr< IDatasmithMeshElement >& MeshElement, const TCHAR* Filepath, FDatasmithMesh& Mesh, FDatasmithMesh* CollisionMesh, EDSExportLightmapUV LightmapUV, bool NewFormat )
 {
 	FString FullPath( GetMeshFilePath( Filepath, MeshElement->GetName() ) );
 	FDatasmithMeshExporterOptions ExportOptions( FullPath, Mesh, LightmapUV, CollisionMesh );
 
-	return LegacyImpl->DoExport( MeshElement, ExportOptions );
+	if (NewFormat)
+	{
+		return Impl->DoExport( MeshElement, ExportOptions );
+	}
+	else
+	{
+		return LegacyImpl->DoExport( MeshElement, ExportOptions );
+	}
 }
 
 
diff --git a/Engine/Source/Developer/Datasmith/DatasmithExporter/Public/DatasmithMeshExporter.h b/Engine/Source/Developer/Datasmith/DatasmithExporter/Public/DatasmithMeshExporter.h
--- a/Engine/Source/Developer/Datasmith/DatasmithExporter/Public/DatasmithMeshExporter.h
+++ b/Engine/Source/Developer/Datasmith/DatasmithExporter/Public/DatasmithMeshExporter.h
@@ -37,7 +37,7 @@ public:
 	 *
 	 * @return				True if export was successful.
 	 */
-	bool ExportToUObject( TSharedPtr< IDatasmithMeshElement >& MeshElement, const TCHAR* Filepath, FDatasmithMesh& Mesh, FDatasmithMesh* CollisionMesh, EDSExportLightmapUV LightmapUV );
+	bool ExportToUObject( TSharedPtr< IDatasmithMeshElement >& MeshElement, const TCHAR* Filepath, FDatasmithMesh& Mesh, FDatasmithMesh* CollisionMesh, EDSExportLightmapUV LightmapUV, bool NewFormat=false );
 
 	/**
 	 * @return The error that happened during the last export, if any
diff --git a/Engine/Source/Runtime/Datasmith/DatasmithCore/Private/DatasmithMeshSerialization.cpp b/Engine/Source/Runtime/Datasmith/DatasmithCore/Private/DatasmithMeshSerialization.cpp
--- a/Engine/Source/Runtime/Datasmith/DatasmithCore/Private/DatasmithMeshSerialization.cpp
+++ b/Engine/Source/Runtime/Datasmith/DatasmithCore/Private/DatasmithMeshSerialization.cpp
@@ -3,6 +3,7 @@
 
 #include "DatasmithMeshUObject.h"
 
+#include "Compression/OodleDataCompression.h"
 #include "HAL/FileManager.h"
 #include "Serialization/MemoryReader.h"
 #include "Serialization/MemoryWriter.h"
@@ -16,47 +17,183 @@ void operator<<(FArchive& Ar, FDatasmithMeshModels& Models)
 	Ar << Models.SourceModels;
 }
 
-void operator<<(FArchive& Ar, FDatasmithPackedMeshes& Pack)
+
+namespace DatasmithMeshSerializationImpl
 {
-	FString Guard = Ar.IsLoading() ? TEXT("") : TEXT("FDatasmithPackedMeshes");
-	Ar << Guard;
-	if (!ensure(Guard == TEXT("FDatasmithPackedMeshes")))
+enum class ECompressionMethod
+{
+	ECM_ZLib  = 1,
+	ECM_Gzip  = 2,
+	ECM_LZ4   = 3,
+	ECM_Oodle = 4,
+
+	ECM_Default = ECM_Oodle,
+};
+
+FName NAME_Oodle("Oodle");
+FName GetMethodName(ECompressionMethod MethodCode)
+{
+	switch (MethodCode)
 	{
-		Ar.SetError();
-		return;
+		case ECompressionMethod::ECM_ZLib: return NAME_Zlib;
+		case ECompressionMethod::ECM_Gzip: return NAME_Gzip;
+		case ECompressionMethod::ECM_LZ4:  return NAME_LZ4;
+		case ECompressionMethod::ECM_Oodle:return NAME_Oodle;
+		default: ensure(0); return NAME_None;
+	}
+}
+
+bool CompressInline(TArray<uint8>& UncompressedData, ECompressionMethod Method)
+{
+	int32 UncompressedSize = UncompressedData.Num();
+	FName MethodName = GetMethodName(Method);
+	int32 CompressedSize = (Method != ECompressionMethod::ECM_Oodle)
+		? FCompression::CompressMemoryBound(MethodName, UncompressedSize)
+		: FOodleDataCompression::CompressedBufferSizeNeeded(UncompressedSize);
+
+	TArray<uint8> CompressedData;
+	int32 HeaderSize = 5;
+	CompressedData.AddUninitialized(CompressedSize + HeaderSize);
+
+	// header
+	{
+		FMemoryWriter Ar(CompressedData);
+		uint8 MethodCode = uint8(Method);
+		Ar << MethodCode;
+		Ar << UncompressedSize;
+		check(HeaderSize == Ar.Tell());
+	}
+
+	if (Method == ECompressionMethod::ECM_Oodle)
+	{
+		auto Compressor = FOodleDataCompression::ECompressor::Kraken;
+		auto Level = FOodleDataCompression::ECompressionLevel::VeryFast;
+		if (int64 ActualCompressedSize = FOodleDataCompression::Compress(CompressedData.GetData() + HeaderSize, CompressedSize, UncompressedData.GetData(), UncompressedSize, Compressor, Level))
+		{
+			CompressedData.SetNum(ActualCompressedSize + HeaderSize, true);
+			UncompressedData = MoveTemp(CompressedData);
+
+			return true;
+		}
+	}
+	else
+	{
+		if (FCompression::CompressMemory(MethodName, CompressedData.GetData() + HeaderSize, CompressedSize, UncompressedData.GetData(), UncompressedSize))
+		{
+
+			CompressedData.SetNum(CompressedSize + HeaderSize, true);
+			UncompressedData = MoveTemp(CompressedData);
+
+			return true;
+		}
+	}
+
+	UE_LOG(LogDatasmith, Warning, TEXT("Compression failed"));
+	return false;
+}
+
+bool DecompressInline(TArray<uint8>& CompressedData)
+{
+	ECompressionMethod Method;
+	uint8* BufferStart = nullptr;
+	int32 UncompressedSize = -1;
+	int32 CompressedSize = CompressedData.Num();
+	int32 HeaderSize = 0;
+	{
+		FMemoryReader Ar(CompressedData);
+		uint8 MethodCode = 0;
+		Ar << MethodCode;
+		Method = ECompressionMethod(MethodCode);
+		Ar << UncompressedSize;
+		HeaderSize = Ar.Tell();
+	}
+
+	FName MethodName = GetMethodName(Method);
+	if (MethodName == NAME_Oodle)
+	{
+		TArray<uint8> UncompressedData;
+		UncompressedData.SetNumUninitialized(UncompressedSize);
+		bool Ok = FOodleDataCompression::Decompress(UncompressedData.GetData(), UncompressedSize, CompressedData.GetData() + HeaderSize, CompressedData.Num() - HeaderSize);
+		if (Ok)
+		{
+			CompressedData = MoveTemp(UncompressedData);
+			return true;
+		}
 	}
+	else if (MethodName != NAME_None)
+	{
+		TArray<uint8> UncompressedData;
+		UncompressedData.SetNumUninitialized(UncompressedSize);
+		bool Ok = FCompression::UncompressMemory(MethodName, UncompressedData.GetData(), UncompressedData.Num(), CompressedData.GetData() + HeaderSize, CompressedData.Num() - HeaderSize);
+		if (Ok)
+		{
+			CompressedData = MoveTemp(UncompressedData);
+			return true;
+		}
+	}
+	UE_LOG(LogDatasmith, Warning, TEXT("Decompression failed"));
+	return false;
+}
+
+} // ns Impl
+
+
+FMD5Hash FDatasmithPackedMeshes::Serialize(FArchive& Ar, bool bCompressed)
+{
+	using namespace DatasmithMeshSerializationImpl;
 
-	uint8 SerialVersionMajor = 1; // increment this when forward compatibility is not possible
-	Ar << SerialVersionMajor;
-	if (SerialVersionMajor > 1)
+	FString Guard = Ar.IsLoading() ? TEXT("") : TEXT("FDatasmithPackedMeshes");
+	Ar << Guard;
+	if (!ensure(Guard == TEXT("FDatasmithPackedMeshes")))
 	{
 		Ar.SetError();
-		return;
+		return {};
 	}
 
-	uint8 SerialVersionMinor = 0; // increment this to handle logic changes that preserves forward compatibility
-	Ar << SerialVersionMinor;
+	uint32 SerialVersion = 0;
+	Ar << SerialVersion;
 
-	uint8 BufferType = 0;
+	enum EBufferType{ RawMeshDescription, CompressedMeshDescription };
+	uint8 BufferType = bCompressed ? CompressedMeshDescription : RawMeshDescription;
 	Ar << BufferType; // (MeshDesc, Zipped Mesh desc;...)
 
+	FMD5Hash OutHash;
 	if (Ar.IsLoading())
 	{
+		FCustomVersionContainer CustomVersions;
+		CustomVersions.Serialize(Ar);
+
 		TArray<uint8> Bytes;
 		Ar << Bytes;
+
+		if (BufferType == CompressedMeshDescription)
+		{
+			DecompressInline(Bytes);
+		}
+
 		FMemoryReader Buffer(Bytes, true);
-		Buffer << Pack.MeshesToExport;
+		Buffer.SetCustomVersions(CustomVersions);
+		Buffer << MeshesToExport;
 	}
 	else
 	{
 		TArray<uint8> Bytes;
 		FMemoryWriter Buffer(Bytes, true);
-		Buffer << Pack.MeshesToExport;
+		Buffer << MeshesToExport;
+
+		// MeshDescriptions uses custom versionning,
+		const_cast<FCustomVersionContainer&>(Buffer.GetCustomVersions()).Serialize(Ar);
+
+		if (BufferType == CompressedMeshDescription)
+		{
+			CompressInline(Bytes, ECompressionMethod::ECM_Default);
+		}
 		Ar << Bytes;
 		FMD5 Md5;
 		Md5.Update(Bytes.GetData(), Bytes.Num());
-		Pack.OutHash.Set(Md5);
+		OutHash.Set(Md5);
 	}
+	return OutHash;
 }
 
 TOptional<FMeshDescription> ExtractToMeshDescription(FDatasmithMeshSourceModel& SourceModel)
@@ -170,7 +307,7 @@ TArray<FDatasmithMeshModels> GetDatasmithMeshFromMeshPath(const FString& MeshPat
 	if (LeagacyNumMeshesCount == 0)
 	{
 		FDatasmithPackedMeshes Pack;
-		*Archive << Pack;
+		Pack.Serialize(*Archive);
 
 		if (!Archive->IsError())
 		{
@@ -188,4 +325,5 @@ TArray<FDatasmithMeshModels> GetDatasmithMeshFromMeshPath(const FString& MeshPat
 	}
 
 	return Result;
-}
\ No newline at end of file
+}
+
diff --git a/Engine/Source/Runtime/Datasmith/DatasmithCore/Public/DatasmithMeshSerialization.h b/Engine/Source/Runtime/Datasmith/DatasmithCore/Public/DatasmithMeshSerialization.h
--- a/Engine/Source/Runtime/Datasmith/DatasmithCore/Public/DatasmithMeshSerialization.h
+++ b/Engine/Source/Runtime/Datasmith/DatasmithCore/Public/DatasmithMeshSerialization.h
@@ -16,11 +16,11 @@ struct FDatasmithMeshModels
 	DATASMITHCORE_API friend void operator << (FArchive& Ar, FDatasmithMeshModels& Models);
 };
 
-struct FDatasmithPackedMeshes
+struct DATASMITHCORE_API FDatasmithPackedMeshes
 {
 	TArray<FDatasmithMeshModels> MeshesToExport;
-	DATASMITHCORE_API friend void operator << (FArchive& Ar, FDatasmithPackedMeshes& Pack);
-	FMD5Hash OutHash;
+
+	FMD5Hash Serialize(FArchive& Ar, bool bSaveCompressed=true);
 };
 
 DATASMITHCORE_API TArray<FDatasmithMeshModels> GetDatasmithMeshFromMeshPath(const FString& MeshPath);
