Datasmith: remove LODScreenSize unused API

diff --git a/Engine/Source/Runtime/Datasmith/DatasmithCore/Private/DatasmithSceneElementsImpl.cpp b/Engine/Source/Runtime/Datasmith/DatasmithCore/Private/DatasmithSceneElementsImpl.cpp
--- a/Engine/Source/Runtime/Datasmith/DatasmithCore/Private/DatasmithSceneElementsImpl.cpp
+++ b/Engine/Source/Runtime/Datasmith/DatasmithCore/Private/DatasmithSceneElementsImpl.cpp
@@ -1046,7 +1046,6 @@ FDatasmithSceneImpl::FDatasmithSceneImpl(const TCHAR * InName)
 	RegisterReferenceProxy(LevelVariantSets, "LevelVariantSets" );
 	RegisterReferenceProxy(PostProcess,      "PostProcess"      );
 
-	Store.RegisterParameter(LODScreenSizes,     "LODScreenSizes"     );
 	Store.RegisterParameter(Hostname,           "Hostname"           );
 	Store.RegisterParameter(ExporterVersion,    "ExporterVersion"    );
 	Store.RegisterParameter(ExporterSDKVersion, "ExporterSDKVersion" );
@@ -1070,7 +1069,6 @@ void FDatasmithSceneImpl::Reset()
 	MetaData.Empty();
 	LevelSequences.Empty();
 	LevelVariantSets.Empty();
-	LODScreenSizes.Get().Reset();
 	PostProcess.Inner.Reset();
 	ElementToMetaDataMap.Empty();
 
diff --git a/Engine/Source/Runtime/Datasmith/DatasmithCore/Private/DatasmithSceneElementsImpl.h b/Engine/Source/Runtime/Datasmith/DatasmithCore/Private/DatasmithSceneElementsImpl.h
--- a/Engine/Source/Runtime/Datasmith/DatasmithCore/Private/DatasmithSceneElementsImpl.h
+++ b/Engine/Source/Runtime/Datasmith/DatasmithCore/Private/DatasmithSceneElementsImpl.h
@@ -1619,12 +1619,7 @@ public:
 	virtual void SetUsePhysicalSky(bool bInUsePhysicalSky) override { bUseSky = bInUsePhysicalSky; }
 	virtual bool GetUsePhysicalSky() const override { return bUseSky; }
 
-	virtual void AddLODScreenSize( float ScreenSize ) override { LODScreenSizes.Get().Add( FMath::Clamp( ScreenSize, 0.f, 1.f ) ); }
-	virtual int32 GetLODScreenSizesCount() const override { return LODScreenSizes.Get().Num(); }
-	virtual float GetLODScreenSize(int32 InIndex) const override { return LODScreenSizes.Get().IsValidIndex( InIndex ) ? LODScreenSizes.Get()[InIndex] : 0.f; }
-
 	virtual void AddMetaData(const TSharedPtr< IDatasmithMetaDataElement >& InMetaData) override { MetaData.Add(InMetaData); GetElementToMetaDataCache().Add(InMetaData->GetAssociatedElement(), InMetaData); }
-
 	virtual int32 GetMetaDataCount() const override { return MetaData.Num(); }
 	virtual TSharedPtr< IDatasmithMetaDataElement > GetMetaData(int32 InIndex) override;
 	virtual const TSharedPtr< IDatasmithMetaDataElement >& GetMetaData(int32 InIndex) const override;
@@ -1651,7 +1646,6 @@ public:
 	virtual void AttachActorToSceneRoot(const TSharedPtr< IDatasmithActorElement >& Child, EDatasmithActorAttachmentRule AttachmentRule) override;
 
 private:
-	TMap< TSharedPtr< IDatasmithElement >, TSharedPtr< IDatasmithMetaDataElement> >& GetElementToMetaDataCache() const;
 
 	TDatasmithReferenceArrayProxy<IDatasmithActorElement>            Actors;
 	TDatasmithReferenceArrayProxy<IDatasmithMeshElement>             Meshes;
@@ -1662,8 +1656,6 @@ private:
 	TDatasmithReferenceArrayProxy<IDatasmithLevelVariantSetsElement> LevelVariantSets;
 	TDatasmithReferenceProxy<IDatasmithPostProcessElement>           PostProcess;
 
-	TReflected<TArray<float>> LODScreenSizes;
-
 	TReflected<FString> Hostname;
 	TReflected<FString> ExporterVersion;
 	TReflected<FString> ExporterSDKVersion;
@@ -1678,6 +1670,7 @@ private:
 
 	TReflected<bool> bUseSky;
 
-	// Internal cache for faster metadata access per-element, should be accessed via GetMetaDataCache(), do not use directly.
+	// Internal cache for faster metadata access per-element, should be accessed via GetElementToMetaDataCache(), do not use directly.
 	mutable TMap< TSharedPtr< IDatasmithElement >, TSharedPtr< IDatasmithMetaDataElement> > ElementToMetaDataMap;
+	TMap< TSharedPtr< IDatasmithElement >, TSharedPtr< IDatasmithMetaDataElement> >& GetElementToMetaDataCache() const;
 };
diff --git a/Engine/Source/Runtime/Datasmith/DatasmithCore/Private/DatasmithSceneXmlReader.cpp b/Engine/Source/Runtime/Datasmith/DatasmithCore/Private/DatasmithSceneXmlReader.cpp
--- a/Engine/Source/Runtime/Datasmith/DatasmithCore/Private/DatasmithSceneXmlReader.cpp
+++ b/Engine/Source/Runtime/Datasmith/DatasmithCore/Private/DatasmithSceneXmlReader.cpp
@@ -1326,12 +1326,6 @@ bool FDatasmithSceneXmlReader::ParseXmlFile(TSharedRef< IDatasmithScene >& OutSc
 			ParseMetaData(Nodes[i], MetaData, OutScene, Actors );
 			OutScene->AddMetaData(MetaData);
 		}
-		//LOD SCREEN SIZES
-		else if (Nodes[i]->GetTag() == DATASMITH_LODSCREENSIZE)
-		{
-			float LODScreenSize = ValueFromString<float>( Nodes[i]->GetAttribute(TEXT("value")) );
-			OutScene->AddLODScreenSize(LODScreenSize);
-		}
 		// EXPORT STATS
 		else if (Nodes[i]->GetTag() == DATASMITH_EXPORT)
 		{
diff --git a/Engine/Source/Runtime/Datasmith/DatasmithCore/Private/DatasmithSceneXmlWriter.cpp b/Engine/Source/Runtime/Datasmith/DatasmithCore/Private/DatasmithSceneXmlWriter.cpp
--- a/Engine/Source/Runtime/Datasmith/DatasmithCore/Private/DatasmithSceneXmlWriter.cpp
+++ b/Engine/Source/Runtime/Datasmith/DatasmithCore/Private/DatasmithSceneXmlWriter.cpp
@@ -1982,16 +1982,6 @@ void FDatasmithSceneXmlWriter::Serialize( TSharedRef< IDatasmithScene > Datasmit
 		FDatasmithSceneXmlWriterImpl::WriteMetaDataElement(MetaDataElement, Archive, 1);
 	}
 
-
-	for ( int32 LODScreenSizeIndex = 0; LODScreenSizeIndex < DatasmithScene->GetLODScreenSizesCount(); ++LODScreenSizeIndex )
-	{
-		float LODScreenSize = DatasmithScene->GetLODScreenSize( LODScreenSizeIndex );
-
-		FDatasmithSceneXmlWriterImpl::WriteIndent(Archive, 1);
-		XmlString = TEXT("<") + FString(DATASMITH_LODSCREENSIZE) + TEXT(" value=\"") + FString::SanitizeFloat( LODScreenSize ) + "\"/>" + LINE_TERMINATOR;
-		FDatasmithSceneXmlWriterImpl::SerializeToArchive( Archive, XmlString );
-	}
-
 	FDatasmithSceneXmlWriterImpl::WriteIndent(Archive, 1);
 	XmlString = TEXT("<") + FString(DATASMITH_EXPORT) + TEXT(" ") + FString(DATASMITH_EXPORTDURATION) + FString::Printf(TEXT("=\"%d\"/>"), DatasmithScene->GetExportDuration()) + LINE_TERMINATOR;
 	FDatasmithSceneXmlWriterImpl::SerializeToArchive( Archive, XmlString );
diff --git a/Engine/Source/Runtime/Datasmith/DatasmithCore/Public/DatasmithDefinitions.h b/Engine/Source/Runtime/Datasmith/DatasmithCore/Public/DatasmithDefinitions.h
--- a/Engine/Source/Runtime/Datasmith/DatasmithCore/Public/DatasmithDefinitions.h
+++ b/Engine/Source/Runtime/Datasmith/DatasmithCore/Public/DatasmithDefinitions.h
@@ -627,6 +627,3 @@ static const TCHAR* KeyValuePropertyTypeStrings[] = { TEXT("String"), TEXT("Colo
 #define DATASMITH_ENVILLUMINATIONMAP			TEXT("Illuminate")
 
 #define DATASMITH_DECALMATERIALNAME				TEXT("DecalMaterial")
-
-//LODS
-#define DATASMITH_LODSCREENSIZE					TEXT("LodScreenSize")
diff --git a/Engine/Source/Runtime/Datasmith/DatasmithCore/Public/IDatasmithSceneElements.h b/Engine/Source/Runtime/Datasmith/DatasmithCore/Public/IDatasmithSceneElements.h
--- a/Engine/Source/Runtime/Datasmith/DatasmithCore/Public/IDatasmithSceneElements.h
+++ b/Engine/Source/Runtime/Datasmith/DatasmithCore/Public/IDatasmithSceneElements.h
@@ -1687,15 +1687,6 @@ public:
 	virtual TSharedPtr< IDatasmithPostProcessElement > GetPostProcess() = 0;
 	virtual const TSharedPtr< IDatasmithPostProcessElement >& GetPostProcess() const = 0;
 
-	/**
-	 * Adds a new LOD screen size setting. The first one added is used for the base LOD, the second one is for LOD1, etc.
-	 *
-	 * @param ScreenSize Ratio of the screen clamped between 0 and 1
-	 */
-	virtual void AddLODScreenSize(float ScreenSize) = 0;
-	virtual int32 GetLODScreenSizesCount() const = 0;
-	virtual float GetLODScreenSize(int32 InIndex) const = 0;
-
 	/**
 	 * Add a metadata to the scene
 	 * There should be only one metadata per Datasmith element (the element associated with the metadata)
