Datasmith: keep IES light export compatible with 5.0

diff --git a/Engine/Source/Runtime/Datasmith/DatasmithCore/Private/DatasmithSceneXmlWriter.cpp b/Engine/Source/Runtime/Datasmith/DatasmithCore/Private/DatasmithSceneXmlWriter.cpp
--- a/Engine/Source/Runtime/Datasmith/DatasmithCore/Private/DatasmithSceneXmlWriter.cpp
+++ b/Engine/Source/Runtime/Datasmith/DatasmithCore/Private/DatasmithSceneXmlWriter.cpp
@@ -889,9 +889,18 @@ void FDatasmithSceneXmlWriterImpl::WriteLightActorElement(const TSharedPtr< IDat
 
 	if ( LightActorElement->GetIesRotation() != FQuat::Identity )
 	{
-		WriteIndent(Archive, Indent + 1);
+		auto QuatToHexStringFloat = [](const FQuat& Value) -> FString
+		{
+			float Tmp[4] = {Value.X, Value.Y, Value.Z, Value.W};
+			FString Result = TEXT(" qhex=\"") + FString::FromHexBlob((uint8*)Tmp, sizeof(Tmp)) + TEXT("\"");
+			return Result;
+		};
 
-		XmlString = TEXT("<") + FString(DATASMITH_LIGHTIESROTATION) + TEXT(" ") + QuatToHexString( LightActorElement->GetIesRotation() ) + TEXT("/>") + LINE_TERMINATOR;
+		WriteIndent(Archive, Indent + 1);
+		XmlString = TEXT("<") DATASMITH_LIGHTIESROTATION TEXT(" ")
+			+ QuatToHexString( LightActorElement->GetIesRotation() )
+			+ QuatToHexStringFloat( LightActorElement->GetIesRotation() ) // required for 5.1 -> 5.0.0 flow, as 5.0.0 cannot parse a qhex64 attribute. Should be fixed for 5.0.1.
+			+ TEXT("/>") + LINE_TERMINATOR;
 		SerializeToArchive( Archive, XmlString );
 	}
 
