buildfix: DatasmithSceneXmlWriter.cpp:894:29: error: non-constant-expression cannot be narrowed from type 'double' to 'float' in initializer list

diff --git a/Engine/Source/Runtime/Datasmith/DatasmithCore/Private/DatasmithSceneXmlWriter.cpp b/Engine/Source/Runtime/Datasmith/DatasmithCore/Private/DatasmithSceneXmlWriter.cpp
--- a/Engine/Source/Runtime/Datasmith/DatasmithCore/Private/DatasmithSceneXmlWriter.cpp
+++ b/Engine/Source/Runtime/Datasmith/DatasmithCore/Private/DatasmithSceneXmlWriter.cpp
@@ -891,7 +891,7 @@ void FDatasmithSceneXmlWriterImpl::WriteLightActorElement(const TSharedPtr< IDat
 	{
 		auto QuatToHexStringFloat = [](const FQuat& Value) -> FString
 		{
-			float Tmp[4] = {Value.X, Value.Y, Value.Z, Value.W};
+			float Tmp[4] = {(float)Value.X, (float)Value.Y, (float)Value.Z, (float)Value.W};
 			FString Result = TEXT(" qhex=\"") + FString::FromHexBlob((uint8*)Tmp, sizeof(Tmp)) + TEXT("\"");
 			return Result;
 		};
