Datasmith: minor cleanup

diff --git a/Engine/Plugins/Enterprise/DatasmithImporter/Source/DatasmithImporter/Private/DatasmithImporter.cpp b/Engine/Plugins/Enterprise/DatasmithImporter/Source/DatasmithImporter/Private/DatasmithImporter.cpp
--- a/Engine/Plugins/Enterprise/DatasmithImporter/Source/DatasmithImporter/Private/DatasmithImporter.cpp
+++ b/Engine/Plugins/Enterprise/DatasmithImporter/Source/DatasmithImporter/Private/DatasmithImporter.cpp
@@ -1806,7 +1806,6 @@ void FDatasmithImporter::FinalizeImport(FDatasmithImportContext& ImportContext,
 			FDatasmithImporterImpl::CheckAssetPersistenceValidity(ExistingMaterialFunction, ImportContext);
 		}
 
-		TArray<UMaterial*> MaterialsToRefreshAfterVirtualTextureConversion;
 		for (const TPair< TSharedRef< IDatasmithBaseMaterialElement >, UMaterialInterface* >& ImportedMaterialPair : ImportContext.ImportedMaterials)
 		{
 			if (ImportContext.bUserCancelled)
@@ -1834,17 +1833,6 @@ void FDatasmithImporter::FinalizeImport(FDatasmithImportContext& ImportContext,
 			{
 				SourceMaterial->UpdateCachedExpressionData();
 
-				TArray<UObject*> ReferencedTextures;
-				ReferencedTextures = SourceMaterial->GetReferencedTextures();
-				for (UTexture2D* VirtualTexture : ImportContext.AssetsContext.VirtualTexturesToConvert)
-				{
-					if (ReferencedTextures.Contains(VirtualTexture))
-					{
-						MaterialsToRefreshAfterVirtualTextureConversion.Add(SourceMaterial);
-						break;
-					}
-				}
-
 				TArray<UMaterialFunctionInterface*> FinalizableFunctions;
 				for (const FMaterialFunctionInfo& MaterialFunctionInfo : SourceMaterial->GetCachedExpressionData().FunctionInfos)
 				{
diff --git a/Engine/Source/Developer/Datasmith/DatasmithExporter/Private/DatasmithMeshExporter.cpp b/Engine/Source/Developer/Datasmith/DatasmithExporter/Private/DatasmithMeshExporter.cpp
--- a/Engine/Source/Developer/Datasmith/DatasmithExporter/Private/DatasmithMeshExporter.cpp
+++ b/Engine/Source/Developer/Datasmith/DatasmithExporter/Private/DatasmithMeshExporter.cpp
@@ -17,7 +17,7 @@
 #include "Serialization/MemoryWriter.h"
 
 /**
- *	Number of datsmith meshes pending deletion for which the Garbage Collection will be triggered.
+ * Number of datasmith meshes pending deletion for which the Garbage Collection will be triggered.
  */
 constexpr int32 DSMeshGarbageCollectionThreshold{ 2000 };
 
@@ -37,8 +37,8 @@ struct FDatasmithMeshExporterOptions
 };
 
 /**
- *	Implementation class of the DatasmithMeshExporter
- *	We use a lockfree UDatasmithMesh pool to avoid creating new UObject when exporting and reduces our memory footprint.
+ * Implementation class of the DatasmithMeshExporter
+ * We use a lockfree UDatasmithMesh pool to avoid creating new UObject when exporting and reduces our memory footprint.
  */
 class FDatasmithMeshExporterLegacyImpl
 {
@@ -219,7 +219,6 @@ bool FDatasmithMeshExporterLegacyImpl::ExportMeshes( const FDatasmithMeshExporte
 }
 
 
-
 FDatasmithMeshExporter::FDatasmithMeshExporter()
 	: LegacyImpl(MakeUnique<FDatasmithMeshExporterLegacyImpl>())
 	, Impl(MakeUnique<FDatasmithMeshExporterImpl>())
diff --git a/Engine/Source/Developer/Datasmith/DatasmithExporter/Public/DatasmithMeshExporter.h b/Engine/Source/Developer/Datasmith/DatasmithExporter/Public/DatasmithMeshExporter.h
--- a/Engine/Source/Developer/Datasmith/DatasmithExporter/Public/DatasmithMeshExporter.h
+++ b/Engine/Source/Developer/Datasmith/DatasmithExporter/Public/DatasmithMeshExporter.h
@@ -29,7 +29,7 @@ public:
 	/**
 	 * Exports a FDatasmithMesh as a UObject and link it to the given IDatasmithMeshElementElement.
 	 *
-	 * @param MeshElement	The existing MeshElement for with we want to export a FDatasmithMesh, the name of the MeshElement will determine the name of the exported file.
+	 * @param MeshElement	The existing MeshElement for which we want to export a FDatasmithMesh, the name of the MeshElement will determine the name of the exported file.
 	 * @param Filepath		The path where the resulting file will be written
 	 * @param Mesh	        The mesh to export
 	 * @param CollisionMesh An optional collision mesh
diff --git a/Engine/Source/Runtime/Datasmith/DatasmithCore/Private/DatasmithMeshSerialization.cpp b/Engine/Source/Runtime/Datasmith/DatasmithCore/Private/DatasmithMeshSerialization.cpp
--- a/Engine/Source/Runtime/Datasmith/DatasmithCore/Private/DatasmithMeshSerialization.cpp
+++ b/Engine/Source/Runtime/Datasmith/DatasmithCore/Private/DatasmithMeshSerialization.cpp
@@ -135,7 +135,7 @@ bool DecompressInline(TArray<uint8>& CompressedData)
 	return false;
 }
 
-} // ns Impl
+} // ns DatasmithMeshSerializationImpl
 
 
 FMD5Hash FDatasmithPackedMeshes::Serialize(FArchive& Ar, bool bCompressed)
@@ -155,7 +155,7 @@ FMD5Hash FDatasmithPackedMeshes::Serialize(FArchive& Ar, bool bCompressed)
 
 	enum EBufferType{ RawMeshDescription, CompressedMeshDescription };
 	uint8 BufferType = bCompressed ? CompressedMeshDescription : RawMeshDescription;
-	Ar << BufferType; // (MeshDesc, Zipped Mesh desc;...)
+	Ar << BufferType;
 
 	FMD5Hash OutHash;
 	if (Ar.IsLoading())
diff --git a/Engine/Source/Runtime/Datasmith/DatasmithCore/Private/DatasmithSceneXmlReader.cpp b/Engine/Source/Runtime/Datasmith/DatasmithCore/Private/DatasmithSceneXmlReader.cpp
--- a/Engine/Source/Runtime/Datasmith/DatasmithCore/Private/DatasmithSceneXmlReader.cpp
+++ b/Engine/Source/Runtime/Datasmith/DatasmithCore/Private/DatasmithSceneXmlReader.cpp
@@ -1,4 +1,5 @@
 // Copyright Epic Games, Inc. All Rights Reserved.
+
 #include "DatasmithSceneXmlReader.h"
 
 #include "DatasmithCore.h"
@@ -402,7 +403,7 @@ void FDatasmithSceneXmlReader::ParseMesh(FXmlNode* InNode, TSharedPtr<IDatasmith
 	{
 		if (MeshNodes[j]->GetTag() == TEXT("file"))
 		{
-			OutElement->SetFile(*MeshNodes[j]->GetAttribute(TEXT("path"))) ;
+			OutElement->SetFile(*UnsanitizeXMLText(MeshNodes[j]->GetAttribute(TEXT("path"))));
 		}
 		else if (MeshNodes[j]->GetTag() == TEXT("Size"))
 		{
@@ -1178,7 +1179,7 @@ bool FDatasmithSceneXmlReader::ParseXmlFile(TSharedRef< IDatasmithScene >& OutSc
 			OutScene->SetUserID(*Nodes[i]->GetAttribute(DATASMITH_USERID));
 			OutScene->SetUserOS(*Nodes[i]->GetAttribute(DATASMITH_USEROS));
 		}
-		//READ STATIC MESHES
+		// STATIC MESHES
 		else if (Nodes[i]->GetTag() == DATASMITH_STATICMESHNAME)
 		{
 			FString ElementName = Nodes[i]->GetAttribute(TEXT("name"));
@@ -1190,7 +1191,7 @@ bool FDatasmithSceneXmlReader::ParseXmlFile(TSharedRef< IDatasmithScene >& OutSc
 
 			Objects.Add( Element->GetName(), Element );
 		}
-		//READ LEVEL SEQUENCES
+		// LEVEL SEQUENCES
 		else if (Nodes[i]->GetTag() == DATASMITH_LEVELSEQUENCENAME)
 		{
 			FString ElementName = Nodes[i]->GetAttribute(TEXT("name"));
@@ -1200,7 +1201,7 @@ bool FDatasmithSceneXmlReader::ParseXmlFile(TSharedRef< IDatasmithScene >& OutSc
 
 			OutScene->AddLevelSequence(Element);
 		}
-		//READ LEVEL VARIANT SETS
+		// LEVEL VARIANT SETS
 		else if (Nodes[i]->GetTag() == DATASMITH_LEVELVARIANTSETSNAME)
 		{
 			FString ElementName = Nodes[i]->GetAttribute(TEXT("name"));
@@ -1210,7 +1211,7 @@ bool FDatasmithSceneXmlReader::ParseXmlFile(TSharedRef< IDatasmithScene >& OutSc
 
 			OutScene->AddLevelVariantSets(Element);
 		}
-		//READ TEXTURES
+		// TEXTURES
 		else if (Nodes[i]->GetTag() == DATASMITH_TEXTURENAME)
 		{
 			TSharedPtr< IDatasmithTextureElement > Element = FDatasmithSceneFactory::CreateTexture(*Nodes[i]->GetAttribute(TEXT("name")));
@@ -1236,7 +1237,7 @@ bool FDatasmithSceneXmlReader::ParseXmlFile(TSharedRef< IDatasmithScene >& OutSc
 				Objects.Add(Element->GetName(), Element);
 			}
 		}
-		//READ ENVIRONMENTS
+		// ENVIRONMENTS
 		else if (Nodes[i]->GetTag() == DATASMITH_ENVIRONMENTNAME)
 		{
 			TSharedPtr< IDatasmithEnvironmentElement > Element = FDatasmithSceneFactory::CreateEnvironment(*Nodes[i]->GetAttribute(TEXT("name")));
@@ -1267,19 +1268,19 @@ bool FDatasmithSceneXmlReader::ParseXmlFile(TSharedRef< IDatasmithScene >& OutSc
 				Objects.Add( Element->GetName(), Element );
 			}
 		}
-		//READ SKY
+		// SKY
 		else if (Nodes[i]->GetTag() == DATASMITH_PHYSICALSKYNAME)
 		{
 			OutScene->SetUsePhysicalSky( ValueFromString<bool>(Nodes[i]->GetAttribute(TEXT("enabled"))));
 		}
-		//READ POSTPROCESS
+		// POSTPROCESS
 		else if (Nodes[i]->GetTag() == DATASMITH_POSTPRODUCTIONNAME)
 		{
 			TSharedPtr< IDatasmithPostProcessElement > PostProcess = FDatasmithSceneFactory::CreatePostProcess();
 			ParsePostProcess(Nodes[i], PostProcess);
 			OutScene->SetPostProcess(PostProcess);
 		}
-		//READ MATERIALS
+		// MATERIALS
 		else if (Nodes[i]->GetTag() == DATASMITH_MATERIALNAME)
 		{
 			TSharedPtr< IDatasmithMaterialElement > Material = FDatasmithSceneFactory::CreateMaterial(*Nodes[i]->GetAttribute(TEXT("name")));
@@ -1289,7 +1290,7 @@ bool FDatasmithSceneXmlReader::ParseXmlFile(TSharedRef< IDatasmithScene >& OutSc
 
 			Objects.Add( Material->GetName(), Material );
 		}
-		//READ MASTER MATERIALS
+		// MASTER MATERIALS
 		else if (Nodes[i]->GetTag() == DATASMITH_MASTERMATERIALNAME)
 		{
 			TSharedPtr< IDatasmithMasterMaterialElement > MasterMaterial = FDatasmithSceneFactory::CreateMasterMaterial(*Nodes[i]->GetAttribute(TEXT("name")));
@@ -1299,7 +1300,7 @@ bool FDatasmithSceneXmlReader::ParseXmlFile(TSharedRef< IDatasmithScene >& OutSc
 
 			Objects.Add( MasterMaterial->GetName(), MasterMaterial );
 		}
-		//READ DECAL MATERIALS
+		// DECAL MATERIALS
 		else if (Nodes[i]->GetTag() == DATASMITH_DECALMATERIALNAME)
 		{
 			TSharedPtr< IDatasmithDecalMaterialElement > DecalMaterial = FDatasmithSceneFactory::CreateDecalMaterial(*Nodes[i]->GetAttribute(TEXT("name")));
@@ -1309,7 +1310,7 @@ bool FDatasmithSceneXmlReader::ParseXmlFile(TSharedRef< IDatasmithScene >& OutSc
 
 			Objects.Add( DecalMaterial->GetName(), DecalMaterial );
 		}
-		//READ UEPBR MATERIALS
+		// UEPBR MATERIALS
 		else if (Nodes[i]->GetTag() == DATASMITH_UEPBRMATERIALNAME)
 		{
 			TSharedPtr< IDatasmithUEPbrMaterialElement > Material = FDatasmithSceneFactory::CreateUEPbrMaterial(*Nodes[i]->GetAttribute(TEXT("name")));
@@ -1319,7 +1320,7 @@ bool FDatasmithSceneXmlReader::ParseXmlFile(TSharedRef< IDatasmithScene >& OutSc
 
 			Objects.Add( Material->GetName(), Material );
 		}
-		//READ METADATA
+		// METADATA
 		else if (Nodes[i]->GetTag() == DATASMITH_METADATANAME)
 		{
 			TSharedPtr< IDatasmithMetaDataElement > MetaData = FDatasmithSceneFactory::CreateMetaData(*Nodes[i]->GetAttribute(TEXT("name")));
@@ -1331,7 +1332,7 @@ bool FDatasmithSceneXmlReader::ParseXmlFile(TSharedRef< IDatasmithScene >& OutSc
 		{
 			OutScene->SetExportDuration(ValueFromString<int32>(Nodes[i]->GetAttribute(DATASMITH_EXPORTDURATION)));
 		}
-		//READ ACTORS
+		// ACTORS
 		else
 		{
 			for ( const TCHAR* ActorTag : ActorTagsView )
@@ -1361,8 +1362,8 @@ bool FDatasmithSceneXmlReader::ParseXmlFile(TSharedRef< IDatasmithScene >& OutSc
 #ifdef USE_LOCALE
 	if (Locale)
 	{
-	::uselocale(PreviousLocale);
-	::freelocale(Locale);
+		::uselocale(PreviousLocale);
+		::freelocale(Locale);
 	}
 #endif
 
