Datasmith: Mesh serialization cleanup

diff --git a/Engine/Plugins/Enterprise/DatasmithImporter/Source/DatasmithNativeTranslator/Private/DatasmithNativeTranslator.cpp b/Engine/Plugins/Enterprise/DatasmithImporter/Source/DatasmithNativeTranslator/Private/DatasmithNativeTranslator.cpp
--- a/Engine/Plugins/Enterprise/DatasmithImporter/Source/DatasmithNativeTranslator/Private/DatasmithNativeTranslator.cpp
+++ b/Engine/Plugins/Enterprise/DatasmithImporter/Source/DatasmithNativeTranslator/Private/DatasmithNativeTranslator.cpp
@@ -117,7 +117,8 @@ bool FDatasmithNativeTranslator::LoadStaticMesh(const TSharedRef<IDatasmithMeshE
 
 	FString FilePath = MeshElement->GetFile();
 
-	for (FDatasmithMeshModels& DatasmithMesh : GetDatasmithMeshFromMeshPath( FilePath ))
+	FDatasmithPackedMeshes Pack = GetDatasmithMeshFromFile(FilePath);
+	for (FDatasmithMeshModels& DatasmithMesh : Pack.Meshes)
 	{
 		if (DatasmithMesh.bIsCollisionMesh)
 		{
diff --git a/Engine/Source/Developer/Datasmith/DatasmithExporter/Private/DatasmithMeshExporter.cpp b/Engine/Source/Developer/Datasmith/DatasmithExporter/Private/DatasmithMeshExporter.cpp
--- a/Engine/Source/Developer/Datasmith/DatasmithExporter/Private/DatasmithMeshExporter.cpp
+++ b/Engine/Source/Developer/Datasmith/DatasmithExporter/Private/DatasmithMeshExporter.cpp
@@ -264,10 +264,10 @@ bool FDatasmithMeshExporterImpl::WriteMeshFile( const FDatasmithMeshExporterOpti
 		return Models;
 	};
 
-	Pack.MeshesToExport.Add(PackMeshModels(ExporterOptions.Mesh, false));
+	Pack.Meshes.Add(PackMeshModels(ExporterOptions.Mesh, false));
 	if (ExporterOptions.CollisionMesh)
 	{
-		Pack.MeshesToExport.Add(PackMeshModels(*ExporterOptions.CollisionMesh, true));
+		Pack.Meshes.Add(PackMeshModels(*ExporterOptions.CollisionMesh, true));
 	}
 
 	TUniquePtr<FArchive> Archive( IFileManager::Get().CreateFileWriter(*ExporterOptions.MeshFullPath) );
diff --git a/Engine/Source/Runtime/Datasmith/DatasmithCore/Private/DatasmithMeshSerialization.cpp b/Engine/Source/Runtime/Datasmith/DatasmithCore/Private/DatasmithMeshSerialization.cpp
--- a/Engine/Source/Runtime/Datasmith/DatasmithCore/Private/DatasmithMeshSerialization.cpp
+++ b/Engine/Source/Runtime/Datasmith/DatasmithCore/Private/DatasmithMeshSerialization.cpp
@@ -173,13 +173,13 @@ FMD5Hash FDatasmithPackedMeshes::Serialize(FArchive& Ar, bool bCompressed)
 
 		FMemoryReader Buffer(Bytes, true);
 		Buffer.SetCustomVersions(CustomVersions);
-		Buffer << MeshesToExport;
+		Buffer << Meshes;
 	}
 	else
 	{
 		TArray<uint8> Bytes;
 		FMemoryWriter Buffer(Bytes, true);
-		Buffer << MeshesToExport;
+		Buffer << Meshes;
 
 		// MeshDescriptions uses custom versionning,
 		const_cast<FCustomVersionContainer&>(Buffer.GetCustomVersions()).Serialize(Ar);
@@ -285,9 +285,9 @@ TArray<FDatasmithMeshModels> GetDatasmithMeshFromMeshPath_Legacy(FArchive* Archi
 	return Result;
 }
 
-TArray<FDatasmithMeshModels> GetDatasmithMeshFromMeshPath(const FString& MeshPath)
+FDatasmithPackedMeshes GetDatasmithMeshFromFile(const FString& MeshPath)
 {
-	TArray< FDatasmithMeshModels > Result;
+	FDatasmithPackedMeshes Result;
 
 	TUniquePtr<FArchive> Archive( IFileManager::Get().CreateFileReader(*MeshPath) );
 	if ( !Archive.IsValid() )
@@ -301,25 +301,15 @@ TArray<FDatasmithMeshModels> GetDatasmithMeshFromMeshPath(const FString& MeshPat
 
 	if (LeagacyNumMeshesCount > 0)
 	{
-		return GetDatasmithMeshFromMeshPath_Legacy(Archive.Get(), LeagacyNumMeshesCount);
+		Result.Meshes = GetDatasmithMeshFromMeshPath_Legacy(Archive.Get(), LeagacyNumMeshesCount);
 	}
-
-	if (LeagacyNumMeshesCount == 0)
+	else
 	{
-		FDatasmithPackedMeshes Pack;
-		Pack.Serialize(*Archive);
+		Result.Serialize(*Archive);
 
-		if (!Archive->IsError())
-		{
-			for (FDatasmithMeshModels& Mesh : Pack.MeshesToExport)
-			{
-				FDatasmithMeshModels& MeshInternal = Result.AddDefaulted_GetRef();
-				MeshInternal.bIsCollisionMesh = Mesh.bIsCollisionMesh;
-				MeshInternal.SourceModels = MoveTemp(Mesh.SourceModels);
-			}
-		}
-		else
+		if (Archive->IsError())
 		{
+			Result = FDatasmithPackedMeshes{};
 			UE_LOG(LogDatasmith, Warning, TEXT("Failed to read meshes from %s"), *MeshPath);
 		}
 	}
diff --git a/Engine/Source/Runtime/Datasmith/DatasmithCore/Public/DatasmithMeshSerialization.h b/Engine/Source/Runtime/Datasmith/DatasmithCore/Public/DatasmithMeshSerialization.h
--- a/Engine/Source/Runtime/Datasmith/DatasmithCore/Public/DatasmithMeshSerialization.h
+++ b/Engine/Source/Runtime/Datasmith/DatasmithCore/Public/DatasmithMeshSerialization.h
@@ -10,7 +10,7 @@
 struct FDatasmithMeshModels
 {
 	FString MeshName;
-	bool bIsCollisionMesh;
+	bool bIsCollisionMesh = false;
 	TArray<FMeshDescription> SourceModels;
 
 	DATASMITHCORE_API friend void operator << (FArchive& Ar, FDatasmithMeshModels& Models);
@@ -18,9 +18,9 @@ struct FDatasmithMeshModels
 
 struct DATASMITHCORE_API FDatasmithPackedMeshes
 {
-	TArray<FDatasmithMeshModels> MeshesToExport;
+	TArray<FDatasmithMeshModels> Meshes;
 
 	FMD5Hash Serialize(FArchive& Ar, bool bSaveCompressed=true);
 };
 
-DATASMITHCORE_API TArray<FDatasmithMeshModels> GetDatasmithMeshFromMeshPath(const FString& MeshPath);
+DATASMITHCORE_API FDatasmithPackedMeshes GetDatasmithMeshFromFile(const FString& MeshPath);
