Datasmith: Add a CloTranslator that can ingest "additional data" json files from MarvelousDesigner
Add a Cloth import extension that support the ChaosClothAsset and can import them

diff --git a/Engine/Source/Runtime/Datasmith/DatasmithCore/Public/DatasmithCloth.h b/Engine/Source/Runtime/Datasmith/DatasmithCore/Public/DatasmithCloth.h
--- a/Engine/Source/Runtime/Datasmith/DatasmithCore/Public/DatasmithCloth.h
+++ b/Engine/Source/Runtime/Datasmith/DatasmithCore/Public/DatasmithCloth.h
@@ -1,31 +1,34 @@
-// Copyright Epic Games, Inc. All Rights Reserved.
-#pragma once
-
-#include "CoreMinimal.h"
-// #include "Templates/PimplPtr.h"
-// #include "Containers/Array.h"
-
-class FDatasmithMesh;
-
-class DATASMITHCORE_API FDatasmithClothPattern
-{
-public:
-	TArray<FVector2f> SimPosition;
-	TArray<FVector3f> SimRestPosition;
-	TArray<uint32> SimTriangleIndices;
-};
-
-/*
- * Structure of a Cloth
- * - One Thin Mesh:
- * - base for physics simulation
- *
- * - One Thick Mesh
- * - aka. render mesh
- */
-class DATASMITHCORE_API FDatasmithCloth
-{
-public:
-	TArray<FDatasmithClothPattern> Patterns;
-};
-
+// Copyright Epic Games, Inc. All Rights Reserved.
+#pragma once
+
+#include "CoreMinimal.h"
+// #include "Templates/PimplPtr.h"
+// #include "Containers/Array.h"
+
+class FDatasmithMesh;
+
+class DATASMITHCORE_API FDatasmithClothPattern
+{
+public:
+	TArray<FVector2f> SimPosition;
+	TArray<FVector3f> SimRestPosition;
+	TArray<uint32> SimTriangleIndices;
+
+public:
+	bool IsValid() { return SimRestPosition.Num() == SimPosition.Num() && SimTriangleIndices.Num() % 3 == 0; }
+};
+
+/*
+ * Structure of a Cloth
+ * - One Thin Mesh:
+ * - base for physics simulation
+ *
+ * - One Thick Mesh
+ * - aka. render mesh
+ */
+class DATASMITHCORE_API FDatasmithCloth
+{
+public:
+	TArray<FDatasmithClothPattern> Patterns;
+};
+
