Buildfix: rollback 20136088 (unknown type name 'GLuint')

diff --git a/Engine/Source/Runtime/Datasmith/DatasmithCore/Public/DatasmithCloth.h b/Engine/Source/Runtime/Datasmith/DatasmithCore/Public/DatasmithCloth.h
--- a/Engine/Source/Runtime/Datasmith/DatasmithCore/Public/DatasmithCloth.h
+++ b/Engine/Source/Runtime/Datasmith/DatasmithCore/Public/DatasmithCloth.h
@@ -13,9 +13,6 @@ public:
 	TArray<FVector2f> SimPosition;
 	TArray<FVector3f> SimRestPosition;
 	TArray<uint32> SimTriangleIndices;
-
-public:
-	bool IsValid() { return SimRestPosition.Num() == SimPosition.Num() && SimTriangleIndices.Num() % 3 == 0; }
 };
 
 /*
