Datasmith: Cloth importer extension

diff --git a/Engine/Source/Runtime/Datasmith/DatasmithCore/Public/DatasmithCloth.h b/Engine/Source/Runtime/Datasmith/DatasmithCore/Public/DatasmithCloth.h
--- a/Engine/Source/Runtime/Datasmith/DatasmithCore/Public/DatasmithCloth.h
+++ b/Engine/Source/Runtime/Datasmith/DatasmithCore/Public/DatasmithCloth.h
@@ -2,8 +2,6 @@
 #pragma once
 
 #include "CoreMinimal.h"
-// #include "Templates/PimplPtr.h"
-// #include "Containers/Array.h"
 
 class FDatasmithMesh;
 
@@ -13,6 +11,9 @@ public:
 	TArray<FVector2f> SimPosition;
 	TArray<FVector3f> SimRestPosition;
 	TArray<uint32> SimTriangleIndices;
+
+public:
+	bool IsValid() { return SimRestPosition.Num() == SimPosition.Num() && SimTriangleIndices.Num() % 3 == 0; }
 };
 
 /*
