CADKernel Improvements (manual merge from 5.0.3)

fix crash in TPolylineApproximator::UpdateSubPolylineBBox (PolylineTools.h)
fix nullptr in FCoreTechBridge::AddLoop
improve void FTopologicalFace::UpdateBBox

improve FTopologicalLoop::RemoveDegeneratedEdges() by reducing the joining angle of an degenerated edge and its neighbour to allow the merge of the degenerated edge
improve FTopologicalLoop::EnsureLogicalClosing
- reducing the joining angle of a loop's gap and its neighbouring to allow to extend the neighbour to fill the gap
- make a degenerated edge to fill lage Gap in 2D (in degenerated area)

fix FParametricMesher::Mesh(FTopologicalEdge& InEdge, const FTopologicalFace& Face): Set meshed degenerated edge as meshed

fix FGrid::GetMeshOfLoop when a degenerated loop is deleted

improve FTechSoftBridge::TraverseShell => Avoid edge link between different shells

fix crash in FTechSoftFileParserCADKernelTessellator::AdaptBRepModel()
fix nullptr in FTechSoftFileParserCADKernelTessellator::GenerateBodyMesh

fix crash in BowyerWatsonTriangulation.h

remove "ensure" in:
- FTopologicalLoop::Orient()
- FTopologicalVertex::RemoveConnectedEdge

diff --git a/Engine/Plugins/Enterprise/DatasmithCADImporter/Source/CADInterfaces/Private/TechSoftBridge.cpp b/Engine/Plugins/Enterprise/DatasmithCADImporter/Source/CADInterfaces/Private/TechSoftBridge.cpp
--- a/Engine/Plugins/Enterprise/DatasmithCADImporter/Source/CADInterfaces/Private/TechSoftBridge.cpp
+++ b/Engine/Plugins/Enterprise/DatasmithCADImporter/Source/CADInterfaces/Private/TechSoftBridge.cpp
@@ -58,8 +58,6 @@ namespace CADLibrary
 namespace TechSoftUtils
 {
 
-const FUVReparameterization FUVReparameterization::Identity(1., 0, 1., 0);
-
 template<typename... InArgTypes>
 A3DStatus GetCurveAsNurbs(const A3DCrvBase* A3DCurve, A3DCrvNurbsData* DataPtr, InArgTypes&&... Args)
 {
@@ -166,16 +164,27 @@ void FillPointArray(const int32 UCount, const int32 VCount, const A3DVector3dDat
 	}
 };
 
-CADKernel::FSurfacicBoundary GetSurfacicBoundary(A3DDomainData& Domain, double UnitScale = 1.0, bool bSwapUV = false)
+CADKernel::FSurfacicBoundary GetSurfacicBoundary(A3DDomainData& Domain, const TechSoftUtils::FUVReparameterization& UVReparameterization)
 {
-	CADKernel::EIso UIndex = bSwapUV ? CADKernel::EIso::IsoV : CADKernel::EIso::IsoU;
-	CADKernel::EIso VIndex = bSwapUV ? CADKernel::EIso::IsoU : CADKernel::EIso::IsoV;
+
+	CADKernel::FPoint2D Min(Domain.m_sMin.m_dX, Domain.m_sMin.m_dY);
+	CADKernel::FPoint2D Max(Domain.m_sMax.m_dX, Domain.m_sMax.m_dY);
+
+	if(UVReparameterization.GetNeedApply())
+	{
+		UVReparameterization.Apply(Min);
+		UVReparameterization.Apply(Max);
+	}
+
+	CADKernel::EIso UIndex = UVReparameterization.GetSwapUV() ? CADKernel::EIso::IsoV : CADKernel::EIso::IsoU;
+	CADKernel::EIso VIndex = UVReparameterization.GetSwapUV() ? CADKernel::EIso::IsoU : CADKernel::EIso::IsoV;
 
 	CADKernel::FSurfacicBoundary Boundary;
-	Boundary[UIndex].Min = Domain.m_sMin.m_dX * UnitScale;
-	Boundary[UIndex].Max = Domain.m_sMax.m_dX * UnitScale;
-	Boundary[VIndex].Min = Domain.m_sMin.m_dY * UnitScale;
-	Boundary[VIndex].Max = Domain.m_sMax.m_dY * UnitScale;
+	Boundary[UIndex].Min = Min.U;
+	Boundary[VIndex].Min = Min.V;
+	Boundary[UIndex].Max = Max.U;
+	Boundary[VIndex].Max = Max.V;
+
 	return Boundary;
 }
 
@@ -332,11 +341,13 @@ void FTechSoftBridge::TraverseShell(const A3DTopoShell* A3DShell, TSharedRef<CAD
 	AddMetadata(MetaData, *Shell);
 
 	TUniqueTSObj<A3DTopoShellData> ShellData(A3DShell);
+
 	if (ShellData.IsValid())
 	{
+		A3DEdgeToEdge.Empty();
 		for (A3DUns32 Index = 0; Index < ShellData->m_uiFaceSize; ++Index)
 		{
-			AddFace(ShellData->m_ppFaces[Index], ShellData->m_pucOrientationWithShell[Index] == 1 ? CADKernel::Front : CADKernel::Back, Shell);
+			AddFace(ShellData->m_ppFaces[Index], ShellData->m_pucOrientationWithShell[Index] == 1 ? CADKernel::Front : CADKernel::Back, Shell, Index);
 		}
 	}
 }
@@ -358,7 +369,6 @@ TSharedPtr<CADKernel::FCurve> FTechSoftBridge::AddCurve(const A3DCrvBase* A3DCur
 			break;
 		case kA3DTypeCrvLine:
 			Curve = AddCurveLine(A3DCurve, UVReparameterization);
-			Report.CurveNurbsCount++;
 			break;
 		case kA3DTypeCrvCircle:
 			Curve = AddCurveCircle(A3DCurve, UVReparameterization);
@@ -371,11 +381,9 @@ TSharedPtr<CADKernel::FCurve> FTechSoftBridge::AddCurve(const A3DCrvBase* A3DCur
 			break;
 		case kA3DTypeCrvHyperbola:
 			Curve = AddCurveHyperbola(A3DCurve, UVReparameterization);
-			Report.CurveNurbsCount++;
 			break;
 		case kA3DTypeCrvHelix:
 			Curve = AddCurveHelix(A3DCurve, UVReparameterization);
-			Report.CurveNurbsCount++;
 			break;
 		case kA3DTypeCrvPolyLine:
 			Curve = AddCurvePolyLine(A3DCurve, UVReparameterization);
@@ -563,12 +571,9 @@ TSharedPtr<CADKernel::FCurve> AddCurveNurbsFromData(A3DCrvNurbsData& A3DNurbs, c
 	Nurbs.Degree = A3DNurbs.m_uiDegree;
 
 	TechSoftUtils::FillPointArray(A3DNurbs.m_uiCtrlSize, A3DNurbs.m_pCtrlPts, Nurbs.Poles);
-	if (UVReparameterization.NeedApply())
+	if(Nurbs.Dimension == 2)
 	{
-		for (CADKernel::FPoint& Point : Nurbs.Poles)
-		{
-			UVReparameterization.Apply(Point);
-		}
+		UVReparameterization.Process(Nurbs.Poles);
 	}
 
 	TechSoftUtils::FillDoubleArray(A3DNurbs.m_uiKnotSize, A3DNurbs.m_pdKnots, Nurbs.NodalVector);
@@ -623,7 +628,10 @@ TSharedPtr<CADKernel::FTopologicalEdge> FTechSoftBridge::AddEdge(const A3DTopoCo
 		return TSharedPtr<CADKernel::FTopologicalEdge>();
 	}
 
-	ensureCADKernel(CoEdgeData->m_pUVCurve);
+	if (CoEdgeData->m_pUVCurve == nullptr)
+	{
+		return TSharedPtr<CADKernel::FTopologicalEdge>();
+	}
 
 	TSharedPtr<CADKernel::FCurve> Curve = AddCurve(CoEdgeData->m_pUVCurve, UVReparameterization);
 	if (!Curve.IsValid())
@@ -633,9 +641,7 @@ TSharedPtr<CADKernel::FTopologicalEdge> FTechSoftBridge::AddEdge(const A3DTopoCo
 
 	TSharedRef<CADKernel::FRestrictionCurve> RestrictionCurve = CADKernel::FEntity::MakeShared<CADKernel::FRestrictionCurve>(Surface, Curve.ToSharedRef());
 
-	TSharedPtr<CADKernel::FTopologicalEdge> Edge;
-	Edge = CADKernel::FTopologicalEdge::Make(RestrictionCurve);
-
+	TSharedPtr<CADKernel::FTopologicalEdge> Edge = CADKernel::FTopologicalEdge::Make(RestrictionCurve);
 	if (!Edge.IsValid())
 	{
 		return TSharedPtr<CADKernel::FTopologicalEdge>();
@@ -710,7 +716,7 @@ TSharedPtr<CADKernel::FTopologicalLoop> FTechSoftBridge::AddLoop(const A3DTopoLo
 	return CADKernel::FTopologicalLoop::Make(Edges, Directions, bIsExternalLoop, GeometricTolerance);
 }
 
-void FTechSoftBridge::AddFace(const A3DTopoFace* A3DFace, CADKernel::EOrientation Orientation, TSharedRef<CADKernel::FShell>& Shell)
+void FTechSoftBridge::AddFace(const A3DTopoFace* A3DFace, CADKernel::EOrientation Orientation, TSharedRef<CADKernel::FShell>& Shell, uint32 ShellIndex)
 {
 	Report.FaceCount++;
 
@@ -725,7 +731,7 @@ void FTechSoftBridge::AddFace(const A3DTopoFace* A3DFace, CADKernel::EOrientatio
 	}
 
 	const A3DSurfBase* A3DSurface = TopoFaceData->m_pSurface;
-	TechSoftUtils::FUVReparameterization UVReparameterization = TechSoftUtils::FUVReparameterization::Identity;
+	TechSoftUtils::FUVReparameterization UVReparameterization;
 	TSharedPtr<CADKernel::FSurface> SurfacePtr = AddSurface(A3DSurface, UVReparameterization);
 	if (!SurfacePtr.IsValid())
 	{
@@ -733,13 +739,19 @@ void FTechSoftBridge::AddFace(const A3DTopoFace* A3DFace, CADKernel::EOrientatio
 		Report.FailedFaceCount++;
 		return;
 	}
-	TSharedRef<CADKernel::FSurface> Surface = SurfacePtr.ToSharedRef();
 
+	if (UVReparameterization.GetNeedSwapOrientation())
+	{
+		SwapOrientation(Orientation);
+	}
+
+	TSharedRef<CADKernel::FSurface> Surface = SurfacePtr.ToSharedRef();
 	TSharedRef<CADKernel::FTopologicalFace> Face = CADKernel::FEntity::MakeShared<CADKernel::FTopologicalFace>(Surface);
 
 	if (TopoFaceData->m_bHasTrimDomain)
 	{
-		CADKernel::FSurfacicBoundary SurfaceBoundary = TechSoftUtils::GetSurfacicBoundary(TopoFaceData->m_sSurfaceDomain, BodyScale);
+		CADKernel::FSurfacicBoundary SurfaceBoundary = TechSoftUtils::GetSurfacicBoundary(TopoFaceData->m_sSurfaceDomain, UVReparameterization);
+
 		Surface->TrimBoundaryTo(SurfaceBoundary);
 	}
 
@@ -754,7 +766,7 @@ void FTechSoftBridge::AddFace(const A3DTopoFace* A3DFace, CADKernel::EOrientatio
 		for (A3DUns32 Index = 0; Index < TopoFaceData->m_uiLoopSize; ++Index)
 		{
 			const bool bIsExternalLoop = (Index == TopoFaceData->m_uiOuterLoopIndex);
-			TSharedPtr<CADKernel::FTopologicalLoop> Loop = FTechSoftBridge::AddLoop(TopoFaceData->m_ppLoops[Index], Surface, UVReparameterization, bIsExternalLoop);
+			TSharedPtr<CADKernel::FTopologicalLoop> Loop = AddLoop(TopoFaceData->m_ppLoops[Index], Surface, UVReparameterization, bIsExternalLoop);
 			if (!Loop.IsValid())
 			{
 				continue;
@@ -787,6 +799,8 @@ void FTechSoftBridge::AddFace(const A3DTopoFace* A3DFace, CADKernel::EOrientatio
 		if (Loops.Num() == 0)
 		{
 			Report.FailedFaceCount++;
+			Face->SetAsDegenerated();
+			Face->SetDeleted();
 			return;
 		}
 		else
@@ -796,6 +810,7 @@ void FTechSoftBridge::AddFace(const A3DTopoFace* A3DFace, CADKernel::EOrientatio
 	}
 
 	AddMetadata(MetaData, *Face);
+	Face->SetHostId(ShellIndex);
 	Shell->Add(Face, Orientation);
 }
 
@@ -827,11 +842,9 @@ TSharedPtr<CADKernel::FSurface> FTechSoftBridge::AddSurface(const A3DSurfBase* A
 			return AddNurbsSurface(A3DSurface, OutUVReparameterization);
 
 		case kA3DTypeSurfCone:
-			OutUVReparameterization.SetScale(1, BodyScale);
 			return AddConeSurface(A3DSurface, OutUVReparameterization);
 
 		case kA3DTypeSurfCylinder:
-			OutUVReparameterization.SetScale(1, BodyScale);
 			return AddCylinderSurface(A3DSurface, OutUVReparameterization);
 
 		case kA3DTypeSurfCylindrical:
@@ -844,7 +857,6 @@ TSharedPtr<CADKernel::FSurface> FTechSoftBridge::AddSurface(const A3DSurfBase* A
 			return AddPipeSurface(A3DSurface, OutUVReparameterization);
 
 		case kA3DTypeSurfPlane:
-			OutUVReparameterization.SetScale(BodyScale, BodyScale);
 			return AddPlaneSurface(A3DSurface, OutUVReparameterization);
 
 		case kA3DTypeSurfRuled:
@@ -890,8 +902,14 @@ TSharedPtr<CADKernel::FSurface> FTechSoftBridge::AddConeSurface(const A3DSurfBas
 	}
 
 	OutUVReparameterization.AddUVTransform(A3DConeData->m_sParam);
+	OutUVReparameterization.ScaleUVTransform(1, BodyScale);
+	if (A3DConeData->m_dSemiAngle < 0)
+	{
+		OutUVReparameterization.SetNeedSwapOrientation();
+	}
+
 	CADKernel::FMatrixH CoordinateSystem = TechSoftUtils::CreateCoordinateSystem(A3DConeData->m_sTrsf, BodyScale);
-	CADKernel::FSurfacicBoundary Boundary = TechSoftUtils::GetSurfacicBoundary(A3DConeData->m_sParam.m_sUVDomain, BodyScale);
+	CADKernel::FSurfacicBoundary Boundary = TechSoftUtils::GetSurfacicBoundary(A3DConeData->m_sParam.m_sUVDomain, OutUVReparameterization);
 	TSharedPtr<CADKernel::FSurface> Cone = CADKernel::FEntity::MakeShared<CADKernel::FConeSurface>(GeometricTolerance, CoordinateSystem, A3DConeData->m_dRadius * BodyScale, A3DConeData->m_dSemiAngle, Boundary);
 
 	return Cone;
@@ -908,8 +926,10 @@ TSharedPtr<CADKernel::FSurface> FTechSoftBridge::AddCylinderSurface(const A3DSur
 	}
 
 	OutUVReparameterization.AddUVTransform(A3DCylinderData->m_sParam);
+	OutUVReparameterization.ScaleUVTransform(1, BodyScale);
+
 	CADKernel::FMatrixH CoordinateSystem = TechSoftUtils::CreateCoordinateSystem(A3DCylinderData->m_sTrsf, BodyScale);
-	CADKernel::FSurfacicBoundary Boundary = TechSoftUtils::GetSurfacicBoundary(A3DCylinderData->m_sParam.m_sUVDomain, BodyScale);
+	CADKernel::FSurfacicBoundary Boundary = TechSoftUtils::GetSurfacicBoundary(A3DCylinderData->m_sParam.m_sUVDomain, OutUVReparameterization);
 	TSharedPtr<CADKernel::FSurface> Cylinder = CADKernel::FEntity::MakeShared<CADKernel::FCylinderSurface>(GeometricTolerance, CoordinateSystem, A3DCylinderData->m_dRadius * BodyScale, Boundary);
 
 	return Cylinder;
@@ -963,8 +983,10 @@ TSharedPtr<CADKernel::FSurface> FTechSoftBridge::AddPlaneSurface(const A3DSurfPl
 	}
 
 	OutUVReparameterization.AddUVTransform(A3DPlaneData->m_sParam);
+	OutUVReparameterization.ScaleUVTransform(BodyScale, BodyScale);
+	
 	CADKernel::FMatrixH CoordinateSystem = TechSoftUtils::CreateCoordinateSystem(A3DPlaneData->m_sTrsf, BodyScale);
-	CADKernel::FSurfacicBoundary Boundary = TechSoftUtils::GetSurfacicBoundary(A3DPlaneData->m_sParam.m_sUVDomain, BodyScale);
+	CADKernel::FSurfacicBoundary Boundary = TechSoftUtils::GetSurfacicBoundary(A3DPlaneData->m_sParam.m_sUVDomain, OutUVReparameterization);
 	TSharedPtr<CADKernel::FSurface> Plane = CADKernel::FEntity::MakeShared<CADKernel::FPlaneSurface>(GeometricTolerance, CoordinateSystem, Boundary);
 
 	return Plane;
@@ -1010,11 +1032,6 @@ TSharedPtr<CADKernel::FSurface> FTechSoftBridge::AddSphereSurface(const A3DSurfB
 {
 	Report.SphereSurfaceCount++;
 
-	if (bUseSurfaceAsNurbs)
-	{
-		return AddSurfaceAsNurbs(Surface, OutUVReparameterization);
-	}
-
 	TUniqueTSObj<A3DSurfSphereData> A3DSphereData(Surface);
 	if (!A3DSphereData.IsValid())
 	{
@@ -1022,8 +1039,9 @@ TSharedPtr<CADKernel::FSurface> FTechSoftBridge::AddSphereSurface(const A3DSurfB
 	}
 
 	OutUVReparameterization.AddUVTransform(A3DSphereData->m_sParam);
+
 	CADKernel::FMatrixH CoordinateSystem = TechSoftUtils::CreateCoordinateSystem(A3DSphereData->m_sTrsf, BodyScale);
-	CADKernel::FSurfacicBoundary Boundary = TechSoftUtils::GetSurfacicBoundary(A3DSphereData->m_sParam.m_sUVDomain, BodyScale);
+	CADKernel::FSurfacicBoundary Boundary = TechSoftUtils::GetSurfacicBoundary(A3DSphereData->m_sParam.m_sUVDomain, OutUVReparameterization);
 	TSharedRef<CADKernel::FSphericalSurface> Sphere = CADKernel::FEntity::MakeShared<CADKernel::FSphericalSurface>(GeometricTolerance, CoordinateSystem, A3DSphereData->m_dRadius * BodyScale, Boundary);
 
 	return Sphere;
@@ -1041,7 +1059,7 @@ TSharedPtr<CADKernel::FSurface> FTechSoftBridge::AddTorusSurface(const A3DSurfBa
 
 	OutUVReparameterization.AddUVTransform(A3DTorusData->m_sParam);
 	CADKernel::FMatrixH CoordinateSystem = TechSoftUtils::CreateCoordinateSystem(A3DTorusData->m_sTrsf, BodyScale);
-	CADKernel::FSurfacicBoundary Boundary = TechSoftUtils::GetSurfacicBoundary(A3DTorusData->m_sParam.m_sUVDomain, BodyScale);
+	CADKernel::FSurfacicBoundary Boundary = TechSoftUtils::GetSurfacicBoundary(A3DTorusData->m_sParam.m_sUVDomain, OutUVReparameterization);
 	TSharedPtr<CADKernel::FSurface> Torus = CADKernel::FEntity::MakeShared<CADKernel::FTorusSurface>(GeometricTolerance, CoordinateSystem, A3DTorusData->m_dMajorRadius * BodyScale, A3DTorusData->m_dMinorRadius * BodyScale, Boundary);
 
 	return Torus;
diff --git a/Engine/Plugins/Enterprise/DatasmithCADImporter/Source/CADInterfaces/Private/TechSoftBridge.h b/Engine/Plugins/Enterprise/DatasmithCADImporter/Source/CADInterfaces/Private/TechSoftBridge.h
--- a/Engine/Plugins/Enterprise/DatasmithCADImporter/Source/CADInterfaces/Private/TechSoftBridge.h
+++ b/Engine/Plugins/Enterprise/DatasmithCADImporter/Source/CADInterfaces/Private/TechSoftBridge.h
@@ -4,7 +4,6 @@
 
 #ifdef USE_TECHSOFT_SDK
 
-#include "CADEnum.h"
 #include "CADFileReport.h"
 
 #include "CADKernel/Geo/GeoEnum.h"
@@ -47,25 +46,14 @@ class FUVReparameterization
 private:
 	double Scale[2] = { 1., 1. };
 	double Offset[2] = { 0., 0. };
+	bool bSwapUV = false;
 	bool bNeedApply = false;
+	bool bNeedSwapOrientation = false;
 
 public:
-	static const FUVReparameterization Identity;
 
-	FUVReparameterization(const double InUScale = 1., const double InUOffset = 0., const double InVScale = 1., const double InVOffset = 0.)
+	FUVReparameterization()
 	{
-		SetCoef(InUScale, InUOffset, InVScale, InVOffset);
-		SetNeedApply();
-	}
-
-	void SetScale(const double InUScale, const double InVScale)
-	{
-		Scale[CADKernel::EIso::IsoU] = InUScale;
-		Scale[CADKernel::EIso::IsoV] = InVScale;
-		Offset[CADKernel::EIso::IsoU] = 0.;
-		Offset[CADKernel::EIso::IsoV] = 0.;
-
-		SetNeedApply();
 	}
 
 	void SetCoef(const double InUScale, const double InUOffset, const double InVScale, const double InVOffset)
@@ -77,11 +65,26 @@ public:
 		SetNeedApply();
 	}
 
-	bool NeedApply() const
+	bool GetNeedApply() const
 	{
 		return bNeedApply;
 	}
 
+	bool GetSwapUV() const
+	{
+		return bSwapUV;
+	}
+
+	bool GetNeedSwapOrientation() const
+	{
+		return bNeedSwapOrientation != bSwapUV;
+	}
+
+	void SetNeedSwapOrientation()
+	{
+		bNeedSwapOrientation = true;
+	}
+
 	void SetNeedApply()
 	{
 		if (!FMath::IsNearlyEqual(Scale[CADKernel::EIso::IsoU], 1.) || !FMath::IsNearlyEqual(Scale[CADKernel::EIso::IsoV], 1.) || !FMath::IsNearlyEqual(Offset[CADKernel::EIso::IsoU], 0.) || !FMath::IsNearlyEqual(Offset[CADKernel::EIso::IsoV], 0.))
@@ -94,8 +97,12 @@ public:
 		}
 	}
 
-	void ScaleTransform(double InUScale, double InVScale)
+	void ScaleUVTransform(double InUScale, double InVScale)
 	{
+		if (bSwapUV)
+		{
+			Swap(InUScale, InVScale);
+		}
 		Scale[CADKernel::EIso::IsoU] *= InUScale;
 		Scale[CADKernel::EIso::IsoV] *= InVScale;
 		Offset[CADKernel::EIso::IsoU] *= InUScale;
@@ -103,23 +110,53 @@ public:
 		SetNeedApply();
 	}
 
+	void Process(TArray<CADKernel::FPoint>& Poles) const
+	{
+		if(bNeedApply)
+		{
+			for (CADKernel::FPoint& Point : Poles)
+			{
+				Apply(Point);
+			}
+		}
+		if (bSwapUV)
+		{
+			for (CADKernel::FPoint& Point : Poles)
+			{
+				GetSwapUV(Point);
+			}
+		}
+	}
+
+	void AddUVTransform(A3DUVParameterizationData& Transform)
+	{
+		bSwapUV = (bool) Transform.m_bSwapUV;
+
+		Scale[0] = Scale[0] * Transform.m_dUCoeffA;
+		Scale[1] = Scale[1] * Transform.m_dVCoeffA;
+		Offset[0] = Offset[0] * Transform.m_dUCoeffA + Transform.m_dUCoeffB;
+		Offset[1] = Offset[1] * Transform.m_dVCoeffA + Transform.m_dVCoeffB;
+		SetNeedApply();
+	}
+
+	void Apply(CADKernel::FPoint2D& Point) const
+	{
+		Point.U = Scale[CADKernel::EIso::IsoU] * Point.U + Offset[CADKernel::EIso::IsoU];
+		Point.V = Scale[CADKernel::EIso::IsoV] * Point.V + Offset[CADKernel::EIso::IsoV];
+	}
+
+private:
 	void Apply(CADKernel::FPoint& Point) const
 	{
 		Point.X = Scale[CADKernel::EIso::IsoU] * Point.X + Offset[CADKernel::EIso::IsoU];
 		Point.Y = Scale[CADKernel::EIso::IsoV] * Point.Y + Offset[CADKernel::EIso::IsoV];
 	}
 
-	void AddUVTransform(A3DUVParameterizationData& Transform, bool bSwapUV = false)
+	void GetSwapUV(CADKernel::FPoint& Point) const
 	{
-		int32 UIndex = bSwapUV ? 1 : 0;
-		int32 VIndex = bSwapUV ? 0 : 1;
-
-		Scale[UIndex] = Scale[UIndex] * Transform.m_dUCoeffA;
-		Scale[VIndex] = Scale[VIndex] * Transform.m_dVCoeffA;
-		Offset[UIndex] = Offset[UIndex] * Transform.m_dUCoeffA + Transform.m_dUCoeffB;
-		Offset[VIndex] = Offset[VIndex] * Transform.m_dVCoeffA + Transform.m_dVCoeffB;
-		SetNeedApply();
+		Swap(Point.X, Point.Y);
 	}
+
 };
 } // ns TechSoftUtils
 
@@ -154,7 +191,7 @@ private:
 
 	void TraverseConnex(const A3DTopoConnex* A3DConnex, TSharedRef<CADKernel::FBody>& Body);
 	void TraverseShell(const A3DTopoShell* A3DShell, TSharedRef<CADKernel::FBody>& Body);
-	void AddFace(const A3DTopoFace* A3DFace, CADKernel::EOrientation Orientation, TSharedRef<CADKernel::FShell>& Body);
+	void AddFace(const A3DTopoFace* A3DFace, CADKernel::EOrientation Orientation, TSharedRef<CADKernel::FShell>& Body, uint32 Index);
 
 	TSharedPtr<CADKernel::FTopologicalLoop> AddLoop(const A3DTopoLoop* A3DLoop, const TSharedRef<CADKernel::FSurface>& Surface, const TechSoftUtils::FUVReparameterization& UVReparameterization, const bool bIsExternalLoop);
 	TSharedPtr<CADKernel::FTopologicalEdge> AddEdge(const A3DTopoCoEdge* A3DCoedge, const TSharedRef<CADKernel::FSurface>& Surface, const TechSoftUtils::FUVReparameterization& UVReparameterization, CADKernel::EOrientation& OutOrientation);
@@ -185,7 +222,7 @@ private:
 
 	TSharedPtr<CADKernel::FRestrictionCurve> AddRestrictionCurve(const A3DCrvBase* A3DCurve, const TSharedRef<CADKernel::FSurface>& Surface);
 
-	TSharedPtr<CADKernel::FCurve> AddCurve(const A3DCrvBase* A3DCurve, const TechSoftUtils::FUVReparameterization& UVReparameterization = TechSoftUtils::FUVReparameterization::Identity);
+	TSharedPtr<CADKernel::FCurve> AddCurve(const A3DCrvBase* A3DCurve, const TechSoftUtils::FUVReparameterization& UVReparameterization);
 
 	TSharedPtr<CADKernel::FCurve> AddCurveCircle(const A3DCrvCircle* A3DCurve, const TechSoftUtils::FUVReparameterization& UVReparameterization);
 	TSharedPtr<CADKernel::FCurve> AddCurveComposite(const A3DCrvComposite* A3DCurve, const TechSoftUtils::FUVReparameterization& UVReparameterization);
diff --git a/Engine/Plugins/Enterprise/DatasmithCADImporter/Source/CADInterfaces/Private/TechSoftFileParserCADKernelTessellator.cpp b/Engine/Plugins/Enterprise/DatasmithCADImporter/Source/CADInterfaces/Private/TechSoftFileParserCADKernelTessellator.cpp
--- a/Engine/Plugins/Enterprise/DatasmithCADImporter/Source/CADInterfaces/Private/TechSoftFileParserCADKernelTessellator.cpp
+++ b/Engine/Plugins/Enterprise/DatasmithCADImporter/Source/CADInterfaces/Private/TechSoftFileParserCADKernelTessellator.cpp
@@ -304,10 +304,6 @@ A3DStatus FTechSoftFileParserCADKernelTessellator::AdaptBRepModel()
 	{
 		// Add warning about error during the adaptation
 		CADFileData.AddWarningMessages(FString::Printf(TEXT("File %s had %d error during BRep adaptation step."), *CADFileData.GetCADFileDescription().GetFileName(), ErrorCount));
-		if (Errors != nullptr)
-		{
-			delete[] Errors;
-		}
 	}
 	else if (Ret != A3D_SUCCESS)
 	{
diff --git a/Engine/Plugins/Enterprise/DatasmithCADImporter/Source/DatasmithOpenNurbsTranslator/Private/OpenNurbsBRepToCADKernelConverter.cpp b/Engine/Plugins/Enterprise/DatasmithCADImporter/Source/DatasmithOpenNurbsTranslator/Private/OpenNurbsBRepToCADKernelConverter.cpp
--- a/Engine/Plugins/Enterprise/DatasmithCADImporter/Source/DatasmithOpenNurbsTranslator/Private/OpenNurbsBRepToCADKernelConverter.cpp
+++ b/Engine/Plugins/Enterprise/DatasmithCADImporter/Source/DatasmithOpenNurbsTranslator/Private/OpenNurbsBRepToCADKernelConverter.cpp
@@ -71,23 +71,25 @@ TSharedRef<CADKernel::FSurface> FOpenNurbsBRepToCADKernelConverter::AddSurface(O
 	FillPerAxisInfo(U, OpenNurbsSurface, NurbsData);
 	FillPerAxisInfo(V, OpenNurbsSurface, NurbsData);
 
-	int32 ControlVertexDimension = OpenNurbsSurface.CVSize();
-	NurbsData.HomogeneousPoles.SetNumUninitialized(NurbsData.PoleUCount * NurbsData.PoleVCount * ControlVertexDimension);
-	double* ControlPoints = NurbsData.HomogeneousPoles.GetData();
-	NurbsData.bIsRational = OpenNurbsSurface.IsRational();
-	ON::point_style PointStyle = OpenNurbsSurface.IsRational() ? ON::point_style::euclidean_rational : ON::point_style::not_rational;
-	for (int32 UIndex = 0; UIndex < NurbsData.PoleUCount; ++UIndex)
 	{
-		for (int32 VIndex = 0; VIndex < NurbsData.PoleVCount; ++VIndex, ControlPoints += ControlVertexDimension)
+		int32 ControlVertexDimension = OpenNurbsSurface.CVSize();
+		NurbsData.HomogeneousPoles.SetNumUninitialized(NurbsData.PoleUCount * NurbsData.PoleVCount * ControlVertexDimension);
+		double* ControlPoints = NurbsData.HomogeneousPoles.GetData();
+		NurbsData.bIsRational = OpenNurbsSurface.IsRational();
+		ON::point_style PointStyle = OpenNurbsSurface.IsRational() ? ON::point_style::euclidean_rational : ON::point_style::not_rational;
+		for (int32 UIndex = 0; UIndex < NurbsData.PoleUCount; ++UIndex)
 		{
-			OpenNurbsSurface.GetCV(UIndex, VIndex, PointStyle, ControlPoints);
+			for (int32 VIndex = 0; VIndex < NurbsData.PoleVCount; ++VIndex, ControlPoints += ControlVertexDimension)
+			{
+				OpenNurbsSurface.GetCV(UIndex, VIndex, PointStyle, ControlPoints);
+			}
 		}
 	}
 
 	// Scale ControlPoints into mm
-	if(!FMath::IsNearlyEqual(ScaleFactor, 1.))
 	{
 		int32 Offset = NurbsData.bIsRational ? 4 : 3;
+		double* ControlPoints = NurbsData.HomogeneousPoles.GetData();
 		for (int32 Index = 0; Index < NurbsData.HomogeneousPoles.Num(); Index += Offset)
 		{
 			ControlPoints[Index + 0] *= ScaleFactor;
@@ -96,29 +98,6 @@ TSharedRef<CADKernel::FSurface> FOpenNurbsBRepToCADKernelConverter::AddSurface(O
 		}
 	}
 
-#ifdef REMOVE_NEGATIVE_WEIGHT
-	if (NurbsData.bIsRational)
-	{
-		bool bHasNegativeWeight = false;
-		for (int32 Index = 3; Index < NurbsData.HomogeneousPoles.Num(); Index += 4)
-		{
-			double Weight = ControlPoints[Index];
-			if (Weight < 0.)
-			{
-				bHasNegativeWeight = true;
-				break;
-			}
-		}
-
-		if (bHasNegativeWeight)
-		{
-			UInfo.IncreaseDegree();
-			VInfo.IncreaseDegree();
-			BuildHull();
-		}
-	}
-#endif
-
 	return CADKernel::FEntity::MakeShared<CADKernel::FNURBSSurface>(GeometricTolerance, NurbsData);
 }
 
diff --git a/Engine/Source/Runtime/Datasmith/CADKernel/Private/CADKernel/Geo/Sampling/SurfacicPolyline.cpp b/Engine/Source/Runtime/Datasmith/CADKernel/Private/CADKernel/Geo/Sampling/SurfacicPolyline.cpp
--- a/Engine/Source/Runtime/Datasmith/CADKernel/Private/CADKernel/Geo/Sampling/SurfacicPolyline.cpp
+++ b/Engine/Source/Runtime/Datasmith/CADKernel/Private/CADKernel/Geo/Sampling/SurfacicPolyline.cpp
@@ -182,8 +182,17 @@ void FSurfacicPolyline::ComputeIntersectionsWithIsos(const FLinearBoundary& InBo
 
 			if ((IsoCoordinate > UMin) && (IsoCoordinate < UMax))
 			{
-				double LocalCoord = (IsoCoordinate - Points2D[Index][TypeIso]) / (Points2D[Index + 1][TypeIso] - Points2D[Index][TypeIso]);
-				double EdgeCoord = PolylineTools::LinearInterpolation(Coordinates, Index, LocalCoord);
+				double EdgeCoord = 0.;
+				const double Delta = (Points2D[Index + 1][TypeIso] - Points2D[Index][TypeIso]);
+				if (FMath::IsNearlyZero(Delta))
+				{
+					EdgeCoord = Coordinates[Index];
+				}
+				else
+				{
+					const double LocalCoord = (IsoCoordinate - Points2D[Index][TypeIso]) / Delta;
+					EdgeCoord = PolylineTools::LinearInterpolation(Coordinates, Index, LocalCoord);
+				}
 
 #ifdef DEBUG_COMPUTEINTERSECTIONSWITHISOS
 				::DisplayPoint(Points2D[Index + 1], EVisuProperty::RedPoint);
diff --git a/Engine/Source/Runtime/Datasmith/CADKernel/Private/CADKernel/Mesh/Criteria/CriteriaGrid.cpp b/Engine/Source/Runtime/Datasmith/CADKernel/Private/CADKernel/Mesh/Criteria/CriteriaGrid.cpp
--- a/Engine/Source/Runtime/Datasmith/CADKernel/Private/CADKernel/Mesh/Criteria/CriteriaGrid.cpp
+++ b/Engine/Source/Runtime/Datasmith/CADKernel/Private/CADKernel/Mesh/Criteria/CriteriaGrid.cpp
@@ -49,6 +49,18 @@ FCriteriaGrid::FCriteriaGrid(FTopologicalFace& InSurface)
 	, CoordinateGrid(InSurface.GetCrossingPointCoordinates())
 {
 	Face.Presample();
+
+	const FCoordinateGrid& FaceGrid = Face.GetCrossingPointCoordinates();
+
+	// If the size of the grid is very huge for a pre-sample step means that the surface is very poorly build
+	// Canceled the face meshing is preferred than waiting a too long time for a dirty result.
+	constexpr int32 MaxGrid = 1000000;
+	if (FaceGrid[EIso::IsoU].Num() * FaceGrid[EIso::IsoV].Num() > MaxGrid)
+	{
+		Face.SetDeleted();
+		return;
+	}
+
 	Face.InitDeltaUs();
 	Init();
 #ifdef DISPLAY_CRITERIA_GRID
diff --git a/Engine/Source/Runtime/Datasmith/CADKernel/Private/CADKernel/Mesh/Meshers/BowyerWatsonTriangulator.h b/Engine/Source/Runtime/Datasmith/CADKernel/Private/CADKernel/Mesh/Meshers/BowyerWatsonTriangulator.h
--- a/Engine/Source/Runtime/Datasmith/CADKernel/Private/CADKernel/Mesh/Meshers/BowyerWatsonTriangulator.h
+++ b/Engine/Source/Runtime/Datasmith/CADKernel/Private/CADKernel/Mesh/Meshers/BowyerWatsonTriangulator.h
@@ -2,13 +2,17 @@
 #pragma once
 
 #include "CADKernel/Core/Types.h"
+#include "CADKernel/Math/SlopeUtils.h"
 #include "CADKernel/Mesh/Meshers/IsoTriangulator/IsoCell.h"
 #include "CADKernel/UI/Display.h"
+#include "CADKernel/Mesh/Meshers/IsoTriangulator/DefineForDebug.h"
 
 //#define DEBUG_BOWYERWATSON
 namespace CADKernel
 {
 
+
+
 class FBowyerWatsonTriangulator
 {
 protected:
@@ -18,20 +22,18 @@ protected:
 		double SquareRadius;
 		FPoint Center;
 
-		FTriangle(const int32& Index0, const int32& Index1, const int32& Index2, const TArray<FPoint2D>& InVertices)
+		FTriangle(const int32& Index0, const int32& Index1, const int32& Index2, const TArray<TPair<int32, FPoint2D>>& InVertices)
 		{
 			Set(Index0, Index1, Index2, InVertices);
 		}
 
-		void Set(const int32& Index0, const int32& Index1, const int32& Index2, const TArray<FPoint2D>& InVertices)
+		void Set(const int32& Index0, const int32& Index1, const int32& Index2, const TArray<TPair<int32, FPoint2D>>& InVertices)
 		{
 			VertexIndices[0] = Index0;
 			VertexIndices[1] = Index1;
 			VertexIndices[2] = Index2;
-			FTriangle2D Triangle(InVertices[Index0], InVertices[Index1], InVertices[Index2]);
+			FTriangle2D Triangle(InVertices[Index0].Value, InVertices[Index1].Value, InVertices[Index2].Value);
 			Center = Triangle.CircumCircleCenterWithSquareRadius(SquareRadius);
-			// to avoid numerical error
-			SquareRadius += KINDA_SMALL_NUMBER;
 		}
 	};
 
@@ -43,7 +45,7 @@ public:
 	 * @param OutEdgeVertices, the edges of the mesh. An edge is defined by the indices of its vertices
 	 * So the ith edge is defined by the vertices EdgeVertexIndices[2 * ith] and EdgeVertexIndices[2 * ith + 1]
 	 */
-	FBowyerWatsonTriangulator(TArray<FPoint2D>& InVertices, TArray<int32>& OutEdgeVertices)
+	FBowyerWatsonTriangulator(TArray<TPair<int32, FPoint2D>>& InVertices, TArray<int32>& OutEdgeVertices)
 		: VerticesCount(InVertices.Num())
 		, Vertices(InVertices)
 		, EdgeVertexIndices(OutEdgeVertices)
@@ -51,57 +53,213 @@ public:
 		Init();
 	}
 
-	void Triangulate(bool bAddFirstVertex)
+	void Triangulate()
 	{
 		FTimePoint StartTime = FChrono::Now();
-#ifdef DEBUG_BOWYERWATSON
-		F3DDebugSession DelaunayDebugSession(bDisplay, TEXT("Delaunay Algo"));
-#endif // DEBUG_DELAUNAY
+
+		Vertices.Sort([](const TPair<int32, FPoint2D>& Vertex1, const TPair<int32, FPoint2D>& Vertex2) {return (Vertex1.Value.U + Vertex1.Value.V) > (Vertex2.Value.U + Vertex2.Value.V); });
 
 		// initialization of Bowyer & Watson algorithm with a bounding mesh of the vertex cloud 
-		// i.e. 2 triangles defined by the corners of the offsetted vertices bounding box
+		// i.e. 2 triangles defined by the corners of the offset vertices bounding box
 		MakeBoundingMesh();
 
-#ifdef DEBUG_BOWYERWATSON
-		if (bDisplay)
+		// insert each point in the mesh
+		// The points are sorted on the diagonal of the bbox and are inserted from one end to the other  
+		int32 VertexIndex = 0;
+		for (int32 VIndex = 0; VIndex < VerticesCount; ++VIndex)
 		{
+			if (VIndex % 2 == 0)
 			{
-				F3DDebugSession _(TEXT("Vertices"));
-				for (int32 LoopIndex = bAddFirstVertex ? 1 : 0; LoopIndex < VerticesCount; ++LoopIndex)
-				{
-					DisplayPoint(Vertices[LoopIndex], EVisuProperty::YellowPoint);
-				}
+				VertexIndex = VIndex / 2;
+			}
+			else
+			{
+				VertexIndex = VerticesCount - 1 - VIndex / 2;
 			}
-			DisplayTriangles();
-			Wait();
-		}
-#endif
 
-		// insert each point in the mesh
-		for (int32 VertexIndex = bAddFirstVertex ? 1 : 0; VertexIndex < VerticesCount; ++VertexIndex)
-		{
 			TriangleIndices.Empty(VerticesCount);
+			AdditionalTriangleIndices.Empty(VerticesCount);
 
-			const FPoint2D& NewVertex = Vertices[VertexIndex];
+			const FPoint2D& NewVertex = Vertices[VertexIndex].Value;
 
-#ifdef DEBUG_BOWYERWATSON
-			F3DDebugSession _(FString::Printf(TEXT("Step %d"), VertexIndex));
-			if(bDisplay)
+			bool bHasAdditionalTriangles = false;
+			// find all triangles whose circumcircles contain the new vertex
+			{
+				// The problem is for triangle with huge circumscribed circle radius (flat triangle)
+				// in this case, if distance between the new vertex and the circumscribed circle center nearly equals its radius
+				//  - So the idea is to check that the new vertex is not outside the triangle and could generate a flatten triangle 
+				//    To check this, the slop between each side is evaluated (ComputeUnorientedSlope(NewVertex, Triangle.Vertex[a], Triangle.Vertex[b]))
+				//    if the slop is nearly null this mean that the new vertex is outside the triangle and will generate a flat triangle
+
+				constexpr double IncreaseFactor = 1.001; // to process the case of new vertex on the circumscribed circle
+				constexpr double ReducingFactor = 0.9999 / IncreaseFactor; // to remove the case of new vertex on the circumscribed circle 
+				// the case of 4 points nearly on the same circle is manage in a second time i.e. we check that the 
+
+				for (int32 TriangleIndex = 0; TriangleIndex < TriangleSet.Num(); TriangleIndex++)
+				{
+					const FPoint2D Center = TriangleSet[TriangleIndex].Center;
+
+					const double SquareDistanceToCenter = Center.SquareDistance(NewVertex);
+
+					const double SquareRadiusMax = TriangleSet[TriangleIndex].SquareRadius * IncreaseFactor;
+					const double SquareRadiusMin = SquareRadiusMax * ReducingFactor;
+
+					if (SquareDistanceToCenter < SquareRadiusMin)
+					{
+						TriangleIndices.Add(TriangleIndex);
+					}
+					else
+					{
+						if (SquareDistanceToCenter < SquareRadiusMax)
+						{
+							bHasAdditionalTriangles = true;
+							AdditionalTriangleIndices.Add(TriangleIndex);
+						}
+					}
+				}
+			}
+
+			if (TriangleIndices.Num() == 0)
 			{
-				F3DDebugSession _(TEXT("New Point"));
-				DisplayPoint(NewVertex, EVisuProperty::YellowPoint);
+				for (int32& TriangleIndex : AdditionalTriangleIndices)
+				{
+					const FPoint2D& Point0 = Vertices[TriangleSet[TriangleIndex].VertexIndices[0]].Value;
+					const FPoint2D& Point1 = Vertices[TriangleSet[TriangleIndex].VertexIndices[1]].Value;
+					const FPoint2D& Point2 = Vertices[TriangleSet[TriangleIndex].VertexIndices[2]].Value;
+					const double Slop0 = ComputePositiveSlope(Point0, Point1, NewVertex);
+					const double Slop1 = ComputePositiveSlope(Point1, Point2, NewVertex);
+					const double Slop2 = ComputePositiveSlope(Point2, Point0, NewVertex);
+					if (Slop0 < 4 && Slop1 < 4 && Slop2 < 4)
+					{
+						TriangleIndices.Add(TriangleIndex);
+						TriangleIndex = -1;
+					}
+				}
 			}
-#endif
 
-			// find all triangles whose circumcircles contain the new vertex
-			for (int32 TriangleIndex = 0; TriangleIndex < TriangleSet.Num(); TriangleIndex++)
+			if (bHasAdditionalTriangles)
 			{
-				const FPoint2D Center = TriangleSet[TriangleIndex].Center;
-				double SquareRadius = TriangleSet[TriangleIndex].SquareRadius;
-				double SquareDistanceToCenter = Center.SquareDistance(NewVertex);
-				if (SquareDistanceToCenter < SquareRadius)
+
+				TArray<int32> TmpVertexIndices;
+				for (int32 TriangleIndex : TriangleIndices)
 				{
-					TriangleIndices.Add(TriangleIndex);
+					for (int32 Index = 0; Index < 3; Index++)
+					{
+						const int32 StartVertex = TriangleSet[TriangleIndex].VertexIndices[Index];
+						TmpVertexIndices.AddUnique(StartVertex);
+					}
+				}
+
+				TArray<TPair<int32, double>> VertexToSlop;
+				for (int32 Index : TmpVertexIndices)
+				{
+					FPoint2D Vertex = Vertices[Index].Value;
+					double Slop = ComputeSlope(NewVertex, Vertex);
+					VertexToSlop.Emplace(Index, Slop);
+				}
+
+				VertexToSlop.Sort([](const TPair<int32, double>& Vertex1, const TPair<int32, double>& Vertex2) {return Vertex1.Value < Vertex2.Value; });
+
+				//Wait();
+				bool bTriangleHasBeenAdded = true;
+				while (bTriangleHasBeenAdded)
+				{
+					bTriangleHasBeenAdded = false;
+					for (int32& TriangleIndex : AdditionalTriangleIndices)
+					{
+						if (TriangleIndex < 0)
+						{
+							continue;
+						}
+
+						int32 CandidateVertexIndex = -1;
+						for (int32 Index = 0; Index < 3; Index++)
+						{
+							int32 Candidate = TriangleSet[TriangleIndex].VertexIndices[Index];
+							if (TmpVertexIndices.Find(Candidate) == INDEX_NONE)
+							{
+								if (CandidateVertexIndex == -1)
+								{
+									CandidateVertexIndex = Candidate;
+								}
+								else
+								{
+									CandidateVertexIndex = -1;
+									break;
+								}
+							}
+						}
+
+						if (CandidateVertexIndex == -1)
+						{
+							continue;
+						}
+
+						FPoint2D CandidateVertexPoint = Vertices[CandidateVertexIndex].Value;
+
+						int32 StartIndex = -1;
+						int32 EndIndex = 0;
+						{
+							double Slope = ComputeSlope(NewVertex, CandidateVertexPoint);
+							for (int32 AIndex = 0; AIndex < VertexToSlop.Num(); ++AIndex)
+							{
+								constexpr double SmallSlop = 0.01; // 0.5 deg
+								if (FMath::IsNearlyEqual(Slope, VertexToSlop[AIndex].Value, SmallSlop))
+								{
+									EndIndex = -1;
+									break;
+								}
+
+								if (Slope < VertexToSlop[AIndex].Value)
+								{
+									EndIndex = AIndex;
+									break;
+								}
+							}
+							if (EndIndex < 0)
+							{
+								continue;
+							}
+							StartIndex = EndIndex == 0 ? VertexToSlop.Num() - 1 : EndIndex - 1;
+						}
+
+						FPoint2D StartPoint = Vertices[VertexToSlop[StartIndex].Key].Value;
+						FPoint2D Vect = StartPoint - NewVertex;
+						FPoint2D Perp(Vect.V, -Vect.U);
+						FPoint2D StartPoint2 = StartPoint - Perp;
+						TSegment<FPoint2D> Segment1(StartPoint, StartPoint2);
+
+						FPoint2D EndPoint = Vertices[VertexToSlop[EndIndex].Key].Value;
+						FPoint2D Vect2 = EndPoint - NewVertex;
+						FPoint2D Perp2(Vect2.V, -Vect2.U);
+						FPoint2D EndPoint2 = EndPoint + Perp2;
+						TSegment<FPoint2D> Segment2(EndPoint, EndPoint2);
+
+						double SlopCandidate = 4;
+						FPoint2D Intersection; 
+						if(FindIntersectionOfLines2D(Segment1, Segment2, Intersection))
+						{
+							SlopCandidate = ComputePositiveSlope(CandidateVertexPoint, NewVertex, Intersection);
+						}
+
+						if (SlopCandidate > 2.00 && SlopCandidate < 6.00)
+						{
+							bTriangleHasBeenAdded = true;
+
+							double SlopeAtNewVertex = ComputeSlope(NewVertex, CandidateVertexPoint);
+							if (EndIndex == 0 && SlopeAtNewVertex > VertexToSlop.Last().Value)
+							{
+								VertexToSlop.Emplace(CandidateVertexIndex, SlopeAtNewVertex);
+							}
+							else
+							{
+								VertexToSlop.EmplaceAt(EndIndex, CandidateVertexIndex, SlopeAtNewVertex);
+							}
+							TmpVertexIndices.Add(CandidateVertexIndex);
+							TriangleIndices.Add(TriangleIndex);
+						}
+						TriangleIndex = -1;
+					}
 				}
 			}
 
@@ -139,39 +297,31 @@ public:
 			}
 
 			// make the new triangles : Each new triangle is defined by an edge of the boundary and the new vertex 
-			int32 EdgeIndex = 0;
-
-			// The deleted triangles are replaced by the new ones
-			for (int32 TriangleIndex : TriangleIndices)
 			{
-				while (EdgeVertexIndices[EdgeIndex] < 0)
+				// The deleted triangles are replaced by the new ones
+				int32 EdgeIndex = 0;
+				for (int32 TriangleIndex : TriangleIndices)
 				{
+					while (EdgeVertexIndices[EdgeIndex] < 0)
+					{
+						EdgeIndex += 2;
+					}
+					TriangleSet[TriangleIndex].Set(EdgeVertexIndices[EdgeIndex + 1], EdgeVertexIndices[EdgeIndex], VertexIndex, Vertices);
 					EdgeIndex += 2;
 				}
-				TriangleSet[TriangleIndex].Set(EdgeVertexIndices[EdgeIndex + 1], EdgeVertexIndices[EdgeIndex], VertexIndex, Vertices);
-				EdgeIndex += 2;
-			}
-			// When all deleted triangles are replaced, the new triangles are added in the array
-			for (; EdgeIndex < EdgeVertexIndices.Num(); EdgeIndex += 2)
-			{
-				if (EdgeVertexIndices[EdgeIndex] < 0)
+
+				// When all deleted triangles are replaced, the new triangles are added in the array
+				for (; EdgeIndex < EdgeVertexIndices.Num(); EdgeIndex += 2)
 				{
-					continue;
+					if (EdgeVertexIndices[EdgeIndex] < 0)
+					{
+						continue;
+					}
+					TriangleSet.Emplace(EdgeVertexIndices[EdgeIndex + 1], EdgeVertexIndices[EdgeIndex], VertexIndex, Vertices);
 				}
-				TriangleSet.Emplace(EdgeVertexIndices[EdgeIndex + 1], EdgeVertexIndices[EdgeIndex], VertexIndex, Vertices);
 			}
-#ifdef DEBUG_BOWYERWATSON
-			DisplayTriangles();
-			Wait(bDisplay);
-#endif
 		}
 
-#ifdef DEBUG_BOWYERWATSON
-		// The final mesh
-		DisplayTriangles();
-		Wait(bDisplay);
-#endif
-
 		// Find all Edges and their type (inner edge or boundary edge)
 		EdgeVertexIndices.Empty(TriangleSet.Num() * 6);
 
@@ -220,12 +370,14 @@ public:
 				EndVertex = StartVertex;
 			}
 		}
-#ifdef DEBUG_BOWYERWATSON
-		DisplayEdges();
-		//Wait();
-#endif
-			// the bounding mesh vertices are removed
-		Vertices.SetNum(Vertices.Num() - 4);
+
+		// the bounding mesh vertices are removed
+		Vertices.SetNum(VerticesCount);
+
+		for (int32& Indice : EdgeVertexIndices)
+		{
+			Indice = Vertices[Indice].Key;
+		}
 	}
 
 	int32 OuterEdgeCount() const
@@ -296,7 +448,7 @@ public:
 private:
 	int32 VerticesCount;
 
-	TArray<FPoint2D>& Vertices;
+	TArray<TPair<int32, FPoint2D>>& Vertices;
 
 	/**
 	 * An edge is defined by the indices of its vertices
@@ -308,6 +460,7 @@ private:
 
 	// It's use to mark all triangles whose circumcircles contain the next vertex
 	TArray<int32> TriangleIndices;
+	TArray<int32> AdditionalTriangleIndices;
 
 	/**
 	 * Use to determine if the edge is a border edge of inner edge
@@ -318,13 +471,19 @@ private:
 	void MakeBoundingMesh()
 	{
 		FAABB2D VerticesBBox;
-		VerticesBBox += Vertices;
-		VerticesBBox.Offset(VerticesBBox.DiagonalLength());
+		for (const TPair<int32, FPoint2D>& Vertex : Vertices)
+		{
+			VerticesBBox += Vertex.Value;
+		}
+
+		const double DiagonalLength = VerticesBBox.DiagonalLength();
+		VerticesBBox.Offset(DiagonalLength);
 
-		Vertices.Emplace(VerticesBBox.GetCorner(3));
-		Vertices.Emplace(VerticesBBox.GetCorner(2));
-		Vertices.Emplace(VerticesBBox.GetCorner(0));
-		Vertices.Emplace(VerticesBBox.GetCorner(1));
+		int32 VerticesId = Vertices.Num();
+		Vertices.Emplace(VerticesId++, VerticesBBox.GetCorner(3));
+		Vertices.Emplace(VerticesId++, VerticesBBox.GetCorner(2));
+		Vertices.Emplace(VerticesId++, VerticesBBox.GetCorner(0));
+		Vertices.Emplace(VerticesId++, VerticesBBox.GetCorner(1));
 
 		TriangleSet.Emplace(VerticesCount, VerticesCount + 1, VerticesCount + 2, Vertices);
 		TriangleSet.Emplace(VerticesCount + 2, VerticesCount + 3, VerticesCount, Vertices);
@@ -336,6 +495,7 @@ private:
 		Vertices.Reserve(VerticesCount + 4);
 		TriangleSet.Reserve(VerticesCount);
 		TriangleIndices.Reserve(VerticesCount);
+		AdditionalTriangleIndices.Reserve(VerticesCount);
 		EdgeVertexIndices.Reserve(4 * VerticesCount);
 		EdgeInstanceCount.Reserve(2 * VerticesCount);
 	}
@@ -353,11 +513,11 @@ private:
 		{
 			if (EdgeInstanceCount[Index / 2] < 2)
 			{
-				DisplaySegment(Vertices[EdgeVertexIndices[Index]], Vertices[EdgeVertexIndices[Index + 1]], 0, EVisuProperty::YellowCurve);
+				DisplaySegment(Vertices[EdgeVertexIndices[Index]].Value * DisplayScale, Vertices[EdgeVertexIndices[Index + 1]].Value * DisplayScale, 0, EVisuProperty::YellowCurve);
 			}
 			else
 			{
-				DisplaySegment(Vertices[EdgeVertexIndices[Index]], Vertices[EdgeVertexIndices[Index + 1]], 0, EVisuProperty::PurpleCurve);
+				DisplaySegment(Vertices[EdgeVertexIndices[Index]].Value * DisplayScale, Vertices[EdgeVertexIndices[Index + 1]].Value * DisplayScale, 0, EVisuProperty::PurpleCurve);
 			}
 		}
 	};
@@ -372,7 +532,7 @@ private:
 		F3DDebugSession _(TEXT("Triangles"));
 		for (int32 Index = 0; Index < TriangleSet.Num(); Index++)
 		{
-			DisplayTriangle(Index, EVisuProperty::YellowCurve);
+			DisplayTriangle(Index, EVisuProperty::BlueCurve);
 		}
 		//Wait();
 	};
@@ -387,7 +547,8 @@ private:
 		F3DDebugSession _(TEXT("Selected Triangles"));
 		for (int32 Index : TriangleIndices)
 		{
-			DisplayTriangle(Index , EVisuProperty::BlueCurve);
+			F3DDebugSession _(TEXT("Triangle"));
+			DisplayTriangle(Index, EVisuProperty::BlueCurve);
 		}
 		//Wait();
 	};
@@ -399,9 +560,9 @@ private:
 			return;
 		}
 
-		DisplaySegment(Vertices[TriangleSet[Index].VertexIndices[0]], Vertices[TriangleSet[Index].VertexIndices[1]], 0, EVisuProperty::BlueCurve);
-		DisplaySegment(Vertices[TriangleSet[Index].VertexIndices[1]], Vertices[TriangleSet[Index].VertexIndices[2]], 0, EVisuProperty::BlueCurve);
-		DisplaySegment(Vertices[TriangleSet[Index].VertexIndices[2]], Vertices[TriangleSet[Index].VertexIndices[0]], 0, EVisuProperty::BlueCurve);
+		DisplaySegment(Vertices[TriangleSet[Index].VertexIndices[0]].Value * DisplayScale, Vertices[TriangleSet[Index].VertexIndices[1]].Value * DisplayScale, 0, Property);
+		DisplaySegment(Vertices[TriangleSet[Index].VertexIndices[1]].Value * DisplayScale, Vertices[TriangleSet[Index].VertexIndices[2]].Value * DisplayScale, 0, Property);
+		DisplaySegment(Vertices[TriangleSet[Index].VertexIndices[2]].Value * DisplayScale, Vertices[TriangleSet[Index].VertexIndices[0]].Value * DisplayScale, 0, Property);
 		//Wait();
 	};
 
diff --git a/Engine/Source/Runtime/Datasmith/CADKernel/Private/CADKernel/Mesh/Meshers/IsoTriangulator.cpp b/Engine/Source/Runtime/Datasmith/CADKernel/Private/CADKernel/Mesh/Meshers/IsoTriangulator.cpp
--- a/Engine/Source/Runtime/Datasmith/CADKernel/Private/CADKernel/Mesh/Meshers/IsoTriangulator.cpp
+++ b/Engine/Source/Runtime/Datasmith/CADKernel/Private/CADKernel/Mesh/Meshers/IsoTriangulator.cpp
@@ -50,10 +50,6 @@ FIsoTriangulator::FIsoTriangulator(FGrid& InGrid, TSharedRef<FFaceMesh> EntityMe
 		bDisplay = true;
 	}
 
-#ifdef DEBUG_BOWYERWATSON
-	bool FBowyerWatsonTriangulator::bDisplay = false;
-#endif
-
 #endif
 }
 
@@ -1873,34 +1869,34 @@ void FIsoTriangulator::MeshCycle(const EGridSpace Space, const TArray<FIsoSegmen
 	}
 
 	// Function used in FindBestTriangle
-	TFunction<void(FIsoNode*, FIsoNode*, FIsoSegment*)> BuildSegmentIfNeeded = [&](FIsoNode* NodeA, FIsoNode* NodeB, FIsoSegment* ABSegment)
+	TFunction<bool(FIsoNode*, FIsoNode*, FIsoSegment*)> BuildSegmentIfNeeded = [&](FIsoNode* NodeA, FIsoNode* NodeB, FIsoSegment* ABSegment) -> bool
 	{
 		if (ABSegment)
 		{
 			if (&ABSegment->GetFirstNode() == NodeA)
 			{
-#ifdef ADD_TRIANGLE_2D
 				if (ABSegment->HasTriangleOnLeft())
 				{
+#ifdef ADD_TRIANGLE_2D
 					F3DDebugSession _(FString::Printf(TEXT("Segment")));
 					Grid.DisplayIsoSegment(EGridSpace::UniformScaled, *NodeA, *NodeB, 0, EVisuProperty::YellowCurve);
 					Wait();
-				}
 #endif
-				ensureCADKernel(!ABSegment->HasTriangleOnLeft());
+					return false;
+				}
 				ABSegment->SetHasTriangleOnLeft();
 			}
 			else
 			{
-#ifdef ADD_TRIANGLE_2D
 				if (ABSegment->HasTriangleOnRight())
 				{
+#ifdef ADD_TRIANGLE_2D
 					F3DDebugSession _(FString::Printf(TEXT("Segment")));
 					Grid.DisplayIsoSegment(EGridSpace::UniformScaled, *NodeA, *NodeB, 0, EVisuProperty::YellowCurve);
 					Wait();
-				}
 #endif
-				ensureCADKernel(!ABSegment->HasTriangleOnRight());
+					return false;
+				}
 				ABSegment->SetHasTriangleOnRight();
 			}
 		}
@@ -1913,6 +1909,7 @@ void FIsoTriangulator::MeshCycle(const EGridSpace Space, const TArray<FIsoSegmen
 			NewSegment.SetHasTriangleOnLeft();
 			SegmentStack.Add(&NewSegment);
 		}
+		return true;
 	};
 
 	const SlopMethod GetSlopAtStartNode = ClockwiseSlop;
@@ -2144,8 +2141,14 @@ void FIsoTriangulator::MeshCycle(const EGridSpace Space, const TArray<FIsoSegmen
 				Segment->SetHasTriangleOnLeft();
 			}
 
-			BuildSegmentIfNeeded(&StartNode, CandidatNode, StartToCandiatSegment);
-			BuildSegmentIfNeeded(CandidatNode, &EndNode, EndToCandiatSegment);
+			if (!BuildSegmentIfNeeded(&StartNode, CandidatNode, StartToCandiatSegment))
+			{
+				return;
+			}
+			if (!BuildSegmentIfNeeded(CandidatNode, &EndNode, EndToCandiatSegment))
+			{
+				return;
+			}
 			Mesh->AddTriangle(EndNode.GetFaceIndex(), StartNode.GetFaceIndex(), CandidatNode->GetFaceIndex());
 
 #ifdef ADD_TRIANGLE_2D
@@ -2694,17 +2697,6 @@ void FIsoTriangulator::TriangulateInnerNodes()
 #endif
 }
 
-// =========================================================================================================================================================================================================
-// =========================================================================================================================================================================================================
-// =========================================================================================================================================================================================================
-//
-//
-//                                                                            NOT YET REVIEWED
-//
-//
-// =========================================================================================================================================================================================================
-// =========================================================================================================================================================================================================
-// =========================================================================================================================================================================================================
 //#define DEBUG_CONNECT_CELL_SUB_LOOPS_BY_NEIGHBORHOOD
 void FIsoTriangulator::ConnectCellSubLoopsByNeighborhood(FCell& Cell)
 {
@@ -2712,12 +2704,13 @@ void FIsoTriangulator::ConnectCellSubLoopsByNeighborhood(FCell& Cell)
 
 	int32 LoopCount = Cell.SubLoops.Num();
 
-	TArray<FPoint2D> LoopBarycenters;
+	TArray<TPair<int32,FPoint2D>> LoopBarycenters;
 	LoopBarycenters.Reserve(LoopCount + 4);
 
+	int32 LoopIndex = -1;
 	for (const TArray<FLoopNode*>& Nodes : Cell.SubLoops)
 	{
-		FPoint2D& BaryCenter = LoopBarycenters.Emplace_GetRef(FPoint2D::ZeroPoint);
+		++LoopIndex;
 
 		// the external loop is not processed 
 		if (Nodes[0]->GetLoopIndex() == 0)
@@ -2725,6 +2718,8 @@ void FIsoTriangulator::ConnectCellSubLoopsByNeighborhood(FCell& Cell)
 			continue;
 		}
 
+		TPair<int32, FPoint2D>& BaryCenterObj = LoopBarycenters.Emplace_GetRef(LoopIndex, FPoint2D::ZeroPoint);
+		FPoint2D& BaryCenter = BaryCenterObj.Value;
 		for (const FLoopNode* Node : Nodes)
 		{
 			BaryCenter += Node->Get2DPoint(EGridSpace::UniformScaled, Grid);
@@ -2777,7 +2772,7 @@ void FIsoTriangulator::ConnectCellSubLoopsByNeighborhood(FCell& Cell)
 		FBowyerWatsonTriangulator::bDisplay = bDisplay;
 #endif
 		FBowyerWatsonTriangulator Triangulator(LoopBarycenters, EdgeVertexIndices);
-		Triangulator.Triangulate(Cell.bHasOuterLoop);
+		Triangulator.Triangulate();
 		Triangulator.GetOuterVertices(Cell.BorderLoopIndices);
 	}
 
diff --git a/Engine/Source/Runtime/Datasmith/CADKernel/Private/CADKernel/Mesh/Meshers/ParametricMesher.cpp b/Engine/Source/Runtime/Datasmith/CADKernel/Private/CADKernel/Mesh/Meshers/ParametricMesher.cpp
--- a/Engine/Source/Runtime/Datasmith/CADKernel/Private/CADKernel/Mesh/Meshers/ParametricMesher.cpp
+++ b/Engine/Source/Runtime/Datasmith/CADKernel/Private/CADKernel/Mesh/Meshers/ParametricMesher.cpp
@@ -156,6 +156,11 @@ void FParametricMesher::ApplyFaceCriteria(FTopologicalFace& Face)
 	}
 
 	FCriteriaGrid Grid(Face);
+	if (Face.IsDeleted())
+	{
+		return;
+	}
+
 	Grid.ApplyCriteria(GetMeshModel().GetCriteria());
 
 	Face.ChooseFinalDeltaUs();
@@ -200,8 +205,10 @@ void FParametricMesher::Mesh(FTopologicalFace& Face)
 	bDisplay = (Face.GetId() == FaceToDebug);
 #endif
 
-	ensureCADKernel(!Face.IsDeleted());
-	ensureCADKernel(!Face.IsMeshed());
+	if (Face.IsDeleted() || Face.IsMeshed())
+	{
+		return;
+	}
 
 	FMessage::Printf(EVerboseLevel::Debug, TEXT("Meshing of surface %d\n"), Face.GetId());
 
@@ -659,6 +666,7 @@ void FParametricMesher::Mesh(FTopologicalEdge& InEdge, const FTopologicalFace& F
 		{
 			FinalEdgeCuttingPointCoordinates.Emplace(CuttingPoint.Coordinate, ECoordinateType::OtherCoordinate);
 		}
+		InEdge.GetLinkActiveEdge()->SetMeshed();
 	}
 	else
 	{
@@ -1145,20 +1153,19 @@ void FParametricMesher::MeshSurfaceByFront(TArray<FCostToFace>& QuadTrimmedSurfa
 			if (CandidateFacesForMesh.Num())
 			{
 				int32 MaxMeshedSideNum = CandidateFacesForMesh[0]->MeshedSideNum();
-				for (int32 Index = 0; Index < CandidateFacesForMesh.Num(); ++Index)
-				{
-					if (CandidateFacesForMesh[Index]->IsMeshed())
-					{
-						CandidateFacesForMesh.RemoveAt(Index);
-						--Index;
-					}
-				}
 
 				// next face with side well meshed are preferred
 				int32 Index = 0;
 				for (; Index < CandidateFacesForMesh.Num(); ++Index)
 				{
 					FTopologicalFace* CandidateSurface = CandidateFacesForMesh[Index];
+					if (CandidateSurface->IsMeshed())
+					{
+						CandidateFacesForMesh.RemoveAt(Index);
+						--Index;
+						continue;
+					}
+
 					if (CandidateSurface->MeshedSideNum() < MaxMeshedSideNum)
 					{
 						break;
@@ -1177,6 +1184,13 @@ void FParametricMesher::MeshSurfaceByFront(TArray<FCostToFace>& QuadTrimmedSurfa
 					for (; Index < CandidateFacesForMesh.Num(); ++Index)
 					{
 						FTopologicalFace* CandidateSurface = CandidateFacesForMesh[Index];
+						if (CandidateSurface->IsMeshed())
+						{
+							CandidateFacesForMesh.RemoveAt(Index);
+							--Index;
+							continue;
+						}
+
 						if (CandidateMeshedSideRatio < CandidateSurface->MeshedSideRatio())
 						{
 							CandidateMeshedSideRatio = CandidateSurface->MeshedSideRatio();
@@ -1195,6 +1209,13 @@ void FParametricMesher::MeshSurfaceByFront(TArray<FCostToFace>& QuadTrimmedSurfa
 			for (int32 Index = 0; Index < SecondChoiceOfCandidateFacesForMesh.Num(); ++Index)
 			{
 				FTopologicalFace* CandidateSurface = SecondChoiceOfCandidateFacesForMesh[Index];
+				if (CandidateSurface->IsMeshed())
+				{
+					SecondChoiceOfCandidateFacesForMesh.RemoveAt(Index);
+					--Index;
+					continue;
+				}
+
 				if (CandidateMeshedSideRatio < CandidateSurface->MeshedSideRatio())
 				{
 					CandidateMeshedSideRatio = CandidateSurface->MeshedSideRatio();
diff --git a/Engine/Source/Runtime/Datasmith/CADKernel/Private/CADKernel/Mesh/Structure/Grid.cpp b/Engine/Source/Runtime/Datasmith/CADKernel/Private/CADKernel/Mesh/Structure/Grid.cpp
--- a/Engine/Source/Runtime/Datasmith/CADKernel/Private/CADKernel/Mesh/Structure/Grid.cpp
+++ b/Engine/Source/Runtime/Datasmith/CADKernel/Private/CADKernel/Mesh/Structure/Grid.cpp
@@ -1109,10 +1109,10 @@ void FGrid::GetMeshOfLoop(const FTopologicalLoop& Loop)
 		FaceLoops2D[EGridSpace::Default2D].Pop();
 		FaceLoops3D.Pop();
 		NormalsOfFaceLoops.Pop();
+		NodeIdsOfFaceLoops.Pop();
 	}
 }
 
-
 //#define DEBUG_GET_BOUNDARY_MESH
 bool FGrid::GetMeshOfLoops()
 {
diff --git a/Engine/Source/Runtime/Datasmith/CADKernel/Private/CADKernel/Mesh/Structure/ModelMeshAnalyzer.cpp b/Engine/Source/Runtime/Datasmith/CADKernel/Private/CADKernel/Mesh/Structure/ModelMeshAnalyzer.cpp
--- a/Engine/Source/Runtime/Datasmith/CADKernel/Private/CADKernel/Mesh/Structure/ModelMeshAnalyzer.cpp
+++ b/Engine/Source/Runtime/Datasmith/CADKernel/Private/CADKernel/Mesh/Structure/ModelMeshAnalyzer.cpp
@@ -20,8 +20,6 @@ void FModelMeshAnalyzer::BuildMesh()
 {
 	ModelMesh.GetNodeCoordinates(NodeCoordinates);
 
-	int32 LastVertexIndex = 0;
-
 	int32 TriangleCount = 0;
 	for (const TSharedPtr<FFaceMesh>& FaceMesh : ModelMesh.GetFaceMeshes())
 	{
@@ -60,9 +58,9 @@ void FModelMeshAnalyzer::BuildMesh()
 		for (int32 Index = 0; Index < FaceMesh->TrianglesVerticesIndex.Num(); Index += 3)
 		{
 			uint32 Vertices[3];
-			Vertices[0] = VerticesGlobalIndex[TrianglesVerticesIndex[Index + 0]] + LastVertexIndex;
-			Vertices[1] = VerticesGlobalIndex[TrianglesVerticesIndex[Index + 1]] + LastVertexIndex;
-			Vertices[2] = VerticesGlobalIndex[TrianglesVerticesIndex[Index + 2]] + LastVertexIndex;
+			Vertices[0] = VerticesGlobalIndex[TrianglesVerticesIndex[Index + 0]];
+			Vertices[1] = VerticesGlobalIndex[TrianglesVerticesIndex[Index + 1]];
+			Vertices[2] = VerticesGlobalIndex[TrianglesVerticesIndex[Index + 2]];
 
 			Analyzer::FEdge* TriangleEdges[3];
 			TriangleEdges[0] = FindOrBuildEdge(Vertices[0], Vertices[1]);
@@ -93,8 +91,6 @@ void FModelMeshAnalyzer::BuildMesh()
 			VertexToEdges[Edge.VertexIndices[1]].Add(&Edge);
 		}
 	}
-
-	LastVertexIndex += NodeCoordinates.Num();
 }
 
 void FModelMeshAnalyzer::ComputeBorderCount(int32& OutBorderCount, int32& OutNonManifoldCount) const
diff --git a/Engine/Source/Runtime/Datasmith/CADKernel/Private/CADKernel/Topo/Shell.cpp b/Engine/Source/Runtime/Datasmith/CADKernel/Private/CADKernel/Topo/Shell.cpp
--- a/Engine/Source/Runtime/Datasmith/CADKernel/Private/CADKernel/Topo/Shell.cpp
+++ b/Engine/Source/Runtime/Datasmith/CADKernel/Private/CADKernel/Topo/Shell.cpp
@@ -371,7 +371,9 @@ int32 FShell::Orient()
 
 	TFunction<FBBoxWithNormal(const double)> GetSubshellBoundingBox = [&](const double ApproximationFactor) -> FBBoxWithNormal
 	{
+#ifdef DEBUG_GET_SHELL_BBOX
 		F3DDebugSession _(TEXT("GetShellBoundingBox"));
+#endif
 
 		FBBoxWithNormal BBox;
 
diff --git a/Engine/Source/Runtime/Datasmith/CADKernel/Private/CADKernel/Topo/TopologicalFace.cpp b/Engine/Source/Runtime/Datasmith/CADKernel/Private/CADKernel/Topo/TopologicalFace.cpp
--- a/Engine/Source/Runtime/Datasmith/CADKernel/Private/CADKernel/Topo/TopologicalFace.cpp
+++ b/Engine/Source/Runtime/Datasmith/CADKernel/Private/CADKernel/Topo/TopologicalFace.cpp
@@ -61,6 +61,8 @@ void FTopologicalFace::UpdateBBox(int32 IsoCount, const double ApproximationFact
 	const FSurface& Surface = GetCarrierSurface().Get();
 	FIsoCurve3DSamplerOnChord Sampler(Surface, SAG, Polyline);
 
+	IsoCount++;
+
 	TFunction <void(const EIso)> UpdateBBoxWithIsos = [&](const EIso IsoType)
 	{
 		const FLinearBoundary& Bounds = GetBoundary().Get(IsoType);
@@ -68,7 +70,6 @@ void FTopologicalFace::UpdateBBox(int32 IsoCount, const double ApproximationFact
 		EIso Other = IsoType == EIso::IsoU ? EIso::IsoV : EIso::IsoU;
 
 		double Coordinate = Bounds.Min;
-		IsoCount++;
 		const double Step = (Bounds.Max - Bounds.Min) / IsoCount;
 
 		for (int32 iIso = 1; iIso < IsoCount; iIso++)
@@ -82,7 +83,7 @@ void FTopologicalFace::UpdateBBox(int32 IsoCount, const double ApproximationFact
 			int32 IntersectionCount = Intersections.Num();
 			if (IntersectionCount == 0)
 			{
-				return;
+				continue;
 			}
 
 			FLinearBoundary CurveBounds(Intersections[0], Intersections.Last());
@@ -91,11 +92,6 @@ void FTopologicalFace::UpdateBBox(int32 IsoCount, const double ApproximationFact
 			Sampler.Set(IsoType, Coordinate, CurveBounds);
 			Sampler.Sample();
 
-			if (IntersectionCount == 0)
-			{
-				continue;
-			}
-
 			if (IntersectionCount % 2 != 0)
 			{
 				TArray<FPoint> SubPolyline;
@@ -112,6 +108,11 @@ void FTopologicalFace::UpdateBBox(int32 IsoCount, const double ApproximationFact
 				IntersectionCount--;
 			}
 
+			if (IntersectionCount == 0)
+			{
+				continue;
+			}
+
 			for (int32 ISection = 0; ISection < IntersectionCount; ISection += 2)
 			{
 				TArray<FPoint> SubPolyline;
diff --git a/Engine/Source/Runtime/Datasmith/CADKernel/Private/CADKernel/Topo/TopologicalLoop.cpp b/Engine/Source/Runtime/Datasmith/CADKernel/Private/CADKernel/Topo/TopologicalLoop.cpp
--- a/Engine/Source/Runtime/Datasmith/CADKernel/Private/CADKernel/Topo/TopologicalLoop.cpp
+++ b/Engine/Source/Runtime/Datasmith/CADKernel/Private/CADKernel/Topo/TopologicalLoop.cpp
@@ -356,26 +356,18 @@ bool FTopologicalLoop::Orient()
 			double ReferenceSlop = 0;
 			if (FMath::IsNearlyEqual(LoopSampling[Index].U, UMin))
 			{
-				ensureCADKernel(!FMath::IsNearlyEqual(LoopSampling[Index].V, VMin));
-				ensureCADKernel(!FMath::IsNearlyEqual(LoopSampling[Index].V, VMax));
 				ReferenceSlop = 6;
 			}
 			else if (FMath::IsNearlyEqual(LoopSampling[Index].U, UMax))
 			{
-				ensureCADKernel(!FMath::IsNearlyEqual(LoopSampling[Index].V, VMin));
-				ensureCADKernel(!FMath::IsNearlyEqual(LoopSampling[Index].V, VMax));
 				ReferenceSlop = 2;
 			}
 			else if (FMath::IsNearlyEqual(LoopSampling[Index].V, VMin))
 			{
-				ensureCADKernel(!FMath::IsNearlyEqual(LoopSampling[Index].U, UMin));
-				ensureCADKernel(!FMath::IsNearlyEqual(LoopSampling[Index].U, UMax));
 				ReferenceSlop = 0;
 			}
 			else if (FMath::IsNearlyEqual(LoopSampling[Index].V, VMax))
 			{
-				ensureCADKernel(!FMath::IsNearlyEqual(LoopSampling[Index].U, UMin));
-				ensureCADKernel(!FMath::IsNearlyEqual(LoopSampling[Index].U, UMax));
 				ReferenceSlop = 4;
 			}
 
@@ -793,10 +785,12 @@ void FTopologicalLoop::RemoveDegeneratedEdges()
 		FPoint2D& NearEdgeExtremity = (bPrevious == (NearOrientedEdge.Direction == EOrientation::Front)) ? NearEdgeExtremities[1].Point2D : NearEdgeExtremities[0].Point2D;
 		FPoint2D& DegeneratedEdgeExtremity = (bPrevious == (DegeneratedOrientedEdge.Direction == EOrientation::Front)) ? DegeneratedEdgeExtremities[0].Point2D : DegeneratedEdgeExtremities[1].Point2D;
 		FPoint2D& OtherDegeneratedEdgeExtremity = (bPrevious == (DegeneratedOrientedEdge.Direction == EOrientation::Front)) ? DegeneratedEdgeExtremities[1].Point2D : DegeneratedEdgeExtremities[0].Point2D;
-		double Distance = NearEdgeExtremity.Distance(DegeneratedEdgeExtremity);
 
-		double Slop = FMath::Abs(TransformIntoOrientedSlope(ComputePositiveSlope(FPoint2D::ZeroPoint, DegeneratedEdgeTangent, NearEdgeTangent)));
-		if (Slop > 3.5 && Distance < 0.01)
+		double Slop = ComputeUnorientedSlope(FPoint2D::ZeroPoint, DegeneratedEdgeTangent, NearEdgeTangent);
+		constexpr double FlatSlop = 3.9; // ~175 deg. 
+		// a degenerated edge has to be really tangent with its neighbor edge to be merge with it
+		// otherwise it could distort the neighbor edge.
+		if (Slop > FlatSlop)
 		{
 			NearOrientedEdge.Entity->ExtendTo((bPrevious == (NearOrientedEdge.Direction != EOrientation::Front)), OtherDegeneratedEdgeExtremity, OtherDegeneratedEdgeVertex);
 			DegeneratedEdge->Delete();
@@ -850,40 +844,42 @@ void FTopologicalLoop::EnsureLogicalClosing(const double Tolerance3D)
 		PreviousExtremity = PreviousExtremities[0];
 	}
 
+	const FSurfacicTolerance& Tolerance2D = Surface->GetIsoTolerances();
+	const FSurfacicTolerance LargeGapTolerance2D = Tolerance2D * 3.;
+	const FSurfacicTolerance SmallGapTolerance2D = Tolerance2D * 0.1;
 	for (int32 Index = 0; Index < Edges.Num(); ++Index)
 	{
 		FOrientedEdge OrientedEdge = Edges[Index];
 		FSurfacicCurvePointWithTolerance EdgeExtremities[2];
 		OrientedEdge.Entity->GetExtremities(EdgeExtremities);
 
-		int32 ExtremityIndex = OrientedEdge.Direction == EOrientation::Front ? 0 : 1;
-		double SquareGap3D = EdgeExtremities[ExtremityIndex].Point.SquareDistance(PreviousExtremity.Point);
+		const int32 ExtremityIndex = OrientedEdge.Direction == EOrientation::Front ? 0 : 1;
+		const double SquareGap3D = EdgeExtremities[ExtremityIndex].Point.SquareDistance(PreviousExtremity.Point);
 
 		TSharedRef<FTopologicalVertex> PreviousEdgeEndVertex = PreviousEdge.Direction == EOrientation::Front ? PreviousEdge.Entity->GetEndVertex() : PreviousEdge.Entity->GetStartVertex();
 		TSharedRef<FTopologicalVertex> EdgeStartVertex = OrientedEdge.Direction == EOrientation::Front ? OrientedEdge.Entity->GetStartVertex() : OrientedEdge.Entity->GetEndVertex();
 
-		FPoint2D Gap2D = EdgeExtremities[ExtremityIndex].Point2D - PreviousExtremity.Point2D;
-		FSurfacicTolerance Tolerance2D = Max(EdgeExtremities[ExtremityIndex].Tolerance, PreviousExtremity.Tolerance);
-
+		const FPoint2D Gap2D = EdgeExtremities[ExtremityIndex].Point2D - PreviousExtremity.Point2D;
 		if (SquareGap3D > SquareTolerance3D)
 		{
 			FMessage::Printf(Log, TEXT("Loop %d Gap 3D : %f\n"), Id, sqrt(SquareGap3D));
-			FPoint PreviousTangent = PreviousEdge.Entity->GetTangentAt(*PreviousEdgeEndVertex);
-			FPoint EdgeTangent = OrientedEdge.Entity->GetTangentAt(*EdgeStartVertex);
+			const FPoint PreviousTangent = PreviousEdge.Entity->GetTangentAt(*PreviousEdgeEndVertex);
+			const FPoint EdgeTangent = OrientedEdge.Entity->GetTangentAt(*EdgeStartVertex);
 
 			FPoint Gap = EdgeExtremities[ExtremityIndex].Point - PreviousExtremity.Point;
 
 			double CosAngle = Gap.ComputeCosinus(EdgeTangent);
-			if (CosAngle > 0.9)
+			constexpr double CosTangentLimit = 0.9; // ~25 deg : 25 deg is not big angle between the extremity curve tangent and the missing segment to close the gap.  
+			if (CosAngle > CosTangentLimit) 
 			{
 				OrientedEdge.Entity->ExtendTo(OrientedEdge.Direction == EOrientation::Front, PreviousExtremity.Point2D, PreviousEdgeEndVertex);
 			}
 			else
 			{
 				CosAngle = Gap.ComputeCosinus(PreviousTangent);
-				if (CosAngle < -0.9)
+				if (CosAngle < -CosTangentLimit) // ~25 deg
 				{
-					PreviousEdge.Entity->ExtendTo(PreviousEdge.Direction == EOrientation::Front, EdgeExtremities[ExtremityIndex].Point2D, EdgeStartVertex);
+					PreviousEdge.Entity->ExtendTo(PreviousEdge.Direction == EOrientation::Back, EdgeExtremities[ExtremityIndex].Point2D, EdgeStartVertex);
 				}
 				else
 				{
@@ -903,27 +899,47 @@ void FTopologicalLoop::EnsureLogicalClosing(const double Tolerance3D)
 				}
 			}
 		}
-		else if (FMath::Abs(Gap2D.U) > Tolerance2D.U || FMath::Abs(Gap2D.V) > Tolerance2D.V)
+		else if (FMath::Abs(Gap2D.U) > LargeGapTolerance2D.U || FMath::Abs(Gap2D.V) > LargeGapTolerance2D.V)
+		{
+			// if the gap is not so large and the extremity of on edge is tangent to the gap segment then we extend the edge to close the gap
+			FMessage::Printf(Debug, TEXT("Loop %d Gap 2D : [%f, %f] vs Tol2D [%f, %f]\n"), Id, FMath::Abs(Gap2D.U), FMath::Abs(Gap2D.V), Tolerance2D.U, Tolerance2D.V);
+
+			// If the gap in 2d is big e.g. side of a degenerated patch => build an edge
+			TSharedPtr<FTopologicalEdge> Edge = FTopologicalEdge::Make(Surface, PreviousExtremity.Point2D, PreviousEdgeEndVertex, EdgeExtremities[ExtremityIndex].Point2D, EdgeStartVertex);
+			if (Edge.IsValid())
+			{
+				Edges.EmplaceAt(Index, Edge, EOrientation::Front);
+				Edge->SetLoop(*this);
+
+				PreviousEdgeEndVertex->Link(*EdgeStartVertex);
+				++Index;
+			}
+		}
+		else if (FMath::Abs(Gap2D.U) > SmallGapTolerance2D.U || FMath::Abs(Gap2D.V) > SmallGapTolerance2D.V)
 		{
+			// if the gap is not so large and the extremity of on edge is tangent to the gap segment then we extend the edge to close the gap
 			FMessage::Printf(Log, TEXT("Loop %d Gap 2D : [%f, %f] vs Tol2D [%f, %f]\n"), Id, FMath::Abs(Gap2D.U), FMath::Abs(Gap2D.V), Tolerance2D.U, Tolerance2D.V);
 
-			FPoint2D PreviousTangent = PreviousEdge.Entity->GetTangent2DAt(*PreviousEdgeEndVertex);
-			FPoint2D EdgeTangent = OrientedEdge.Entity->GetTangent2DAt(*EdgeStartVertex);
+			const FPoint2D PreviousTangent = PreviousEdge.Entity->GetTangent2DAt(*PreviousEdgeEndVertex);
+			const FPoint2D EdgeTangent = OrientedEdge.Entity->GetTangent2DAt(*EdgeStartVertex);
 
 			double CosAngle = Gap2D.ComputeCosinus(EdgeTangent);
-			if (CosAngle > 0.9)
+			constexpr double CosFlatTangent = 0.98;  // ~10 deg: In 2D, the distortion due to the parametric space, degenerated area and else imposes to be more careful before extending a curve.
+			if (CosAngle > CosFlatTangent)
 			{
 				OrientedEdge.Entity->ExtendTo(OrientedEdge.Direction == EOrientation::Front, PreviousExtremity.Point2D, PreviousEdgeEndVertex);
 			}
 			else
 			{
 				CosAngle = Gap2D.ComputeCosinus(PreviousTangent);
-				if (CosAngle < -0.9)
+				if (CosAngle < -CosFlatTangent)
 				{
-					PreviousEdge.Entity->ExtendTo(PreviousEdge.Direction == EOrientation::Front, EdgeExtremities[ExtremityIndex].Point2D, EdgeStartVertex);
+					PreviousEdge.Entity->ExtendTo(PreviousEdge.Direction == EOrientation::Back, EdgeExtremities[ExtremityIndex].Point2D, EdgeStartVertex);
 				}
 				else
 				{
+					// if the gap is not so large and the extremity of on edge is tangent to the gap segment then we extend the edge to close the gap
+					FMessage::Printf(Log, TEXT("Add degenerated edges for small 2d gap\n"));
 					TSharedPtr<FTopologicalEdge> Edge = FTopologicalEdge::Make(Surface, PreviousExtremity.Point2D, PreviousEdgeEndVertex, EdgeExtremities[ExtremityIndex].Point2D, EdgeStartVertex);
 					if (Edge.IsValid())
 					{
@@ -932,7 +948,6 @@ void FTopologicalLoop::EnsureLogicalClosing(const double Tolerance3D)
 
 						PreviousEdgeEndVertex->Link(*EdgeStartVertex);
 						++Index;
-
 					}
 				}
 			}
diff --git a/Engine/Source/Runtime/Datasmith/CADKernel/Private/CADKernel/Topo/TopologicalVertex.cpp b/Engine/Source/Runtime/Datasmith/CADKernel/Private/CADKernel/Topo/TopologicalVertex.cpp
--- a/Engine/Source/Runtime/Datasmith/CADKernel/Private/CADKernel/Topo/TopologicalVertex.cpp
+++ b/Engine/Source/Runtime/Datasmith/CADKernel/Private/CADKernel/Topo/TopologicalVertex.cpp
@@ -40,7 +40,6 @@ void FTopologicalVertex::RemoveConnectedEdge(FTopologicalEdge& Edge)
 			return;
 		}
 	}
-	ensureCADKernel(false);
 }
 
 bool FTopologicalVertex::IsBorderVertex()
diff --git a/Engine/Source/Runtime/Datasmith/CADKernel/Private/CADKernel/UI/Display.cpp b/Engine/Source/Runtime/Datasmith/CADKernel/Private/CADKernel/UI/Display.cpp
--- a/Engine/Source/Runtime/Datasmith/CADKernel/Private/CADKernel/UI/Display.cpp
+++ b/Engine/Source/Runtime/Datasmith/CADKernel/Private/CADKernel/UI/Display.cpp
@@ -1035,9 +1035,18 @@ void DisplayMesh(const FFaceMesh& Mesh)
 	const TArray<int32>& VertexIndices = Mesh.VerticesGlobalIndex;
 	for (int32 Index = 0; Index < TriangleIndices.Num();)
 	{
-		Points[0] = *NodeIdToCoordinates[VertexIndices[TriangleIndices[Index++]]];
-		Points[1] = *NodeIdToCoordinates[VertexIndices[TriangleIndices[Index++]]];
-		Points[2] = *NodeIdToCoordinates[VertexIndices[TriangleIndices[Index++]]];
+		const FPoint** P0 = NodeIdToCoordinates.Find(VertexIndices[TriangleIndices[Index++]]);
+		const FPoint** P1 = NodeIdToCoordinates.Find(VertexIndices[TriangleIndices[Index++]]);
+		const FPoint** P2 = NodeIdToCoordinates.Find(VertexIndices[TriangleIndices[Index++]]);
+
+		if (!P0 || !P1 || !P2)
+		{
+			continue;
+		}
+
+		Points[0] = **P0;
+		Points[1] = **P1;
+		Points[2] = **P2;
 		DrawElement(2, Points, EVisuProperty::Element);
 		DrawSegment(Points[0], Points[1], EVisuProperty::EdgeMesh);
 		DrawSegment(Points[1], Points[2], EVisuProperty::EdgeMesh);
@@ -1050,8 +1059,13 @@ void DisplayMesh(const FFaceMesh& Mesh)
 
 	for (const int32& Index : VertexIndices)
 	{
+		const FPoint** PointPtr = NodeIdToCoordinates.Find(Index);
+		if (!PointPtr)
+		{
+			continue;
+		}
 		F3DDebugSegment GraphicSegment(Index);
-		DrawPoint(*NodeIdToCoordinates[Index], EVisuProperty::NodeMesh);
+		DrawPoint(**PointPtr, EVisuProperty::NodeMesh);
 	}
 
 	bool test = FSystem::Get().GetVisu()->GetParameters()->bDisplayNormals;
@@ -1061,11 +1075,17 @@ void DisplayMesh(const FFaceMesh& Mesh)
 		const TArray<FVector>& Normals = Mesh.Normals;
 		for (int32 Index = 0; Index < VertexIndices.Num(); ++Index)
 		{
+			const FPoint** PointPtr = NodeIdToCoordinates.Find(VertexIndices[Index]);
+			if (!PointPtr)
+			{
+				continue;
+			}
+
 			F3DDebugSegment GraphicSegment(Index);
 			FVector Normal = Normals[Index];
 			Normal.Normalize();
 			Normal *= NormalLength;
-			DrawSegment(*NodeIdToCoordinates[VertexIndices[Index]], *NodeIdToCoordinates[VertexIndices[Index]] + Normal, EVisuProperty::EdgeMesh);
+			DrawSegment(**PointPtr, **PointPtr + Normal, EVisuProperty::EdgeMesh);
 		}
 	}
 #endif
diff --git a/Engine/Source/Runtime/Datasmith/CADKernel/Public/CADKernel/Geo/Sampling/PolylineTools.h b/Engine/Source/Runtime/Datasmith/CADKernel/Public/CADKernel/Geo/Sampling/PolylineTools.h
--- a/Engine/Source/Runtime/Datasmith/CADKernel/Public/CADKernel/Geo/Sampling/PolylineTools.h
+++ b/Engine/Source/Runtime/Datasmith/CADKernel/Public/CADKernel/Geo/Sampling/PolylineTools.h
@@ -623,14 +623,17 @@ public:
 		int32 BoundaryIndices[2];
 		GetStartEndIndex(InBoundary, BoundaryIndices);
 
-		if (FMath::IsNearlyEqual(PolylineCoordinates[BoundaryIndices[0] + 1], InBoundary.Min, (double)SMALL_NUMBER))
+		if(BoundaryIndices[1] - BoundaryIndices[0] > 0)
 		{
-			BoundaryIndices[0]++;
-		}
+			if (FMath::IsNearlyEqual(PolylineCoordinates[BoundaryIndices[0] + 1], InBoundary.Min, (double)SMALL_NUMBER))
+			{
+				BoundaryIndices[0]++;
+			}
 
-		if (FMath::IsNearlyEqual(PolylineCoordinates[BoundaryIndices[1]], InBoundary.Max, (double)SMALL_NUMBER))
-		{
-			BoundaryIndices[1]--;
+			if (FMath::IsNearlyEqual(PolylineCoordinates[BoundaryIndices[1]], InBoundary.Max, (double)SMALL_NUMBER))
+			{
+				BoundaryIndices[1]--;
+			}
 		}
 
 		OutBBox.Update(InBoundary.Min, ComputePoint(BoundaryIndices[0], InBoundary.Min));
diff --git a/Engine/Source/Runtime/Datasmith/CADKernel/Public/CADKernel/Geo/Surfaces/TorusSurface.h b/Engine/Source/Runtime/Datasmith/CADKernel/Public/CADKernel/Geo/Surfaces/TorusSurface.h
--- a/Engine/Source/Runtime/Datasmith/CADKernel/Public/CADKernel/Geo/Surfaces/TorusSurface.h
+++ b/Engine/Source/Runtime/Datasmith/CADKernel/Public/CADKernel/Geo/Surfaces/TorusSurface.h
@@ -65,7 +65,7 @@ namespace CADKernel
 			Tolerance2DU /= ComputeScaleAlongAxis(Point2DU, Matrix, Origin);
 			Tolerance2DV /= ComputeScaleAlongAxis(Point2DV, Matrix, Origin);
 
-			MinToleranceIso.Set(Tolerance2DU, Tolerance2DV);
+			MinToleranceIso.Set(FMath::Abs(Tolerance2DU), FMath::Abs(Tolerance2DV));
 		}
 		
 	public:
diff --git a/Engine/Source/Runtime/Datasmith/CADKernel/Public/CADKernel/Math/Boundary.h b/Engine/Source/Runtime/Datasmith/CADKernel/Public/CADKernel/Math/Boundary.h
--- a/Engine/Source/Runtime/Datasmith/CADKernel/Public/CADKernel/Math/Boundary.h
+++ b/Engine/Source/Runtime/Datasmith/CADKernel/Public/CADKernel/Math/Boundary.h
@@ -164,6 +164,12 @@ namespace CADKernel
 
 		void TrimAt(const FLinearBoundary& MaxBound)
 		{
+			if (Max < MaxBound.Min || Min > MaxBound.Max)
+			{
+				*this = MaxBound;
+				return;
+			}
+
 			Min = FMath::Max(Min, MaxBound.Min);
 			Max = FMath::Min(Max, MaxBound.Max);
 		}
diff --git a/Engine/Source/Runtime/Datasmith/CADKernel/Public/CADKernel/Math/Geometry.h b/Engine/Source/Runtime/Datasmith/CADKernel/Public/CADKernel/Math/Geometry.h
--- a/Engine/Source/Runtime/Datasmith/CADKernel/Public/CADKernel/Math/Geometry.h
+++ b/Engine/Source/Runtime/Datasmith/CADKernel/Public/CADKernel/Math/Geometry.h
@@ -502,8 +502,8 @@ CADKERNEL_API void FindLoopIntersectionsWithIso(const EIso Iso, const double Iso
  */
 inline bool FastIntersectSegments2D(const TSegment<FPoint2D>& SegmentAB, const TSegment<FPoint2D>& SegmentCD)
 {
-	constexpr const double Min = -SMALL_NUMBER;
-	constexpr const double Max = 1. + SMALL_NUMBER;
+	constexpr const double Min = -DOUBLE_SMALL_NUMBER;
+	constexpr const double Max = 1. + DOUBLE_SMALL_NUMBER;
 
 	FPoint2D AB = SegmentAB[1] - SegmentAB[0];
 	FPoint2D CD = SegmentCD[1] - SegmentCD[0];
@@ -531,8 +531,8 @@ inline bool FastIntersectSegments2D(const TSegment<FPoint2D>& SegmentAB, const T
  */
 inline FPoint2D FindIntersectionOfSegments2D(const TSegment<FPoint2D>& SegmentAB, const TSegment<FPoint2D>& SegmentCD, double& OutABIntersectionCoordinate)
 {
-	constexpr const double Min = -SMALL_NUMBER;
-	constexpr const double Max = 1. + SMALL_NUMBER;
+	constexpr const double Min = -DOUBLE_SMALL_NUMBER;
+	constexpr const double Max = 1. + DOUBLE_SMALL_NUMBER;
 
 	FPoint2D AB = SegmentAB[1] - SegmentAB[0];
 	FPoint2D DC = SegmentCD[0] - SegmentCD[1];
@@ -547,9 +547,9 @@ inline FPoint2D FindIntersectionOfSegments2D(const TSegment<FPoint2D>& SegmentAB
 		FPoint2D AD = SegmentCD[1] - SegmentAB[0];
 		double DCoordinate = (AB * AD) / SquareAB;
 
-		if (CCoordinate >= -KINDA_SMALL_NUMBER && CCoordinate <= 1 + KINDA_SMALL_NUMBER)
+		if (CCoordinate >= -DOUBLE_KINDA_SMALL_NUMBER && CCoordinate <= 1 + DOUBLE_KINDA_SMALL_NUMBER)
 		{
-			if (DCoordinate >= -KINDA_SMALL_NUMBER && DCoordinate <= 1 + KINDA_SMALL_NUMBER)
+			if (DCoordinate >= -DOUBLE_KINDA_SMALL_NUMBER && DCoordinate <= 1 + DOUBLE_KINDA_SMALL_NUMBER)
 			{
 				OutABIntersectionCoordinate = (DCoordinate + CCoordinate) * 0.5;
 				return SegmentCD[0].Middle(SegmentCD[1]);
@@ -559,7 +559,7 @@ inline FPoint2D FindIntersectionOfSegments2D(const TSegment<FPoint2D>& SegmentAB
 			OutABIntersectionCoordinate = CCoordinate;
 			return SegmentCD[0];
 		}
-		else if (DCoordinate >= -KINDA_SMALL_NUMBER && DCoordinate <= 1 + KINDA_SMALL_NUMBER)
+		else if (DCoordinate >= -DOUBLE_KINDA_SMALL_NUMBER && DCoordinate <= 1 + DOUBLE_KINDA_SMALL_NUMBER)
 		{
 			DCoordinate = FMath::Clamp(DCoordinate, 0., 1.);
 			OutABIntersectionCoordinate = DCoordinate;
@@ -587,6 +587,29 @@ inline FPoint2D FindIntersectionOfSegments2D(const TSegment<FPoint2D>& SegmentAB
 	return FindIntersectionOfSegments2D(SegmentAB, SegmentCD, ABIntersectionCoordinate);
 }
 
+/**
+ * @return false if the lines are parallele
+ */
+inline bool FindIntersectionOfLines2D(const TSegment<FPoint2D>& LineAB, const TSegment<FPoint2D>& LineCD, FPoint2D& OutIntersectionPoint)
+{
+	constexpr const double Min = -DOUBLE_SMALL_NUMBER;
+	constexpr const double Max = 1. + DOUBLE_SMALL_NUMBER;
+
+	const FPoint2D AB = LineAB[1] - LineAB[0];
+	const FPoint2D DC = LineCD[0] - LineCD[1];
+	const FPoint2D AC = LineCD[0] - LineAB[0];
+
+	double ParallelCoef = DC ^ AB;
+	if (FMath::IsNearlyZero(ParallelCoef))
+	{
+		return false;
+	}
+
+	double OutABIntersectionCoordinate = (DC ^ AC) / ParallelCoef;
+	OutIntersectionPoint = LineAB[0] + AB * OutABIntersectionCoordinate;
+	return true;
+}
+
 /**
  * Similar as FastIntersectSegments2D but check intersection if both segment are carried by the same line.
  * This method is 50% slower than FastIntersectSegments2D even if the segments tested are never carried by the same line
