Datasmith: enable warning on missing case

diff --git a/Engine/Source/Runtime/Datasmith/DatasmithCore/Private/DatasmithSceneFactory.cpp b/Engine/Source/Runtime/Datasmith/DatasmithCore/Private/DatasmithSceneFactory.cpp
--- a/Engine/Source/Runtime/Datasmith/DatasmithCore/Private/DatasmithSceneFactory.cpp
+++ b/Engine/Source/Runtime/Datasmith/DatasmithCore/Private/DatasmithSceneFactory.cpp
@@ -6,6 +6,12 @@
 #include "DatasmithSceneElementsImpl.h"
 #include "DatasmithVariantElementsImpl.h"
 
+// enable a warning if not all cases values are covered by a switch statement
+#if defined(_MSC_VER)
+#pragma warning(push)
+#pragma warning (default: 4062)
+#endif
+
 TSharedPtr< IDatasmithElement > FDatasmithSceneFactory::CreateElement( EDatasmithElementType InType, const TCHAR* InName )
 {
 	constexpr uint64 DefaultSubType = 0;
@@ -34,7 +40,6 @@ TSharedPtr< IDatasmithElement > FDatasmithSceneFactory::CreateElement( EDatasmit
 		case EDatasmithElementAnimationSubType::VisibilityAnimation:
 			return CreateVisibilityAnimation( InName );
 		case EDatasmithElementAnimationSubType::BaseAnimation:
-		default:
 			ensure( false );
 			break;
 		}
@@ -114,13 +119,9 @@ TSharedPtr< IDatasmithElement > FDatasmithSceneFactory::CreateElement( EDatasmit
 		case EDatasmithElementVariantSubType::VariantSet:
 			return CreateVariantSet( InName );
 		case EDatasmithElementVariantSubType::None:
-		default:
 			ensure( false );
 			break;
 		}
-	default:
-		ensure( false );
-		break;
 	}
 
 	return TSharedPtr< IDatasmithElement >();
@@ -259,7 +260,7 @@ TSharedPtr< IDatasmithMaterialExpression > FDatasmithSceneFactory::CreateMateria
 	case EDatasmithMaterialExpressionType::Custom:
 		Expression = MakeShared<FDatasmithMaterialExpressionCustomImpl>();
 		break;
-	default:
+	case EDatasmithMaterialExpressionType::None:
 		check( false );
 		break;
 	}
@@ -371,3 +372,7 @@ TSharedRef< IDatasmithScene > FDatasmithSceneFactory::DuplicateScene( const TSha
 {
 	return MakeShared< FDatasmithSceneImpl >( StaticCastSharedRef< FDatasmithSceneImpl >( InScene ).Get() );
 }
+
+#if defined(_MSC_VER)
+#pragma warning(pop)
+#endif
