Datasmith: Setup a standard locale on xml read/write

diff --git a/Engine/Source/Runtime/Datasmith/DatasmithCore/Private/DatasmithLocaleScope.cpp b/Engine/Source/Runtime/Datasmith/DatasmithCore/Private/DatasmithLocaleScope.cpp
new file mode 100644
--- /dev/null
+++ b/Engine/Source/Runtime/Datasmith/DatasmithCore/Private/DatasmithLocaleScope.cpp
@@ -0,0 +1,23 @@
+// Copyright Epic Games, Inc. All Rights Reserved.
+
+#include "DatasmithLocaleScope.h"
+
+#include "Containers/StringConv.h"
+
+#include <clocale>
+
+
+FDatasmithLocaleScope::FDatasmithLocaleScope(const TCHAR* InScopedLocale)
+{
+	if (const char* saved = std::setlocale(LC_ALL, nullptr))
+	{
+		SavedLocale = ANSI_TO_TCHAR(saved);
+	}
+	ScopedLocale = ANSI_TO_TCHAR(std::setlocale(LC_ALL, TCHAR_TO_ANSI(InScopedLocale)));
+}
+
+FDatasmithLocaleScope::~FDatasmithLocaleScope()
+{
+	std::setlocale(LC_ALL, TCHAR_TO_ANSI(*SavedLocale));
+}
+
diff --git a/Engine/Source/Runtime/Datasmith/DatasmithCore/Private/DatasmithLocaleScope.h b/Engine/Source/Runtime/Datasmith/DatasmithCore/Private/DatasmithLocaleScope.h
new file mode 100644
--- /dev/null
+++ b/Engine/Source/Runtime/Datasmith/DatasmithCore/Private/DatasmithLocaleScope.h
@@ -0,0 +1,23 @@
+// Copyright Epic Games, Inc. All Rights Reserved.
+
+#pragma once
+
+#include "Containers/UnrealString.h"
+
+
+/**
+ * Setup a locale for a given scope, the previous locale is restored on scope end.
+ */
+struct FDatasmithLocaleScope
+{
+	// Set a custom locale within this object scope
+	FDatasmithLocaleScope(const TCHAR* InScopedLocale = TEXT("C"));
+
+	// restore the captured Locale
+	~FDatasmithLocaleScope();
+
+private:
+	FString SavedLocale; // Locale previously used, captured, and to be restored
+	FString ScopedLocale; // Locale in use within the scope
+};
+
diff --git a/Engine/Source/Runtime/Datasmith/DatasmithCore/Private/DatasmithSceneXmlReader.cpp b/Engine/Source/Runtime/Datasmith/DatasmithCore/Private/DatasmithSceneXmlReader.cpp
--- a/Engine/Source/Runtime/Datasmith/DatasmithCore/Private/DatasmithSceneXmlReader.cpp
+++ b/Engine/Source/Runtime/Datasmith/DatasmithCore/Private/DatasmithSceneXmlReader.cpp
@@ -4,6 +4,7 @@
 
 #include "DatasmithCore.h"
 #include "DatasmithDefinitions.h"
+#include "DatasmithLocaleScope.h"
 #include "DatasmithSceneFactory.h"
 #include "DatasmithUtils.h"
 
@@ -15,15 +16,6 @@
 #include "Templates/SharedPointer.h"
 #include "XmlParser.h"
 
-// Set proper locale on Mac and Linux since locale "C" and not "en_US.UTF-8" is the default on those platforms
-#if PLATFORM_MAC | PLATFORM_LINUX
-#define USE_LOCALE
-#if PLATFORM_MAC
-#include <xlocale.h>
-#else
-#include <locale.h>
-#endif
-#endif
 
 FDatasmithSceneXmlReader::FDatasmithSceneXmlReader() = default;
 FDatasmithSceneXmlReader::~FDatasmithSceneXmlReader() = default;
@@ -1132,25 +1124,13 @@ bool FDatasmithSceneXmlReader::ParseXmlFile(TSharedRef< IDatasmithScene >& OutSc
 
 	TRACE_CPUPROFILER_EVENT_SCOPE(FDatasmithSceneXmlReader::ParseXmlFile);
 
+	FDatasmithLocaleScope CLocaleScope;
+
 	if (bInAppend == false)
 	{
 		OutScene->Reset();
 	}
 
-	// Set locale to support UTF-8 character set only on current thread
-#ifdef USE_LOCALE
-	locale_t Locale = ::newlocale(LC_ALL_MASK, "en_US.UTF-8", nullptr);
-	locale_t PreviousLocale = nullptr;
-	if (!Locale)
-	{
-		UE_LOG(LogDatasmith, Warning, TEXT("locale en_US.UTF-8 is not supported by this platform. The parsing of the udatasmith file may fail."));
-	}
-	else
-	{
-		PreviousLocale = ::uselocale(Locale);
-	}
-#endif
-
 	OutScene->SetExporterSDKVersion( TEXT("N/A") ); // We're expecting to read the SDK Version from the XML file. If it's not available, put "N/A"
 
 	TMap< FString, TSharedPtr<IDatasmithActorElement> > Actors;
@@ -1384,15 +1364,6 @@ bool FDatasmithSceneXmlReader::ParseXmlFile(TSharedRef< IDatasmithScene >& OutSc
 
 	FDatasmithSceneUtils::CleanUpScene(OutScene);
 
-	// Restore locale only on current thread
-#ifdef USE_LOCALE
-	if (Locale)
-	{
-		::uselocale(PreviousLocale);
-		::freelocale(Locale);
-	}
-#endif
-
 	return true;
 }
 
diff --git a/Engine/Source/Runtime/Datasmith/DatasmithCore/Private/DatasmithSceneXmlWriter.cpp b/Engine/Source/Runtime/Datasmith/DatasmithCore/Private/DatasmithSceneXmlWriter.cpp
--- a/Engine/Source/Runtime/Datasmith/DatasmithCore/Private/DatasmithSceneXmlWriter.cpp
+++ b/Engine/Source/Runtime/Datasmith/DatasmithCore/Private/DatasmithSceneXmlWriter.cpp
@@ -4,6 +4,7 @@
 
 #include "DatasmithAnimationElements.h"
 #include "DatasmithDefinitions.h"
+#include "DatasmithLocaleScope.h"
 #include "DatasmithMaterialElements.h"
 #include "DatasmithUtils.h"
 #include "DatasmithVariantElements.h"
@@ -11,6 +12,7 @@
 
 #include "Math/Quat.h"
 
+
 class FDatasmithSceneXmlWriterImpl
 {
 public:
@@ -1871,6 +1873,9 @@ FString FDatasmithSceneXmlWriterImpl::CompModeToText(EDatasmithCompMode Mode)
 
 void FDatasmithSceneXmlWriter::Serialize( TSharedRef< IDatasmithScene > DatasmithScene, FArchive& Archive )
 {
+	// explicitly set the locale to have consistent exports
+	FDatasmithLocaleScope CLocaleScope;
+
 	FString XmlString = FString( TEXT("<DatasmithUnrealScene>") ) + LINE_TERMINATOR;
 	FDatasmithSceneXmlWriterImpl::SerializeToArchive( Archive, XmlString );
 
