CADKernel
Move Headers from public to private
Use DOUBLE constants
Remove unused files

diff --git a/Engine/Plugins/Enterprise/DatasmithCADImporter/Source/CADInterfaces/Private/TechSoftBridge.cpp b/Engine/Plugins/Enterprise/DatasmithCADImporter/Source/CADInterfaces/Private/TechSoftBridge.cpp
--- a/Engine/Plugins/Enterprise/DatasmithCADImporter/Source/CADInterfaces/Private/TechSoftBridge.cpp
+++ b/Engine/Plugins/Enterprise/DatasmithCADImporter/Source/CADInterfaces/Private/TechSoftBridge.cpp
@@ -12,35 +12,12 @@
 
 #include "CADKernel/Core/Session.h"
 
-#include "CADKernel/Geo/Curves/BezierCurve.h"
-#include "CADKernel/Geo/Curves/BoundedCurve.h"
-#include "CADKernel/Geo/Curves/CompositeCurve.h"
 #include "CADKernel/Geo/Curves/Curve.h"
-#include "CADKernel/Geo/Curves/EllipseCurve.h"
-#include "CADKernel/Geo/Curves/HyperbolaCurve.h"
-#include "CADKernel/Geo/Curves/NURBSCurve.h"
-#include "CADKernel/Geo/Curves/NURBSCurve.h"
-#include "CADKernel/Geo/Curves/ParabolaCurve.h"
-#include "CADKernel/Geo/Curves/PolylineCurve.h"
-#include "CADKernel/Geo/Curves/RestrictionCurve.h"
-#include "CADKernel/Geo/Curves/RestrictionCurve.h"
-#include "CADKernel/Geo/Curves/SegmentCurve.h"
-#include "CADKernel/Geo/Curves/SurfacicCurve.h"
-
-#include "CADKernel/Geo/Surfaces/BezierSurface.h"
-#include "CADKernel/Geo/Surfaces/CompositeSurface.h"
-#include "CADKernel/Geo/Surfaces/ConeSurface.h"
-#include "CADKernel/Geo/Surfaces/CoonsSurface.h"
-#include "CADKernel/Geo/Surfaces/CylinderSurface.h"
-#include "CADKernel/Geo/Surfaces/NURBSSurface.h"
-#include "CADKernel/Geo/Surfaces/OffsetSurface.h"
-#include "CADKernel/Geo/Surfaces/PlaneSurface.h"
-#include "CADKernel/Geo/Surfaces/RevolutionSurface.h"
-#include "CADKernel/Geo/Surfaces/RuledSurface.h"
-#include "CADKernel/Geo/Surfaces/SphericalSurface.h"
+#include "CADKernel/Geo/Curves/NURBSCurveData.h"
+
+#include "CADKernel/Geo/Surfaces/NurbsSurfaceData.h"
 #include "CADKernel/Geo/Surfaces/Surface.h"
-#include "CADKernel/Geo/Surfaces/TabulatedCylinderSurface.h"
-#include "CADKernel/Geo/Surfaces/TorusSurface.h"
+
 #include "CADKernel/Topo/Body.h"
 #include "CADKernel/Topo/Model.h"
 #include "CADKernel/Topo/Shell.h"
@@ -584,7 +561,7 @@ TSharedPtr<CADKernel::FCurve> AddCurveNurbsFromData(A3DCrvNurbsData& A3DNurbs, c
 
 	A3DCrvNurbsGet(NULL, &A3DNurbs);
 
-	return CADKernel::FEntity::MakeShared<CADKernel::FNURBSCurve>(Nurbs);
+	return CADKernel::FCurve::MakeNurbsCurve(Nurbs);
 }
 
 TSharedPtr<CADKernel::FCurve> FTechSoftBridge::AddCurveNurbs(const A3DCrvNurbs* A3DNurbs, const TechSoftUtils::FUVReparameterization& UVReparameterization)
@@ -751,7 +728,6 @@ void FTechSoftBridge::AddFace(const A3DTopoFace* A3DFace, CADKernel::EOrientatio
 	if (TopoFaceData->m_bHasTrimDomain)
 	{
 		CADKernel::FSurfacicBoundary SurfaceBoundary = TechSoftUtils::GetSurfacicBoundary(TopoFaceData->m_sSurfaceDomain, UVReparameterization);
-
 		Surface->TrimBoundaryTo(SurfaceBoundary);
 	}
 
@@ -910,9 +886,7 @@ TSharedPtr<CADKernel::FSurface> FTechSoftBridge::AddConeSurface(const A3DSurfBas
 
 	CADKernel::FMatrixH CoordinateSystem = TechSoftUtils::CreateCoordinateSystem(A3DConeData->m_sTrsf, BodyScale);
 	CADKernel::FSurfacicBoundary Boundary = TechSoftUtils::GetSurfacicBoundary(A3DConeData->m_sParam.m_sUVDomain, OutUVReparameterization);
-	TSharedPtr<CADKernel::FSurface> Cone = CADKernel::FEntity::MakeShared<CADKernel::FConeSurface>(GeometricTolerance, CoordinateSystem, A3DConeData->m_dRadius * BodyScale, A3DConeData->m_dSemiAngle, Boundary);
-
-	return Cone;
+	return CADKernel::FSurface::MakeConeSurface(GeometricTolerance, CoordinateSystem, A3DConeData->m_dRadius * BodyScale, A3DConeData->m_dSemiAngle, Boundary);
 }
 
 TSharedPtr<CADKernel::FSurface> FTechSoftBridge::AddCylinderSurface(const A3DSurfBase* Surface, TechSoftUtils::FUVReparameterization& OutUVReparameterization)
@@ -930,9 +904,7 @@ TSharedPtr<CADKernel::FSurface> FTechSoftBridge::AddCylinderSurface(const A3DSur
 
 	CADKernel::FMatrixH CoordinateSystem = TechSoftUtils::CreateCoordinateSystem(A3DCylinderData->m_sTrsf, BodyScale);
 	CADKernel::FSurfacicBoundary Boundary = TechSoftUtils::GetSurfacicBoundary(A3DCylinderData->m_sParam.m_sUVDomain, OutUVReparameterization);
-	TSharedPtr<CADKernel::FSurface> Cylinder = CADKernel::FEntity::MakeShared<CADKernel::FCylinderSurface>(GeometricTolerance, CoordinateSystem, A3DCylinderData->m_dRadius * BodyScale, Boundary);
-
-	return Cylinder;
+	return CADKernel::FSurface::MakeCylinderSurface(GeometricTolerance, CoordinateSystem, A3DCylinderData->m_dRadius * BodyScale, Boundary);
 }
 
 TSharedPtr<CADKernel::FSurface> FTechSoftBridge::AddLinearTransfoSurface(const A3DSurfBase* Surface, TechSoftUtils::FUVReparameterization& OutUVReparameterization)
@@ -987,9 +959,7 @@ TSharedPtr<CADKernel::FSurface> FTechSoftBridge::AddPlaneSurface(const A3DSurfPl
 	
 	CADKernel::FMatrixH CoordinateSystem = TechSoftUtils::CreateCoordinateSystem(A3DPlaneData->m_sTrsf, BodyScale);
 	CADKernel::FSurfacicBoundary Boundary = TechSoftUtils::GetSurfacicBoundary(A3DPlaneData->m_sParam.m_sUVDomain, OutUVReparameterization);
-	TSharedPtr<CADKernel::FSurface> Plane = CADKernel::FEntity::MakeShared<CADKernel::FPlaneSurface>(GeometricTolerance, CoordinateSystem, Boundary);
-
-	return Plane;
+	return CADKernel::FSurface::MakePlaneSurface(GeometricTolerance, CoordinateSystem, Boundary);
 }
 
 TSharedPtr<CADKernel::FSurface> FTechSoftBridge::AddRevolutionSurface(const A3DSurfRevolution* Surface, TechSoftUtils::FUVReparameterization& OutUVReparameterization)
@@ -1042,9 +1012,7 @@ TSharedPtr<CADKernel::FSurface> FTechSoftBridge::AddSphereSurface(const A3DSurfB
 
 	CADKernel::FMatrixH CoordinateSystem = TechSoftUtils::CreateCoordinateSystem(A3DSphereData->m_sTrsf, BodyScale);
 	CADKernel::FSurfacicBoundary Boundary = TechSoftUtils::GetSurfacicBoundary(A3DSphereData->m_sParam.m_sUVDomain, OutUVReparameterization);
-	TSharedRef<CADKernel::FSphericalSurface> Sphere = CADKernel::FEntity::MakeShared<CADKernel::FSphericalSurface>(GeometricTolerance, CoordinateSystem, A3DSphereData->m_dRadius * BodyScale, Boundary);
-
-	return Sphere;
+	return CADKernel::FSurface::MakeSphericalSurface(GeometricTolerance, CoordinateSystem, A3DSphereData->m_dRadius * BodyScale, Boundary);
 }
 
 TSharedPtr<CADKernel::FSurface> FTechSoftBridge::AddTorusSurface(const A3DSurfBase* Surface, TechSoftUtils::FUVReparameterization& OutUVReparameterization)
@@ -1060,9 +1028,7 @@ TSharedPtr<CADKernel::FSurface> FTechSoftBridge::AddTorusSurface(const A3DSurfBa
 	OutUVReparameterization.AddUVTransform(A3DTorusData->m_sParam);
 	CADKernel::FMatrixH CoordinateSystem = TechSoftUtils::CreateCoordinateSystem(A3DTorusData->m_sTrsf, BodyScale);
 	CADKernel::FSurfacicBoundary Boundary = TechSoftUtils::GetSurfacicBoundary(A3DTorusData->m_sParam.m_sUVDomain, OutUVReparameterization);
-	TSharedPtr<CADKernel::FSurface> Torus = CADKernel::FEntity::MakeShared<CADKernel::FTorusSurface>(GeometricTolerance, CoordinateSystem, A3DTorusData->m_dMajorRadius * BodyScale, A3DTorusData->m_dMinorRadius * BodyScale, Boundary);
-
-	return Torus;
+	return CADKernel::FSurface::MakeTorusSurface(GeometricTolerance, CoordinateSystem, A3DTorusData->m_dMajorRadius * BodyScale, A3DTorusData->m_dMinorRadius * BodyScale, Boundary);
 }
 
 TSharedPtr<CADKernel::FSurface> FTechSoftBridge::AddBlend01Surface(const A3DSurfBase* Surface, TechSoftUtils::FUVReparameterization& OutUVReparameterization)
@@ -1234,7 +1200,7 @@ TSharedPtr<CADKernel::FSurface> FTechSoftBridge::AddSurfaceNurbs(const A3DSurfNu
 		TechSoftUtils::FillDoubleArray(NurbsData.PoleUCount, NurbsData.PoleVCount, A3DNurbsData.m_pdWeights, NurbsData.Weights);
 	}
 
-	return CADKernel::FEntity::MakeShared<CADKernel::FNURBSSurface>(GeometricTolerance, NurbsData);
+	return CADKernel::FSurface::MakeNurbsSurface(GeometricTolerance, NurbsData);
 }
 
 TSharedPtr<CADKernel::FSurface> FTechSoftBridge::AddSurfaceAsNurbs(const A3DSurfBase* Surface, TechSoftUtils::FUVReparameterization& OutUVReparameterization)
diff --git a/Engine/Plugins/Enterprise/DatasmithCADImporter/Source/CADInterfaces/Private/TechSoftInterface.cpp b/Engine/Plugins/Enterprise/DatasmithCADImporter/Source/CADInterfaces/Private/TechSoftInterface.cpp
--- a/Engine/Plugins/Enterprise/DatasmithCADImporter/Source/CADInterfaces/Private/TechSoftInterface.cpp
+++ b/Engine/Plugins/Enterprise/DatasmithCADImporter/Source/CADInterfaces/Private/TechSoftInterface.cpp
@@ -66,7 +66,9 @@ bool FTechSoftInterface::InitializeKernel(const TCHAR* InEnginePluginsPath)
 			}
 		}
 	}
+#ifndef CADKERNEL_DEV
 	UE_LOG(LogCADInterfaces, Warning, TEXT("Failed to load required library in %s. Plug-in will not be functional."), *TechSoftDllPath);
+#endif
 
 #endif
 	return false;
diff --git a/Engine/Plugins/Enterprise/DatasmithCADImporter/Source/CADInterfaces/Private/TechSoftUtilsPrivate.h b/Engine/Plugins/Enterprise/DatasmithCADImporter/Source/CADInterfaces/Private/TechSoftUtilsPrivate.h
--- a/Engine/Plugins/Enterprise/DatasmithCADImporter/Source/CADInterfaces/Private/TechSoftUtilsPrivate.h
+++ b/Engine/Plugins/Enterprise/DatasmithCADImporter/Source/CADInterfaces/Private/TechSoftUtilsPrivate.h
@@ -1,8 +1,6 @@
 // Copyright Epic Games, Inc. All Rights Reserved.
 #pragma once
 
-#include "CoreMinimal.h"
-
 #include "TechSoftInterface.h"
 
 namespace CADLibrary
@@ -17,8 +15,8 @@ void ExtractAttribute(const A3DMiscAttributeData& AttributeData, TMap<FString, F
 
 FString CleanSdkName(const FString& Name);
 FString CleanCatiaInstanceSdkName(const FString& Name);
-FString Clean3dxmlInstanceSdkName(const FString& Name);
 FString CleanCatiaReferenceName(const FString& Name);
+FString Clean3dxmlInstanceSdkName(const FString& Name);
 FString Clean3dxmlReferenceSdkName(const FString& Name);
 FString CleanSwInstanceSdkName(const FString& Name);
 FString CleanSwReferenceSdkName(const FString& Name);
diff --git a/Engine/Plugins/Enterprise/DatasmithCADImporter/Source/DatasmithOpenNurbsTranslator/Private/OpenNurbsBRepToCADKernelConverter.cpp b/Engine/Plugins/Enterprise/DatasmithCADImporter/Source/DatasmithOpenNurbsTranslator/Private/OpenNurbsBRepToCADKernelConverter.cpp
--- a/Engine/Plugins/Enterprise/DatasmithCADImporter/Source/DatasmithOpenNurbsTranslator/Private/OpenNurbsBRepToCADKernelConverter.cpp
+++ b/Engine/Plugins/Enterprise/DatasmithCADImporter/Source/DatasmithOpenNurbsTranslator/Private/OpenNurbsBRepToCADKernelConverter.cpp
@@ -13,9 +13,10 @@
 #include "CADKernelTools.h"
 
 #include "CADKernel/Core/Session.h"
-#include "CADKernel/Geo/Curves/NURBSCurve.h"
-#include "CADKernel/Geo/Surfaces/NURBSSurface.h"
+#include "CADKernel/Geo/Curves/Curve.h"
+#include "CADKernel/Geo/Curves/NURBSCurveData.h"
 #include "CADKernel/Geo/Surfaces/Surface.h"
+#include "CADKernel/Geo/Surfaces/NurbsSurfaceData.h"
 
 #include "CADKernel/Math/Boundary.h"
 #include "CADKernel/Math/Point.h"
@@ -64,7 +65,7 @@ namespace
 	}
 }
 
-TSharedRef<CADKernel::FSurface> FOpenNurbsBRepToCADKernelConverter::AddSurface(ON_NurbsSurface& OpenNurbsSurface)
+TSharedPtr<CADKernel::FSurface> FOpenNurbsBRepToCADKernelConverter::AddSurface(ON_NurbsSurface& OpenNurbsSurface)
 {
 	CADKernel::FNurbsSurfaceHomogeneousData NurbsData;
 	FillPerAxisInfo(U, OpenNurbsSurface, NurbsData);
@@ -97,10 +98,10 @@ TSharedRef<CADKernel::FSurface> FOpenNurbsBRepToCADKernelConverter::AddSurface(O
 		}
 	}
 
-	return CADKernel::FEntity::MakeShared<CADKernel::FNURBSSurface>(GeometricTolerance, NurbsData);
+	return CADKernel::FSurface::MakeNurbsSurface(GeometricTolerance, NurbsData);
 }
 
-TSharedPtr<CADKernel::FTopologicalLoop> FOpenNurbsBRepToCADKernelConverter::AddLoop(const ON_BrepLoop& OpenNurbsLoop, TSharedRef<CADKernel::FSurface> & CarrierSurface, const bool bIsExternal)
+TSharedPtr<CADKernel::FTopologicalLoop> FOpenNurbsBRepToCADKernelConverter::AddLoop(const ON_BrepLoop& OpenNurbsLoop, TSharedPtr<CADKernel::FSurface> & CarrierSurface, const bool bIsExternal)
 {
 	using namespace CADKernel;
 
@@ -175,7 +176,7 @@ void FOpenNurbsBRepToCADKernelConverter::LinkEdgesLoop(const ON_BrepLoop& OpenNu
 	}
 }
 
-TSharedPtr<CADKernel::FTopologicalEdge> FOpenNurbsBRepToCADKernelConverter::AddEdge(const ON_BrepTrim& OpenNurbsTrim, TSharedRef<CADKernel::FSurface>& CarrierSurface)
+TSharedPtr<CADKernel::FTopologicalEdge> FOpenNurbsBRepToCADKernelConverter::AddEdge(const ON_BrepTrim& OpenNurbsTrim, TSharedPtr<CADKernel::FSurface>& CarrierSurface)
 {
 	using namespace CADKernel;
 
@@ -186,7 +187,7 @@ TSharedPtr<CADKernel::FTopologicalEdge> FOpenNurbsBRepToCADKernelConverter::AddE
 	}
 
 	ON_NurbsCurve OpenNurbsCurve;
-	int32 NurbFormSuccess = OpenNurbsTrim.GetNurbForm(OpenNurbsCurve); // 0:Nok 1:Ok 2:OkBut
+	int32 NurbFormSuccess = OpenNurbsTrim.GetNurbForm(OpenNurbsCurve); // 0:No ok 1:Ok 2:OkBut
 	if (NurbFormSuccess == 0)
 	{
 		return TSharedPtr<FTopologicalEdge>();
@@ -232,9 +233,9 @@ TSharedPtr<CADKernel::FTopologicalEdge> FOpenNurbsBRepToCADKernelConverter::AddE
 		NurbsCurveData.Poles[Index].Z = 0;
 	}
 
-	TSharedRef<FNURBSCurve> Nurbs = FEntity::MakeShared<FNURBSCurve>(NurbsCurveData);
+	TSharedPtr<FCurve> Nurbs = FCurve::MakeNurbsCurve(NurbsCurveData);
 
-	TSharedRef<FRestrictionCurve> RestrictionCurve = FEntity::MakeShared<FRestrictionCurve>(CarrierSurface, Nurbs);
+	TSharedRef<FRestrictionCurve> RestrictionCurve = FEntity::MakeShared<FRestrictionCurve>(CarrierSurface.ToSharedRef(), Nurbs.ToSharedRef());
 
 	ON_Interval dom = OpenNurbsCurve.Domain();
 	FLinearBoundary Boundary(dom.m_t[0], dom.m_t[1]);
@@ -261,7 +262,11 @@ TSharedPtr<CADKernel::FTopologicalFace> FOpenNurbsBRepToCADKernelConverter::AddF
 	ON_NurbsSurface OpenNurbsSurface;
 	OpenNurbsFace.NurbsSurface(&OpenNurbsSurface);
 
-	TSharedRef<FSurface> Surface = AddSurface(OpenNurbsSurface);
+	TSharedPtr<FSurface> Surface = AddSurface(OpenNurbsSurface);
+	if (!Surface.IsValid())
+	{
+		return TSharedPtr<CADKernel::FTopologicalFace>();
+	}
 
 	TSharedRef<FTopologicalFace> Face = FEntity::MakeShared<FTopologicalFace>(Surface);
 
diff --git a/Engine/Plugins/Enterprise/DatasmithCADImporter/Source/DatasmithOpenNurbsTranslator/Private/OpenNurbsBRepToCADKernelConverter.h b/Engine/Plugins/Enterprise/DatasmithCADImporter/Source/DatasmithOpenNurbsTranslator/Private/OpenNurbsBRepToCADKernelConverter.h
--- a/Engine/Plugins/Enterprise/DatasmithCADImporter/Source/DatasmithOpenNurbsTranslator/Private/OpenNurbsBRepToCADKernelConverter.h
+++ b/Engine/Plugins/Enterprise/DatasmithCADImporter/Source/DatasmithOpenNurbsTranslator/Private/OpenNurbsBRepToCADKernelConverter.h
@@ -42,8 +42,8 @@ public:
 	/**
 	 * Set BRep to tessellate, offsetting it prior to tessellation(used to set mesh pivot at the center of the surface bounding box)
 	 *
-	 * @param  Brep	a brep to tessellate
-	 * @param  Offset translate brep by this value before tessellating
+	 * @param  Brep	a BRep to tessellate
+	 * @param  Offset translate BRep by this value before tessellating
 	 */
 	bool AddBRep(ON_Brep& Brep, const ON_3dVector& Offset);
 
@@ -51,9 +51,9 @@ public:
 
 private:
 	TSharedPtr<CADKernel::FTopologicalFace> AddFace(const ON_BrepFace& OpenNurbsFace);
-	TSharedRef<CADKernel::FSurface> AddSurface(ON_NurbsSurface& Surface);
+	TSharedPtr<CADKernel::FSurface> AddSurface(ON_NurbsSurface& Surface);
 
-	TSharedPtr<CADKernel::FTopologicalLoop> AddLoop(const ON_BrepLoop& OpenNurbsLoop, TSharedRef<CADKernel::FSurface>& CarrierSurface, const bool bIsExternal);
+	TSharedPtr<CADKernel::FTopologicalLoop> AddLoop(const ON_BrepLoop& OpenNurbsLoop, TSharedPtr<CADKernel::FSurface>& CarrierSurface, const bool bIsExternal);
 
 	/**
 	 * Build face's links with its neighbor have to be done after the loop is finalize.
@@ -61,7 +61,7 @@ private:
 	 */
 	void LinkEdgesLoop(const ON_BrepLoop& OpenNurbsLoop, CADKernel::FTopologicalLoop& Loop);
 
-	TSharedPtr<CADKernel::FTopologicalEdge> AddEdge(const ON_BrepTrim& OpenNurbsTrim, TSharedRef<CADKernel::FSurface>& CarrierSurface);
+	TSharedPtr<CADKernel::FTopologicalEdge> AddEdge(const ON_BrepTrim& OpenNurbsTrim, TSharedPtr<CADKernel::FSurface>& CarrierSurface);
 
 protected:
 
diff --git a/Engine/Plugins/Enterprise/DatasmithCADImporter/Source/DatasmithWireTranslator/Private/AliasModelToCADKernelConverter.cpp b/Engine/Plugins/Enterprise/DatasmithCADImporter/Source/DatasmithWireTranslator/Private/AliasModelToCADKernelConverter.cpp
--- a/Engine/Plugins/Enterprise/DatasmithCADImporter/Source/DatasmithWireTranslator/Private/AliasModelToCADKernelConverter.cpp
+++ b/Engine/Plugins/Enterprise/DatasmithCADImporter/Source/DatasmithWireTranslator/Private/AliasModelToCADKernelConverter.cpp
@@ -8,9 +8,9 @@
 #include "MeshDescriptionHelper.h"
 
 #include "CADKernel/Core/Session.h"
-#include "CADKernel/Geo/Curves/NURBSCurve.h"
+#include "CADKernel/Geo/Curves/NURBSCurveData.h"
 #include "CADKernel/Geo/GeoEnum.h"
-#include "CADKernel/Geo/Surfaces/NURBSSurface.h"
+#include "CADKernel/Geo/Surfaces/NurbsSurfaceData.h"
 #include "CADKernel/Geo/Surfaces/Surface.h"
 
 #include "CADKernel/Math/Point.h"
@@ -99,7 +99,7 @@ namespace AliasToCADKernelUtils
 			Poles[Index + 2] *= 10.;
 		}
 
-		return CADKernel::FEntity::MakeShared<CADKernel::FNURBSSurface>(GeometricTolerance, NURBSData);
+		return CADKernel::FSurface::MakeNurbsSurface(GeometricTolerance, NURBSData);
 	}
 }
 
@@ -134,9 +134,9 @@ TSharedPtr<CADKernel::FTopologicalEdge> FAliasModelToCADKernelConverter::AddEdge
 		NurbsCurveData.Poles[Index].Z = 0;
 	}
 
-	TSharedRef<CADKernel::FNURBSCurve> Nurbs = CADKernel::FEntity::MakeShared<CADKernel::FNURBSCurve>(NurbsCurveData);
+	TSharedPtr<CADKernel::FCurve> Nurbs = CADKernel::FCurve::MakeNurbsCurve(NurbsCurveData);
 
-	TSharedRef<CADKernel::FRestrictionCurve> RestrictionCurve = CADKernel::FEntity::MakeShared<CADKernel::FRestrictionCurve>(CarrierSurface.ToSharedRef(), Nurbs);
+	TSharedRef<CADKernel::FRestrictionCurve> RestrictionCurve = CADKernel::FCurve::MakeShared<CADKernel::FRestrictionCurve>(CarrierSurface.ToSharedRef(), Nurbs.ToSharedRef());
 	TSharedPtr<CADKernel::FTopologicalEdge> Edge = CADKernel::FTopologicalEdge::Make(RestrictionCurve);
 	if (!Edge.IsValid())
 	{
diff --git a/Engine/Source/Runtime/Datasmith/CADKernel/Private/CADKernel/Core/System.cpp b/Engine/Source/Runtime/Datasmith/CADKernel/Private/CADKernel/Core/System.cpp
--- a/Engine/Source/Runtime/Datasmith/CADKernel/Private/CADKernel/Core/System.cpp
+++ b/Engine/Source/Runtime/Datasmith/CADKernel/Private/CADKernel/Core/System.cpp
@@ -4,7 +4,6 @@
 #include "CADKernel/Core/Database.h"
 #include "CADKernel/Core/Entity.h"
 #include "CADKernel/Core/KernelParameters.h"
-#include "CADKernel/Core/Version.h"
 #include "CADKernel/UI/Visu.h"
 #include "CADKernel/Utils/Util.h"
 
@@ -12,6 +11,8 @@
 #include "Misc/Paths.h"
 
 #ifdef CADKERNEL_DEV
+#include "CADKernel/Core/Version.h"
+
 #include <stdlib.h>
 #include <signal.h>
 #endif 
@@ -139,7 +140,10 @@ void FSystem::InitializeCADKernel()
 
 FString FSystem::GetToolkitVersion()
 {
+#ifdef CADKERNEL_DEV
 	return UTF8_TO_TCHAR(TOOLKIT_VERSION_ASCII);
+#endif
+	return FString();
 }
 
 FString FSystem::GetCompilationDate()
@@ -151,9 +155,10 @@ void FSystem::PrintHeader()
 {
 	FMessage::Printf(Log, TEXT("_______________________________________________________________________________\n"));
 	FMessage::Printf(Log, TEXT("\n"));
-
+#ifdef CADKERNEL_DEV
 	FMessage::Printf(Log, TEXT("\tDatasmith CAD Kernel Toolkit release %s (%s)\n"), UTF8_TO_TCHAR(TOOLKIT_VERSION_ASCII), UTF8_TO_TCHAR(__DATE__));
-	FMessage::Printf(Log, TEXT("\t" EPIC_COPYRIGHT "\n"));
+#endif
+	FMessage::Printf(Log, TEXT("\tCopyright Epic Games, Inc. All Rights Reserved.\n"));
 	FMessage::Printf(Log, TEXT("\n"));
 	FMessage::Printf(Log, TEXT("_______________________________________________________________________________\n"));
 	FMessage::Printf(Log, TEXT("\n"));
diff --git a/Engine/Source/Runtime/Datasmith/CADKernel/Private/CADKernel/Core/Version.h b/Engine/Source/Runtime/Datasmith/CADKernel/Private/CADKernel/Core/Version.h
deleted file mode 100644
diff --git a/Engine/Source/Runtime/Datasmith/CADKernel/Private/CADKernel/Geo/Curves/BezierCurve.h b/Engine/Source/Runtime/Datasmith/CADKernel/Private/CADKernel/Geo/Curves/BezierCurve.h
new file mode 100644
--- /dev/null
+++ b/Engine/Source/Runtime/Datasmith/CADKernel/Private/CADKernel/Geo/Curves/BezierCurve.h
@@ -0,0 +1,58 @@
+// Copyright Epic Games, Inc. All Rights Reserved.
+#pragma once
+
+#include "CADKernel/Geo/Curves/Curve.h"
+
+namespace CADKernel
+{
+	class CADKERNEL_API FBezierCurve : public FCurve
+	{
+		friend class FEntity;
+
+	protected:
+		TArray<FPoint> Poles;
+
+		FBezierCurve(const TArray<FPoint>& InPoles)
+			: Poles{ InPoles }
+		{
+		}
+
+		FBezierCurve() = default;
+
+	public:
+
+		virtual void Serialize(FCADKernelArchive& Ar) override
+		{
+			FCurve::Serialize(Ar);
+			Ar.Serialize(Poles);
+		}
+
+#ifdef CADKERNEL_DEV
+		virtual FInfoEntity& GetInfo(FInfoEntity&) const override;
+#endif
+
+		virtual ECurve GetCurveType() const override
+		{
+			return ECurve::Bezier;
+		}
+
+		int32 GetDegre() const
+		{
+			return Poles.Num() - 1;
+		}
+
+		const TArray<FPoint>& GetPoles() const
+		{
+			return Poles;
+		}
+
+		virtual TSharedPtr<FEntityGeom> ApplyMatrix(const FMatrixH& InMatrix) const override;
+
+		virtual void EvaluatePoint(double Coordinate, FCurvePoint& OutPoint, int32 DerivativeOrder = 0) const override;
+
+		virtual void ExtendTo(const FPoint& Point) override;
+
+	};
+
+} // namespace CADKernel
+
diff --git a/Engine/Source/Runtime/Datasmith/CADKernel/Private/CADKernel/Geo/Curves/BoundedCurve.cpp b/Engine/Source/Runtime/Datasmith/CADKernel/Private/CADKernel/Geo/Curves/BoundedCurve.cpp
--- a/Engine/Source/Runtime/Datasmith/CADKernel/Private/CADKernel/Geo/Curves/BoundedCurve.cpp
+++ b/Engine/Source/Runtime/Datasmith/CADKernel/Private/CADKernel/Geo/Curves/BoundedCurve.cpp
@@ -40,7 +40,7 @@ TSharedPtr<FCurve> FBoundedCurve::MakeBoundedCurve(const FLinearBoundary& InBoun
 		NewBoundary.Max = UMax;
 	}
 
-	if((NewBoundary.Min -SMALL_NUMBER)<UMin && (NewBoundary.Max +SMALL_NUMBER)>UMax)
+	if((NewBoundary.Min -DOUBLE_SMALL_NUMBER)<UMin && (NewBoundary.Max +DOUBLE_SMALL_NUMBER)>UMax)
 	{
 		return FEntity::MakeShared<FBoundedCurve>(*this);
 	}
diff --git a/Engine/Source/Runtime/Datasmith/CADKernel/Private/CADKernel/Geo/Curves/BoundedCurve.h b/Engine/Source/Runtime/Datasmith/CADKernel/Private/CADKernel/Geo/Curves/BoundedCurve.h
new file mode 100644
--- /dev/null
+++ b/Engine/Source/Runtime/Datasmith/CADKernel/Private/CADKernel/Geo/Curves/BoundedCurve.h
@@ -0,0 +1,74 @@
+// Copyright Epic Games, Inc. All Rights Reserved.
+#pragma once
+
+#include "CADKernel/Geo/Curves/Curve.h"
+
+namespace CADKernel
+{
+
+	class CADKERNEL_API FBoundedCurve : public FCurve
+	{
+		friend class FEntity;
+
+	protected:
+
+		TSharedPtr<FCurve> Curve;
+
+		FBoundedCurve(TSharedRef<FCurve> InCurve, const FLinearBoundary& InBoundary, int8 InDimension)
+			: FCurve(InBoundary, InDimension)
+			, Curve(InCurve)
+		{
+		}
+
+		FBoundedCurve() = default;
+
+	public:
+
+		virtual void Serialize(FCADKernelArchive& Ar) override
+		{
+			FCurve::Serialize(Ar);
+			SerializeIdent(Ar, Curve);
+		}
+
+		virtual void SpawnIdent(FDatabase& Database) override
+		{
+			if (!FEntity::SetId(Database))
+			{
+				return;
+			}
+
+			Curve->SpawnIdent(Database);
+		}
+
+		virtual void ResetMarkersRecursively() override
+		{
+			ResetMarkers();
+			Curve->ResetMarkersRecursively();
+		}
+
+#ifdef CADKERNEL_DEV
+		virtual FInfoEntity& GetInfo(FInfoEntity&) const override;
+#endif
+
+		virtual ECurve GetCurveType() const override
+		{
+			return ECurve::BoundedCurve;
+		}
+
+		TSharedRef<const FCurve> GetBaseCurve()
+		{
+			return Curve.ToSharedRef();
+		}
+
+		virtual TSharedPtr<FEntityGeom> ApplyMatrix(const FMatrixH& InMatrix) const override;
+
+		virtual void EvaluatePoint(double Coordinate, FCurvePoint& OutPoint, int32 DerivativeOrder = 0) const override;
+		virtual void Evaluate2DPoint(double Coordinate, FCurvePoint2D& OutPoint, int32 DerivativeOrder = 0) const override;
+
+		virtual void FindNotDerivableCoordinates(const FLinearBoundary& InBoundary, int32 DerivativeOrder, TArray<double>& OutNotDerivableCoordinates) const override;
+
+		virtual TSharedPtr<FCurve> MakeBoundedCurve(const FLinearBoundary& InBoundary) override;
+	};
+
+} // namespace CADKernel
+
diff --git a/Engine/Source/Runtime/Datasmith/CADKernel/Private/CADKernel/Geo/Curves/CompositeCurve.h b/Engine/Source/Runtime/Datasmith/CADKernel/Private/CADKernel/Geo/Curves/CompositeCurve.h
new file mode 100644
--- /dev/null
+++ b/Engine/Source/Runtime/Datasmith/CADKernel/Private/CADKernel/Geo/Curves/CompositeCurve.h
@@ -0,0 +1,93 @@
+// Copyright Epic Games, Inc. All Rights Reserved.
+#pragma once
+
+#include "CADKernel/Core/OrientedEntity.h"
+#include "CADKernel/Geo/Curves/Curve.h"
+
+namespace CADKernel
+{
+	struct FCurvePoint;
+
+	class CADKERNEL_API FOrientedCurve : public TOrientedEntity<FCurve>
+	{
+	public:
+		FOrientedCurve(TSharedPtr<FCurve>& InEntity, EOrientation InDirection)
+			: TOrientedEntity(InEntity, InDirection)
+		{
+		}
+
+		FOrientedCurve(const FOrientedCurve& OrientiredEntity)
+			: TOrientedEntity(OrientiredEntity)
+		{
+		}
+
+		FOrientedCurve()
+			: TOrientedEntity()
+		{
+		}
+	};
+
+	class CADKERNEL_API FCompositeCurve : public FCurve
+	{
+		friend class FEntity;
+
+	protected:
+		TArray<FOrientedCurve> Curves;
+		TArray<double> Coordinates;
+
+		FCompositeCurve(const TArray<TSharedPtr<FCurve>>& Curves, bool bDoInversions = false);
+
+		FCompositeCurve() = default;
+
+	public:
+
+		virtual void Serialize(FCADKernelArchive& Ar) override
+		{
+			FCurve::Serialize(Ar);
+			SerializeIdents(Ar, (TArray<TOrientedEntity<FEntity>>&) Curves);
+			Ar.Serialize(Coordinates);
+		}
+
+		virtual void SpawnIdent(FDatabase& Database) override
+		{
+			if (!FEntity::SetId(Database))
+			{
+				return;
+			}
+
+			SpawnIdentOnEntities((TArray<TOrientedEntity<FEntity>>&) Curves, Database);
+		}
+
+		virtual void ResetMarkersRecursively() override
+		{
+			ResetMarkers();
+			ResetMarkersRecursivelyOnEntities((TArray<TOrientedEntity<FEntity>>&) Curves);
+		}
+
+#ifdef CADKERNEL_DEV
+		virtual FInfoEntity& GetInfo(FInfoEntity&) const override;
+#endif
+
+		virtual ECurve GetCurveType() const override
+		{
+			return ECurve::Composite;
+		}
+
+		virtual TSharedPtr<FEntityGeom> ApplyMatrix(const FMatrixH& InMatrix) const override;
+
+		virtual TSharedPtr<FCurve> GetCurve(int32 Index) const 
+		{
+			return Curves[Index].Entity;
+		}
+
+		virtual void EvaluatePoint(double Coordinate, FCurvePoint& OutPoint, int32 DerivativeOrder = 0) const override;
+		virtual void Evaluate2DPoint(double Coordinate, FCurvePoint2D& OutPoint, int32 DerivativeOrder = 0) const override;
+
+		virtual void FindNotDerivableCoordinates(const FLinearBoundary& InBoundary, int32 DerivativeOrder, TArray<double>& OutNotDerivableCoordinates) const override;
+
+		double LocalToGlobalCoordinate(int32 CurveIndex, double LocalCoordinate) const;
+		double GlobalToLocalCoordinate(int32 CurveIndex, double GlobalCoordinate) const;
+	};
+
+} // namespace CADKernel
+
diff --git a/Engine/Source/Runtime/Datasmith/CADKernel/Private/CADKernel/Geo/Curves/CurveMaker.cpp b/Engine/Source/Runtime/Datasmith/CADKernel/Private/CADKernel/Geo/Curves/CurveMaker.cpp
new file mode 100644
--- /dev/null
+++ b/Engine/Source/Runtime/Datasmith/CADKernel/Private/CADKernel/Geo/Curves/CurveMaker.cpp
@@ -0,0 +1,24 @@
+// Copyright Epic Games, Inc. All Rights Reserved.
+
+#include "CADKernel/Geo/Curves/BoundedCurve.h"
+#include "CADKernel/Geo/Curves/BezierCurve.h"
+#include "CADKernel/Geo/Curves/Curve.h"
+#include "CADKernel/Geo/Curves/NURBSCurve.h"
+#include "CADKernel/Geo/Curves/SegmentCurve.h"
+#include "CADKernel/Geo/Curves/EllipseCurve.h"
+#include "CADKernel/Geo/Curves/HyperbolaCurve.h"
+#include "CADKernel/Geo/Curves/ParabolaCurve.h"
+#include "CADKernel/Geo/Curves/CompositeCurve.h"
+#include "CADKernel/Geo/Curves/PolylineCurve.h"
+#include "CADKernel/Geo/Curves/RestrictionCurve.h"
+#include "CADKernel/Geo/Curves/SurfacicCurve.h"
+
+namespace CADKernel
+{
+
+TSharedPtr<FCurve> FCurve::MakeNurbsCurve(FNurbsCurveData& InNurbsData)
+{
+	return FEntity::MakeShared<CADKernel::FNURBSCurve>(InNurbsData);
+}
+
+}
\ No newline at end of file
diff --git a/Engine/Source/Runtime/Datasmith/CADKernel/Private/CADKernel/Geo/Curves/EllipseCurve.h b/Engine/Source/Runtime/Datasmith/CADKernel/Private/CADKernel/Geo/Curves/EllipseCurve.h
new file mode 100644
--- /dev/null
+++ b/Engine/Source/Runtime/Datasmith/CADKernel/Private/CADKernel/Geo/Curves/EllipseCurve.h
@@ -0,0 +1,115 @@
+// Copyright Epic Games, Inc. All Rights Reserved.
+#pragma once
+
+#include "CADKernel/Geo/Curves/Curve.h"
+#include "CADKernel/Geo/GeoPoint.h"
+#include "CADKernel/Math/MathConst.h"
+
+namespace CADKernel
+{
+
+	class CADKERNEL_API FEllipseCurve : public FCurve
+	{
+		friend class FEntity;
+
+	protected:
+		FMatrixH Matrix;
+		double   RadiusU;
+		double   RadiusV;
+		FMatrixH InverseMatrix;
+
+		FEllipseCurve(const FMatrixH& InMatrix, double InRadiusU, double InRadiusV, int8 InDimension = 3)
+			: FCurve(FLinearBoundary(0, PI * 2.), InDimension)
+			, Matrix(InMatrix)
+			, RadiusU(InRadiusU)
+			, RadiusV(InRadiusV)
+			, InverseMatrix(InMatrix.GetInverse())
+		{
+		}
+
+		FEllipseCurve(const FMatrixH& InMatrix, double InRadiusU, double InRadiusV, const FLinearBoundary& InBounds, int8 InDimension = 3)
+			: FCurve(InBounds, InDimension)
+			, Matrix(InMatrix)
+			, RadiusU(InRadiusU)
+			, RadiusV(InRadiusV)
+			, InverseMatrix(InMatrix.GetInverse())
+		{
+		}
+
+		FEllipseCurve() = default;
+
+	public:
+
+		virtual void Serialize(FCADKernelArchive& Ar) override
+		{
+			FCurve::Serialize(Ar);
+			Ar << Matrix;
+			Ar << RadiusU;
+			Ar << RadiusV;
+			Ar << InverseMatrix;
+		}
+
+#ifdef CADKERNEL_DEV
+		virtual FInfoEntity& GetInfo(FInfoEntity&) const override;
+#endif
+
+		virtual ECurve GetCurveType() const override
+		{
+			return ECurve::Ellipse;
+		}
+
+		const FMatrixH& GetMatrix() const 
+		{
+			return Matrix;
+		}
+
+		bool IsCircular() const
+		{
+			return FMath::IsNearlyZero(RadiusU - RadiusV);
+		}
+
+		virtual TSharedPtr<FEntityGeom> ApplyMatrix(const FMatrixH& InMatrix) const override;
+
+		virtual void EvaluatePoint(double Coordinate, FCurvePoint& OutPoint, int32 DerivativeOrder = 0) const override
+		{
+			ensure(Dimension == 3);
+			Evaluate<FCurvePoint, FPoint>(Coordinate, OutPoint, DerivativeOrder);
+		}
+
+		virtual void Evaluate2DPoint(double Coordinate, FCurvePoint2D& OutPoint, int32 DerivativeOrder = 0) const override
+		{
+			ensure(Dimension == 2);
+			Evaluate<FCurvePoint2D, FPoint2D>(Coordinate, OutPoint, DerivativeOrder);
+		}
+
+	private:
+
+		template <typename CurvePointType, typename PointType>
+		void Evaluate(double Coordinate, CurvePointType& OutPoint, int32 DerivativeOrder) const
+		{
+			OutPoint.DerivativeOrder = DerivativeOrder;
+
+			double CosU = cos(Coordinate);
+			double SinU = sin(Coordinate);
+
+			{
+				PointType Result(RadiusU * CosU, RadiusV * SinU);
+				OutPoint.Point = Matrix.Multiply(Result);
+			}
+
+			if (DerivativeOrder > 0)
+			{
+				PointType Result(RadiusU * (-SinU), RadiusV * CosU);
+				OutPoint.Gradient = Matrix.MultiplyVector(Result);
+			}
+
+			if (DerivativeOrder > 1)
+			{
+				PointType Result(RadiusU * (-CosU), RadiusV * (-SinU));
+				OutPoint.Laplacian = Matrix.MultiplyVector(Result);
+			}
+		}
+	};
+
+} // namespace CADKernel
+
diff --git a/Engine/Source/Runtime/Datasmith/CADKernel/Private/CADKernel/Geo/Curves/HyperbolaCurve.h b/Engine/Source/Runtime/Datasmith/CADKernel/Private/CADKernel/Geo/Curves/HyperbolaCurve.h
new file mode 100644
--- /dev/null
+++ b/Engine/Source/Runtime/Datasmith/CADKernel/Private/CADKernel/Geo/Curves/HyperbolaCurve.h
@@ -0,0 +1,98 @@
+// Copyright Epic Games, Inc. All Rights Reserved.
+#pragma once
+
+#include "CADKernel/Geo/Curves/Curve.h"
+
+namespace CADKernel
+{
+
+	class CADKERNEL_API FHyperbolaCurve : public FCurve
+	{
+		friend class FIGESEntity104;
+		friend class FEntity;
+		friend class FEntity;
+
+	protected:
+		FMatrixH Matrix;
+		double SemiMajorAxis;
+		double SemiImaginaryAxis;
+
+		FHyperbolaCurve(const FMatrixH& InMatrix, double InSemiAxis, double InSemiImagAxis, const FLinearBoundary& InBounds, int8 InDimension = 3)
+			: FCurve(InBounds, InDimension)
+			, Matrix(InMatrix)
+			, SemiMajorAxis(InSemiAxis)
+			, SemiImaginaryAxis(InSemiImagAxis)
+		{
+		}
+
+		FHyperbolaCurve() = default;
+
+	public:
+
+		virtual void Serialize(FCADKernelArchive& Ar) override
+		{
+			FCurve::Serialize(Ar);
+			Ar << Matrix;
+			Ar << SemiMajorAxis;
+			Ar << SemiImaginaryAxis;
+		}
+
+#ifdef CADKERNEL_DEV
+		virtual FInfoEntity& GetInfo(FInfoEntity&) const override;
+#endif
+
+		virtual ECurve GetCurveType() const override
+		{
+			return ECurve::Hyperbola;
+		}
+
+		FMatrixH& GetMatrix() 
+		{
+			return Matrix;
+		}
+
+		virtual TSharedPtr<FEntityGeom> ApplyMatrix(const FMatrixH& InMatrix) const override;
+
+		virtual void EvaluatePoint(double Coordinate, FCurvePoint& OutPoint, int32 DerivativeOrder = 0) const override
+		{
+			ensure(Dimension == 3);
+			Evaluate<FCurvePoint, FPoint>(Coordinate, OutPoint, DerivativeOrder);
+		}
+
+		virtual void Evaluate2DPoint(double Coordinate, FCurvePoint2D& OutPoint, int32 DerivativeOrder = 0) const override
+		{
+			ensure(Dimension == 2);
+			Evaluate<FCurvePoint2D, FPoint2D>(Coordinate, OutPoint, DerivativeOrder);
+		}
+
+	private:
+
+		template <typename CurvePointType, typename PointType>
+		void Evaluate(double Coordinate, CurvePointType& OutPoint, int32 DerivativeOrder) const
+		{
+			OutPoint.DerivativeOrder = DerivativeOrder;
+
+			const double CosHUCoord = cosh(Coordinate);
+			const double SinHUCoord = sinh(Coordinate);
+
+			{
+				PointType Result(SemiMajorAxis * CosHUCoord, SemiImaginaryAxis * SinHUCoord);
+				OutPoint.Point = Matrix.Multiply(Result);
+			}
+
+			if (DerivativeOrder > 0)
+			{
+				FPoint Result(SemiMajorAxis * SinHUCoord, SemiImaginaryAxis * CosHUCoord);
+				OutPoint.Gradient = Matrix.MultiplyVector(Result);
+			}
+
+			if (DerivativeOrder > 1)
+			{
+				FPoint Result(SemiMajorAxis * CosHUCoord,  SemiImaginaryAxis * SinHUCoord);
+				OutPoint.Laplacian = Matrix.MultiplyVector(Result);
+			}
+		}
+	};
+
+} // namespace CADKernel
+
diff --git a/Engine/Source/Runtime/Datasmith/CADKernel/Private/CADKernel/Geo/Curves/NURBSCurve.h b/Engine/Source/Runtime/Datasmith/CADKernel/Private/CADKernel/Geo/Curves/NURBSCurve.h
new file mode 100644
--- /dev/null
+++ b/Engine/Source/Runtime/Datasmith/CADKernel/Private/CADKernel/Geo/Curves/NURBSCurve.h
@@ -0,0 +1,141 @@
+// Copyright Epic Games, Inc. All Rights Reserved.
+#pragma once
+
+#include "CADKernel/Geo/Curves/Curve.h"
+#include "CADKernel/Math/BSpline.h"
+
+namespace CADKernel
+{
+	struct FNurbsCurveData
+	{
+		bool bIsRational;
+		int32 Dimension = 0;
+
+		int32 Degree = 0;
+		TArray<double> NodalVector;
+
+		TArray<double> Weights;
+		TArray<FPoint> Poles;
+	};
+
+	class CADKERNEL_API FNURBSCurve : public FCurve
+	{
+		friend class FEntity;
+
+	protected:
+
+		int32 Degree;
+
+		TArray<double> NodalVector;
+		TArray<double> Weights;
+
+		TArray<FPoint> Poles;
+
+		bool bIsRational;
+
+		/**
+		 * Data generated at initialization which are not serialized
+		 */
+		TArray<double> HomogeneousPoles;
+		int32 PoleDimension;
+
+		FNURBSCurve(int32 InDegree, const TArray<double>& InNodalVector, const TArray<FPoint>& InPoles, const TArray<double>& InWeights, int8 InDimension = 3);
+		FNURBSCurve(FNurbsCurveData& NurbsCurveData);
+		FNURBSCurve(const FNURBSCurve& Nurbs);
+
+		FNURBSCurve() = default;
+
+	public:
+
+		virtual void Serialize(FCADKernelArchive& Ar) override
+		{
+			FCurve::Serialize(Ar);
+			Ar << Degree;
+			Ar.Serialize(NodalVector);
+			Ar.Serialize(Weights);
+			Ar.Serialize(Poles);
+			Ar << bIsRational;
+
+			if (Ar.IsLoading())
+			{
+				Finalize();
+			}
+		}
+
+#ifdef CADKERNEL_DEV
+		virtual FInfoEntity& GetInfo(FInfoEntity&) const override;
+#endif
+
+		virtual ECurve GetCurveType() const override
+		{
+			return ECurve::Nurbs;
+		}
+
+		int32 GetDegree() const
+		{
+			return Degree;
+		}
+
+		const int32 GetPoleCount() const
+		{
+			return Poles.Num();
+		}
+
+		const TArray<FPoint>& GetPoles() const
+		{
+			return Poles;
+		}
+
+		const TArray<double>& GetWeights() const
+		{
+			return Weights;
+		}
+
+		TArray<double> GetHPoles() const
+		{
+			return HomogeneousPoles;
+		}
+		
+		const TArray<double>& GetNodalVector() const
+		{
+			return NodalVector;
+		}
+
+		bool IsRational() const
+		{
+			return bIsRational;
+		}
+
+		virtual TSharedPtr<FEntityGeom> ApplyMatrix(const FMatrixH& InMatrix) const override;
+
+		virtual void EvaluatePoint(double Coordinate, FCurvePoint& OutPoint, int32 DerivativeOrder = 0) const override
+		{
+			BSpline::EvaluatePoint(*this, Coordinate, OutPoint, DerivativeOrder);
+		}
+
+		virtual void Evaluate2DPoint(double Coordinate, FCurvePoint2D& OutPoint, int32 DerivativeOrder = 0) const override
+		{
+			BSpline::Evaluate2DPoint(*this, Coordinate, OutPoint, DerivativeOrder);
+		}
+
+		virtual void FindNotDerivableCoordinates(const FLinearBoundary& InBoundary, int32 DerivativeOrder, TArray<double>& OutNotDerivableCoordinates) const override
+		{
+			BSpline::FindNotDerivableParameters(*this, DerivativeOrder, InBoundary, OutNotDerivableCoordinates);
+		}
+
+		virtual void ExtendTo(const FPoint& DesiredPosition) override;
+
+		void Invert();
+		void SetStartNodalCoordinate(double NewStartBoundary);
+
+	private:
+
+		/**
+		 * Fill homogeneous points and set bounds
+		 */
+		void Finalize();
+
+	};
+
+} // namespace CADKernel
+
diff --git a/Engine/Source/Runtime/Datasmith/CADKernel/Private/CADKernel/Geo/Curves/ParabolaCurve.h b/Engine/Source/Runtime/Datasmith/CADKernel/Private/CADKernel/Geo/Curves/ParabolaCurve.h
new file mode 100644
--- /dev/null
+++ b/Engine/Source/Runtime/Datasmith/CADKernel/Private/CADKernel/Geo/Curves/ParabolaCurve.h
@@ -0,0 +1,91 @@
+// Copyright Epic Games, Inc. All Rights Reserved.
+#pragma once
+
+#include "CADKernel/Geo/Curves/Curve.h"
+
+namespace CADKernel
+{
+
+	class CADKERNEL_API FParabolaCurve : public FCurve
+	{
+		friend class FIGESEntity104;
+		friend class FEntity;
+
+	protected:
+		FMatrixH Matrix;
+		double FocalDistance;
+
+		FParabolaCurve(const FMatrixH& InMatrix, double InFocalDistance, const FLinearBoundary& InBounds, int8 InDimension = 3)
+			: FCurve(InBounds, InDimension)
+			, Matrix(InMatrix)
+			, FocalDistance(InFocalDistance)
+		{
+		}
+
+		FParabolaCurve() = default;
+
+	public:
+
+		virtual void Serialize(FCADKernelArchive& Ar) override
+		{
+			FCurve::Serialize(Ar);
+			Ar << Matrix;
+			Ar << FocalDistance;
+		}
+
+#ifdef CADKERNEL_DEV
+		virtual FInfoEntity& GetInfo(FInfoEntity&) const override;
+#endif
+
+		virtual ECurve GetCurveType() const override
+		{
+			return ECurve::Parabola;
+		}
+
+		virtual TSharedPtr<FEntityGeom> ApplyMatrix(const FMatrixH& InMatrix) const override;
+
+		virtual void ExtendTo(const FPoint& Point) override
+		{
+			ensureCADKernel(false);
+		}
+
+		virtual void EvaluatePoint(double Coordinate, FCurvePoint& OutPoint, int32 DerivativeOrder = 0) const override
+		{
+			ensure(Dimension == 3);
+			Evaluate<FCurvePoint, FPoint>(Coordinate, OutPoint, DerivativeOrder);
+		}
+
+		virtual void Evaluate2DPoint(double Coordinate, FCurvePoint2D& OutPoint, int32 DerivativeOrder = 0) const override
+		{
+			ensure(Dimension == 2);
+			Evaluate<FCurvePoint2D, FPoint2D>(Coordinate, OutPoint, DerivativeOrder);
+		}
+
+	private:
+
+		template <typename CurvePointType, typename PointType>
+		void Evaluate(double Coordinate, CurvePointType& OutPoint, int32 DerivativeOrder) const
+		{
+			OutPoint.DerivativeOrder = DerivativeOrder;
+
+			{
+				PointType Result(Coordinate * Coordinate * FocalDistance, Coordinate);
+				OutPoint.Point = Matrix.Multiply(Result);
+			}
+
+			if (DerivativeOrder > 0)
+			{
+				FPoint Result(2.0 * Coordinate * FocalDistance, 1);
+				OutPoint.Gradient = Matrix.MultiplyVector(Result);
+			}
+
+			if (DerivativeOrder > 1)
+			{
+				FPoint Result(2.0 * FocalDistance, 0.);
+				OutPoint.Laplacian = Matrix.MultiplyVector(Result);
+			}
+		}
+	};
+
+} // namespace CADKernel
+
diff --git a/Engine/Source/Runtime/Datasmith/CADKernel/Private/CADKernel/Geo/Curves/PolylineCurve.h b/Engine/Source/Runtime/Datasmith/CADKernel/Private/CADKernel/Geo/Curves/PolylineCurve.h
new file mode 100644
--- /dev/null
+++ b/Engine/Source/Runtime/Datasmith/CADKernel/Private/CADKernel/Geo/Curves/PolylineCurve.h
@@ -0,0 +1,217 @@
+// Copyright Epic Games, Inc. All Rights Reserved.
+#pragma once
+
+#include "CADKernel/Geo/Curves/Curve.h"
+#include "CADKernel/Geo/Sampling/Polyline.h"
+#include "CADKernel/Geo/Sampling/PolylineTools.h"
+
+namespace CADKernel
+{
+	class FInfoEntity;
+	class FSurface;
+
+	template<class PointType, class PointCurveType>
+	class CADKERNEL_API TPolylineCurve : public FCurve, public TPolyline<PointType>
+	{
+		friend class FEntity;
+		friend class FPolylineTools;
+
+	protected:
+		
+		TPolylineApproximator<PointType> Approximator;
+
+		TPolylineCurve(const TArray<PointType>& InPoints, const TArray<double>& InCoordinates, int8 InDimension)
+			: FCurve(InDimension)
+			, Approximator(this->Coordinates, this->Points)
+		{
+			this->Coordinates = InCoordinates;
+			this->Points = InPoints;
+			ensureCADKernel(this->Coordinates[0] < this->Coordinates.Last());
+			Boundary.Set(this->Coordinates.HeapTop(), this->Coordinates.Last());
+		}
+
+		TPolylineCurve(const TArray<PointType>& InPoints, int8 InDimension)
+			: FCurve(InDimension)
+			, TPolyline<PointType>(InPoints)
+			, Approximator(this->Coordinates, this->Points)
+		{
+			this->Coordinates.Reserve(InPoints.Num());
+			this->Coordinates.Add(0.);
+
+			double CurvilineLength = 0;
+			for (int32 iPoint = 1; iPoint < this->Points.Num(); iPoint++)
+			{
+				CurvilineLength += this->Points[iPoint].Distance(this->Points[iPoint - 1]);
+				this->Coordinates.Add(CurvilineLength);
+			}
+
+			Boundary.Set(0, CurvilineLength);
+		}
+
+		TPolylineCurve()
+			: Approximator(this->Coordinates, this->Points)
+		{
+		}
+
+	public:
+
+		virtual void Serialize(FCADKernelArchive& Ar) override
+		{
+			FCurve::Serialize(Ar);
+			TPolyline<PointType>::Serialize(Ar);
+		}
+
+		virtual void EvaluateCurvesPoint(double InCoordinate, PointCurveType& OutPoint, int32 InDerivativeOrder = 0) const
+		{
+			Approximator.ApproximatePoint(InCoordinate, OutPoint, InDerivativeOrder);
+		}
+
+		PointType EvaluatePointAt(double InCoordinate) const
+		{
+			return Approximator.ApproximatePoint(InCoordinate);
+		}
+
+		virtual void EvaluateCurvesPoints(const TArray<double>& InCoordinates, TArray<PointCurveType>& OutPoints, int32 InDerivativeOrder = 0) const
+		{
+			Approximator.ApproximatePoints(InCoordinates, OutPoints, InDerivativeOrder);
+		}
+
+		virtual double ComputeSubLength(const FLinearBoundary& InBoundary) const
+		{
+			return Approximator.ComputeLengthOfSubPolyline(InBoundary);
+		}
+
+		const TArray<PointType>& GetPolylinePoints() const
+		{
+			return this->Points;
+		}
+
+		const TArray<double>& GetPolylineParameters() const
+		{
+			return this->Coordinates;
+		}
+
+		virtual void FindNotDerivableCoordinates(const FLinearBoundary& InBoundary, int32 DerivativeOrder, TArray<double>& OutNotDerivableCoordinates) const override
+		{
+			//TO DO
+			ensureCADKernel(false);
+		}
+
+		void SetPoints(const TArray<PointType>& InPoints)
+		{
+			this->Points = InPoints;
+			GlobalLength.Empty();
+		}
+
+		template<class PolylineType>
+		TSharedPtr<FEntityGeom> ApplyMatrixImpl(const FMatrixH& InMatrix) const
+		{
+			TArray<PointType> NewPoints;
+			NewPoints.Reserve(this->Points.Num());
+
+			for (PointType Point : this->Points)
+			{
+				NewPoints.Emplace(InMatrix.Multiply(Point));
+			}
+
+			return FEntity::MakeShared<PolylineType>(NewPoints, this->Coordinates);
+		}
+
+		virtual void ExtendTo(const FPoint& DesiredPoint) override
+		{
+			PolylineTools::ExtendTo(this->Points, (PointType) DesiredPoint);
+		}
+	};
+
+	class CADKERNEL_API FPolylineCurve : public TPolylineCurve<FPoint, FCurvePoint>
+	{
+		friend class FEntity;
+
+	protected:
+		FPolylineCurve(const TArray<FPoint>& InPoints, const TArray<double>& InCoordinates)
+			: TPolylineCurve<FPoint, FCurvePoint>(InPoints, InCoordinates, 3)
+		{
+		}
+
+		FPolylineCurve(const TArray<FPoint>& InPoints)
+			: TPolylineCurve<FPoint, FCurvePoint>(InPoints, 3)
+		{
+		}
+
+		FPolylineCurve() = default;
+
+	public:
+
+		virtual ECurve GetCurveType() const override
+		{
+			return ECurve::Polyline3D;
+		}
+
+		virtual TSharedPtr<FEntityGeom> ApplyMatrix(const FMatrixH& InMatrix) const override
+		{
+			return ApplyMatrixImpl<FPolylineCurve>(InMatrix);
+		}
+
+		virtual FPoint EvaluatePoint(double InCoordinate) const override
+		{
+			return EvaluatePointAt(InCoordinate);
+		} 
+
+#ifdef CADKERNEL_DEV
+		virtual FInfoEntity& GetInfo(FInfoEntity& Info) const override;
+#endif
+
+	};
+
+	class CADKERNEL_API FPolyline2DCurve : public TPolylineCurve<FPoint2D, FCurvePoint2D>
+	{
+		friend class FEntity;
+
+	protected:
+		FPolyline2DCurve(const TArray<FPoint2D>& InPoints, const TArray<double>& InCoordinates)
+			: TPolylineCurve<FPoint2D, FCurvePoint2D>(InPoints, InCoordinates, 3)
+		{
+		}
+
+		FPolyline2DCurve(const TArray<FPoint2D>& InPoints)
+			: TPolylineCurve<FPoint2D, FCurvePoint2D>(InPoints, 3)
+		{
+		}
+
+		FPolyline2DCurve() = default;
+
+	public:
+
+		virtual ECurve GetCurveType() const override
+		{
+			return ECurve::Polyline3D;
+		}
+
+		virtual TSharedPtr<FEntityGeom> ApplyMatrix(const FMatrixH& InMatrix) const override
+		{
+			return ApplyMatrixImpl<FPolyline2DCurve>(InMatrix);
+		}
+
+		virtual void Evaluate2DPoint(double InCoordinate, FCurvePoint2D& OutPoint, int32 InDerivativeOrder = 0) const override
+		{
+			EvaluateCurvesPoint(InCoordinate, OutPoint, InDerivativeOrder);
+		}
+
+		virtual FPoint2D Evaluate2DPoint(double InCoordinate) const override
+		{
+			return EvaluatePointAt(InCoordinate);
+		}
+
+		virtual void Evaluate2DPoints(const TArray<double>& InCoordinates, TArray<FCurvePoint2D>& OutPoints, int32 InDerivativeOrder = 0) const override
+		{
+			EvaluateCurvesPoints(InCoordinates, OutPoints, InDerivativeOrder);
+		}
+
+#ifdef CADKERNEL_DEV
+		virtual FInfoEntity& GetInfo(FInfoEntity& Info) const override;
+#endif
+
+	}; 
+	
+} // namespace CADKernel
+
diff --git a/Engine/Source/Runtime/Datasmith/CADKernel/Private/CADKernel/Geo/Curves/SegmentCurve.h b/Engine/Source/Runtime/Datasmith/CADKernel/Private/CADKernel/Geo/Curves/SegmentCurve.h
new file mode 100644
--- /dev/null
+++ b/Engine/Source/Runtime/Datasmith/CADKernel/Private/CADKernel/Geo/Curves/SegmentCurve.h
@@ -0,0 +1,101 @@
+// Copyright Epic Games, Inc. All Rights Reserved.
+#pragma once
+
+#include "CADKernel/Geo/Curves/Curve.h"
+#include "CADKernel/Geo/GeoEnum.h"
+
+namespace CADKernel
+{
+
+	class CADKERNEL_API FSegmentCurve : public FCurve
+	{
+		friend class FEntity;
+
+	protected:
+		FPoint StartPoint;
+		FPoint EndPoint;
+
+		FSegmentCurve(const FPoint& InStartPoint, const FPoint& InEndPoint, int8 InDimension = 3)
+			: FCurve(InDimension)
+			, StartPoint(InStartPoint)
+			, EndPoint(InEndPoint)
+		{
+		}
+
+		FSegmentCurve() = default;
+
+	public:
+
+		virtual void Serialize(FCADKernelArchive& Ar) override
+		{
+			FCurve::Serialize(Ar);
+			Ar << StartPoint;
+			Ar << EndPoint;
+		}
+
+#ifdef CADKERNEL_DEV
+		virtual FInfoEntity& GetInfo(FInfoEntity&) const override;
+#endif
+
+		virtual ECurve GetCurveType() const override
+		{
+			return ECurve::Segment;
+		}
+
+		virtual TSharedPtr<FEntityGeom> ApplyMatrix(const FMatrixH& InMatrix) const override;
+
+		virtual void EvaluatePoint(double Coordinate, FCurvePoint& OutPoint, int32 DerivativeOrder = 0) const override
+		{
+			ensure(Dimension == 3);
+			Evaluate<FCurvePoint, FPoint>(Coordinate, OutPoint, DerivativeOrder);
+		}
+
+		virtual void Evaluate2DPoint(double Coordinate, FCurvePoint2D& OutPoint, int32 DerivativeOrder = 0) const override
+		{
+			ensure(Dimension == 2);
+			Evaluate<FCurvePoint2D, FPoint2D>(Coordinate, OutPoint, DerivativeOrder);
+		}
+
+		const FPoint& GetStartPoint() const
+		{
+			return StartPoint;
+		}
+
+		const FPoint& GetEndPoint() const
+		{
+			return EndPoint;
+		}
+
+		virtual void ExtendTo(const FPoint& DesiredPosition) override
+		{
+			double DistanceToStartPoint = DesiredPosition.SquareDistance(StartPoint);
+			double DistanceToEndPoint = DesiredPosition.SquareDistance(EndPoint);
+			if (DistanceToEndPoint < DistanceToStartPoint)
+			{
+				EndPoint = DesiredPosition;
+			}
+			else
+			{
+				StartPoint = DesiredPosition;
+			}
+		}
+
+	private:
+		template <typename CurvePointType, typename PointType>
+		void Evaluate(double Coordinate, CurvePointType& OutPoint, int32 DerivativeOrder) const
+		{
+			OutPoint.DerivativeOrder = DerivativeOrder;
+
+			PointType Tangent = (EndPoint - StartPoint);
+			OutPoint.Point = Tangent * Coordinate + StartPoint;
+
+			if (DerivativeOrder > 0)
+			{
+				OutPoint.Gradient = Tangent;
+			}
+		}
+
+	};
+
+} // namespace CADKernel
+
diff --git a/Engine/Source/Runtime/Datasmith/CADKernel/Private/CADKernel/Geo/Sampling/SurfacicPolyline.cpp b/Engine/Source/Runtime/Datasmith/CADKernel/Private/CADKernel/Geo/Sampling/SurfacicPolyline.cpp
--- a/Engine/Source/Runtime/Datasmith/CADKernel/Private/CADKernel/Geo/Sampling/SurfacicPolyline.cpp
+++ b/Engine/Source/Runtime/Datasmith/CADKernel/Private/CADKernel/Geo/Sampling/SurfacicPolyline.cpp
@@ -140,8 +140,8 @@ void FSurfacicPolyline::ComputeIntersectionsWithIsos(const FLinearBoundary& InBo
 		}
 
 		GetMinMax(Points2D[Index][TypeIso], Points2D[Index + 1][TypeIso], UMin, UMax);
-		UMin -= SMALL_NUMBER;
-		UMax += SMALL_NUMBER;
+		UMin -= DOUBLE_SMALL_NUMBER;
+		UMax += DOUBLE_SMALL_NUMBER;
 
 		if (IsoCoordinate < UMin || IsoCoordinate > UMax)
 		{
diff --git a/Engine/Source/Runtime/Datasmith/CADKernel/Private/CADKernel/Geo/Surfaces/BezierSurface.h b/Engine/Source/Runtime/Datasmith/CADKernel/Private/CADKernel/Geo/Surfaces/BezierSurface.h
new file mode 100644
--- /dev/null
+++ b/Engine/Source/Runtime/Datasmith/CADKernel/Private/CADKernel/Geo/Surfaces/BezierSurface.h
@@ -0,0 +1,82 @@
+// Copyright Epic Games, Inc. All Rights Reserved.
+#pragma once
+
+#include "CADKernel/Core/CADKernelArchive.h"
+#include "CADKernel/Geo/Surfaces/Surface.h"
+
+namespace CADKernel
+{
+
+	class CADKERNEL_API FBezierSurface : public FSurface
+	{
+		friend FEntity;
+
+	protected:
+
+		int32 UDegre;
+		int32 VDegre;
+
+		int32 UPoleNum;
+		int32 VPoleNum;
+
+		TArray<FPoint> Poles;
+
+		FBezierSurface(const double InToleranceGeometric, int32 InUDegre, int32 InVDegre, const TArray<FPoint>& InPoles)
+			: FSurface(InToleranceGeometric)
+			, UDegre(InUDegre)
+			, VDegre(InVDegre)
+			, UPoleNum(InUDegre + 1)
+			, VPoleNum(InVDegre + 1)
+			, Poles(InPoles)
+		{
+			ensureCADKernel((UDegre + 1) * (VDegre + 1) == Poles.Num());
+			ComputeDefaultMinToleranceIso();
+		}
+
+		FBezierSurface() = default;
+
+	public:
+
+		virtual void Serialize(FCADKernelArchive& Ar) override
+		{
+			FSurface::Serialize(Ar);
+			Ar << UDegre;
+			Ar << VDegre;
+			Ar << UPoleNum;
+			Ar << VPoleNum;
+			Ar.Serialize(Poles);
+		}
+
+		ESurface GetSurfaceType() const
+		{
+			return ESurface::Bezier;
+		}
+
+		int32 GetUDegre() const
+		{
+			return UDegre;
+		}
+
+		int32 GetVDegre() const
+		{
+			return VDegre;
+		}
+
+		virtual const TArray<FPoint>& GetPoles() const
+		{
+			return Poles;
+		}
+
+		virtual TSharedPtr<FEntityGeom> ApplyMatrix(const FMatrixH& InMatrix) const override;
+
+		virtual void EvaluatePoint(const FPoint2D& InSurfacicCoordinate, FSurfacicPoint& OutPoint3D, int32 InDerivativeOrder = 0) const override;
+
+		virtual void Presample(const FSurfacicBoundary& InBoundaries, FCoordinateGrid& OutCoordinates) override;
+
+#ifdef CADKERNEL_DEV
+		virtual FInfoEntity& GetInfo(FInfoEntity&) const override;
+#endif
+	};
+
+} // namespace CADKernel
+
diff --git a/Engine/Source/Runtime/Datasmith/CADKernel/Private/CADKernel/Geo/Surfaces/CompositeSurface.h b/Engine/Source/Runtime/Datasmith/CADKernel/Private/CADKernel/Geo/Surfaces/CompositeSurface.h
new file mode 100644
--- /dev/null
+++ b/Engine/Source/Runtime/Datasmith/CADKernel/Private/CADKernel/Geo/Surfaces/CompositeSurface.h
@@ -0,0 +1,91 @@
+// Copyright Epic Games, Inc. All Rights Reserved.
+#pragma once
+
+#include "CADKernel/Core/CADKernelArchive.h"
+#include "CADKernel/Geo/Surfaces/Surface.h"
+
+namespace CADKernel
+{
+
+	class CADKERNEL_API FCompositeSurface : public FSurface
+	{
+		friend FEntity;
+
+	protected:
+
+		/**
+		 * U & V Coordinates of the composite surface's parametric space at each change of elementary surface
+		 */
+		TArray<double> GlobalCoordinates[2];
+
+		/**
+		 * Elementary surfaces
+		 */
+		TArray<TSharedPtr<FSurface>> Surfaces;
+
+		/**
+		 * Native UV Loops of each surface composing the composite surface
+		 */
+		TArray<FSurfacicBoundary> NativeUVBoundaries;
+
+		FCompositeSurface(const double InToleranceGeometric, int32 USurfaceNum, int32 VSurfaceNum, const TArray<double>& UCoordinates, const TArray<double>& VCoordinates, const TArray<TSharedPtr<FSurface>>& Surfaces);
+
+		FCompositeSurface() = default;
+
+	public:
+
+		virtual void Serialize(FCADKernelArchive& Ar) override
+		{
+			FSurface::Serialize(Ar);
+			Ar << GlobalCoordinates[EIso::IsoU];
+			Ar << GlobalCoordinates[EIso::IsoV];
+			SerializeIdents(Ar, Surfaces);
+			Ar << NativeUVBoundaries;
+		}
+
+		virtual void SpawnIdent(FDatabase& Database) override
+		{
+			if (!FEntity::SetId(Database))
+			{
+				return;
+			}
+
+			SpawnIdentOnEntities(Surfaces, Database);
+		}
+
+		virtual void ResetMarkersRecursively() override
+		{
+			ResetMarkers();
+			ResetMarkersRecursivelyOnEntities(Surfaces);
+		}
+
+		ESurface GetSurfaceType() const
+		{
+			return ESurface::Composite;
+		}
+
+#ifdef CADKERNEL_DEV
+		virtual FInfoEntity& GetInfo(FInfoEntity&) const override;
+#endif
+
+		virtual TSharedPtr<FEntityGeom> ApplyMatrix(const FMatrixH& InMatrix) const override;
+
+		int32 GetSurfNum(EIso Iso) const
+		{
+			return GlobalCoordinates[Iso].Num() - 1;
+		}
+
+		virtual void InitBoundary() override;
+
+		virtual void LinesNotDerivables(const FSurfacicBoundary& Bounds, int32 InDerivativeOrder, FCoordinateGrid& OutNotDerivableCoordinates) const override;
+
+		virtual void EvaluatePoint(const FPoint2D& InSurfacicCoordinate, FSurfacicPoint& OutPoint3D, int32 InDerivativeOrder = 0) const override;
+
+		virtual void Presample(const FSurfacicBoundary& InBoundaries, FCoordinateGrid& OutCoordinates) override;
+
+		double LocalToGlobalCoordinate(int32 SurfaceIndex, double CoordinateU, EIso Iso) const;
+	};
+
+} // namespace CADKernel
+
+
diff --git a/Engine/Source/Runtime/Datasmith/CADKernel/Private/CADKernel/Geo/Surfaces/ConeSurface.h b/Engine/Source/Runtime/Datasmith/CADKernel/Private/CADKernel/Geo/Surfaces/ConeSurface.h
new file mode 100644
--- /dev/null
+++ b/Engine/Source/Runtime/Datasmith/CADKernel/Private/CADKernel/Geo/Surfaces/ConeSurface.h
@@ -0,0 +1,118 @@
+// Copyright Epic Games, Inc. All Rights Reserved.
+#pragma once
+
+#include "CADKernel/Geo/Surfaces/Surface.h"
+
+namespace CADKernel
+{
+	class CADKERNEL_API FConeSurface : public FSurface
+	{
+		friend FEntity;
+
+	protected:
+		FMatrixH Matrix;
+		double   StartRadius;
+		double   ConeAngle;
+
+		/**
+		 * A Cone surface is defined by a point on the axis of the cone, the direction of the axis of the cone, 
+		 * the radius of the cone at the axis point and the cone semi-angle. 
+		 * 
+		 * The local coordinate system is defined with the origin at the axis point and the Z axis in the axis direction, 
+		 * The equation of the surface in this system is x2 + y2 – (StartRadius + z*tan(ConeAngle) )2 = 0
+		 * where ConeAngle is the cone semi-angle and StartRadius is the given cone radius. 
+		 * 
+		 * The cone is placed at its final position and orientation by the Matrix
+		 */
+		FConeSurface(const double InToleranceGeometric, const FMatrixH& InMatrix, double InStartRadius, double InConeAngle, double InStartRuleLength = -HUGE_VALUE, double InEndRuleLength = HUGE_VALUE, double InStartAngle = 0.0, double InEndAngle = DOUBLE_TWO_PI)
+			: FConeSurface(InToleranceGeometric, InMatrix, InStartRadius, InConeAngle, FSurfacicBoundary(InStartAngle, InEndAngle, InStartRuleLength, InEndRuleLength))
+		{
+		}
+
+		/**
+		 * A Cone surface is defined by a point on the axis of the cone, the direction of the axis of the cone,
+		 * the radius of the cone at the axis point and the cone semi-angle.
+		 *
+		 * The local coordinate system is defined with the origin at the axis point and the Z axis in the axis direction,
+		 * The equation of the surface in this system is x2 + y2 – (StartRadius + z*tan(ConeAngle) )2 = 0
+		 * where ConeAngle is the cone semi-angle and StartRadius is the given cone radius.
+		 *
+		 * The cone is placed at its final position and orientation by the Matrix
+		 *
+		 * The bounds of the cone are defined as follow:
+		 * Bounds[EIso::IsoU].Min = StartRuleLength;
+		 * Bounds[EIso::IsoU].Max = EndRuleLength;
+		 * Bounds[EIso::IsoV].Min = StartAngle;
+		 * Bounds[EIso::IsoV].Max = EndAngle;
+		 */
+		FConeSurface(const double InToleranceGeometric, const FMatrixH& InMatrix, double InStartRadius, double InConeAngle, const FSurfacicBoundary& InBoundary)
+			: FSurface(InToleranceGeometric, InBoundary)
+			, Matrix(InMatrix)
+			, StartRadius(InStartRadius)
+			, ConeAngle(InConeAngle)
+		{
+			ComputeMinToleranceIso();
+		}
+
+		FConeSurface() = default;
+
+		void ComputeMinToleranceIso()
+		{
+			FPoint Origin = Matrix.Multiply(FPoint::ZeroPoint);
+
+			double DeltaVR = tan(ConeAngle);
+			double Radius1 = FMath::Abs(StartRadius + Boundary[EIso::IsoV].Max * DeltaVR);
+			double Radius2 = FMath::Abs(StartRadius + Boundary[EIso::IsoV].Min * DeltaVR);
+			double Radius = FMath::Max(Radius2, Radius1);
+
+			FPoint Point2DU{ 1 , 0, 0 };
+			double ScaleU = ComputeScaleAlongAxis(Point2DU, Matrix, Origin);
+			double ToleranceU = Tolerance3D / Radius;
+			ToleranceU /= ScaleU;
+
+			FPoint Point2DV{ 0, 1, 0 };
+			double ScaleV = ComputeScaleAlongAxis(Point2DV, Matrix, Origin);
+			double ToleranceV = Tolerance3D / FMath::Sqrt(FMath::Square(DeltaVR) + 1);
+			ToleranceV /= ScaleV;
+
+			MinToleranceIso.Set(ToleranceU, ToleranceV);
+		}
+
+	public:
+
+		virtual void Serialize(FCADKernelArchive& Ar) override
+		{
+			FSurface::Serialize(Ar);
+			Ar << Matrix;
+			Ar << StartRadius;
+			Ar << ConeAngle;
+		}
+
+
+		ESurface GetSurfaceType() const
+		{
+			return ESurface::Cone;
+		}
+
+#ifdef CADKERNEL_DEV
+		virtual FInfoEntity& GetInfo(FInfoEntity&) const override;
+#endif
+
+		virtual TSharedPtr<FEntityGeom> ApplyMatrix(const FMatrixH& InMatrix) const override;
+
+		virtual void EvaluatePoint(const FPoint2D& InSurfacicCoordinate, FSurfacicPoint& OutPoint3D, int32 InDerivativeOrder = 0) const override;
+		virtual void EvaluatePointGrid(const FCoordinateGrid& Coordinates, FSurfacicSampling& OutPoints, bool bComputeNormals = false) const override;
+
+		virtual void Presample(const FSurfacicBoundary& InBoundaries, FCoordinateGrid& OutCoordinates) override 
+		{
+			PresampleIsoCircle(InBoundaries, OutCoordinates, EIso::IsoU);
+
+			OutCoordinates[EIso::IsoV].Empty(3);
+			OutCoordinates[EIso::IsoV].Add(InBoundaries[EIso::IsoV].Min);
+			OutCoordinates[EIso::IsoV].Add((InBoundaries[EIso::IsoV].Max + InBoundaries[EIso::IsoV].Min) / 2.0);
+			OutCoordinates[EIso::IsoV].Add(InBoundaries[EIso::IsoV].Max);
+		}
+
+	};
+}
+
diff --git a/Engine/Source/Runtime/Datasmith/CADKernel/Private/CADKernel/Geo/Surfaces/CoonsSurface.h b/Engine/Source/Runtime/Datasmith/CADKernel/Private/CADKernel/Geo/Surfaces/CoonsSurface.h
new file mode 100644
--- /dev/null
+++ b/Engine/Source/Runtime/Datasmith/CADKernel/Private/CADKernel/Geo/Surfaces/CoonsSurface.h
@@ -0,0 +1,54 @@
+// Copyright Epic Games, Inc. All Rights Reserved.
+#pragma once
+
+#include "CADKernel/Core/CADKernelArchive.h"
+#include "CADKernel/Geo/Surfaces/Surface.h"
+
+namespace CADKernel
+{
+	class FCurve;
+	class FNURBSSurface;
+
+	class CADKERNEL_API FCoonsSurface : public FSurface
+	{
+		friend FEntity;
+
+	protected:
+		TSharedPtr<FCurve> Curves[4];
+		TArray<FPoint> Corners;
+
+		FCoonsSurface(const double InToleranceGeometric, TSharedPtr<FCurve> InCurve1, TSharedPtr<FCurve> InCurve2, TSharedPtr<FCurve> InCurve3, TSharedPtr<FCurve> InCurve4);
+		FCoonsSurface(const double InToleranceGeometric, TSharedPtr<FCurve> InCurves[4]);
+		FCoonsSurface() = default;
+
+	public:
+
+		virtual void Serialize(FCADKernelArchive& Ar) override
+		{
+			FSurface::Serialize(Ar);
+			SerializeIdent(Ar, Curves[0]);
+			SerializeIdent(Ar, Curves[1]);
+			SerializeIdent(Ar, Curves[2]);
+			SerializeIdent(Ar, Curves[3]);
+			Ar << Corners;
+		}
+
+
+#ifdef CADKERNEL_DEV
+		virtual FInfoEntity& GetInfo(FInfoEntity&) const override;
+#endif
+
+		virtual ESurface GetSurfaceType() const override
+		{
+			return ESurface::Coons;
+		}
+
+		virtual TSharedPtr<FEntityGeom> ApplyMatrix(const FMatrixH& InMatrix) const override;
+
+		virtual void EvaluatePoint(const FPoint2D& InSurfacicCoordinate, FSurfacicPoint& OutPoint3D, int32 InDerivativeOrder = 0) const override;
+
+		virtual void LinesNotDerivables(const FSurfacicBoundary& Bounds, int32 InDerivativeOrder, FCoordinateGrid& OutNotDerivableCoordinates) const override;
+	};
+
+} // namespace CADKernel
+
diff --git a/Engine/Source/Runtime/Datasmith/CADKernel/Private/CADKernel/Geo/Surfaces/CylinderSurface.h b/Engine/Source/Runtime/Datasmith/CADKernel/Private/CADKernel/Geo/Surfaces/CylinderSurface.h
new file mode 100644
--- /dev/null
+++ b/Engine/Source/Runtime/Datasmith/CADKernel/Private/CADKernel/Geo/Surfaces/CylinderSurface.h
@@ -0,0 +1,87 @@
+// Copyright Epic Games, Inc. All Rights Reserved.
+#pragma once
+
+#include "CADKernel/Geo/Surfaces/Surface.h"
+
+namespace CADKernel
+{
+	class FNURBSSurface;
+
+	class CADKERNEL_API FCylinderSurface : public FSurface
+	{
+		friend FEntity;
+
+	protected:
+		FMatrixH Matrix;
+		double Radius;
+
+		/**
+		 * The Cylinder surface is a right circular cylinder surface starting at the axis origin and with z axis as the axis direction. 
+		 * It is defined by its Radius and its bounds (StartAngle, EndAngle, StartLength, EndLength)
+		 * 
+		 * The surface is placed at its final position and orientation by the Matrix
+		 *
+		 * The bounds of the cylinder are defined as follow:
+		 * Bounds[EIso::IsoU].Min = StartAngle;
+		 * Bounds[EIso::IsoU].Max = EndAngle;
+		 * Bounds[EIso::IsoV].Min = StartLength;
+		 * Bounds[EIso::IsoV].Max = EndLength;
+		 */
+		FCylinderSurface(const double InToleranceGeometric, const FMatrixH& InMatrix, double InRadius, double InStartLength = -HUGE_VALUE, double InEndLength = HUGE_VALUE, double InStartAngle = 0.0, double InEndAngle = DOUBLE_TWO_PI);
+		FCylinderSurface(const double InToleranceGeometric, const FMatrixH& InMatrix, const double InRadius, const FSurfacicBoundary& Boundary);
+
+		FCylinderSurface() = default;
+
+		void ComputeMinToleranceIso()
+		{
+			FPoint Origin = Matrix.Multiply(FPoint::ZeroPoint);
+
+			FPoint Point2DU { 1 , 0, 0 };
+			double ToleranceU = Tolerance3D / Radius;
+
+			double ScaleU = ComputeScaleAlongAxis(Point2DU, Matrix, Origin);
+			ToleranceU /= ScaleU;
+
+			FPoint Point2DV{ 0, 1, 0 };
+			double ToleranceV = Tolerance3D / ComputeScaleAlongAxis(Point2DV, Matrix, Origin);
+
+			MinToleranceIso.Set(ToleranceU, ToleranceV);
+		}
+
+	public:
+
+		virtual void Serialize(FCADKernelArchive& Ar) override
+		{
+			FSurface::Serialize(Ar);
+			Ar << Matrix;
+			Ar << Radius;
+		}
+
+		virtual TSharedPtr<FEntityGeom> ApplyMatrix(const FMatrixH& InMatrix) const override;
+
+#ifdef CADKERNEL_DEV
+		virtual FInfoEntity& GetInfo(FInfoEntity&) const override;
+#endif
+
+		ESurface GetSurfaceType() const
+		{
+			return ESurface::Cylinder;
+		}
+
+		virtual void InitBoundary() override;
+
+		virtual void EvaluatePoint(const FPoint2D& InSurfacicCoordinate, FSurfacicPoint& OutPoint3D, int32 InDerivativeOrder = 0) const override;
+
+		virtual void Presample(const FSurfacicBoundary& InBoundaries, FCoordinateGrid& OutCoordinates) override
+		{
+			PresampleIsoCircle(InBoundaries, OutCoordinates, EIso::IsoU);
+
+			OutCoordinates[EIso::IsoV].Empty(3);
+			OutCoordinates[EIso::IsoV].Add(InBoundaries[EIso::IsoV].Min);
+			OutCoordinates[EIso::IsoV].Add((InBoundaries[EIso::IsoV].Max + InBoundaries[EIso::IsoV].Min) / 2.0);
+			OutCoordinates[EIso::IsoV].Add(InBoundaries[EIso::IsoV].Max);
+		}
+
+	};
+}
+
diff --git a/Engine/Source/Runtime/Datasmith/CADKernel/Private/CADKernel/Geo/Surfaces/NURBSSurface.h b/Engine/Source/Runtime/Datasmith/CADKernel/Private/CADKernel/Geo/Surfaces/NURBSSurface.h
new file mode 100644
--- /dev/null
+++ b/Engine/Source/Runtime/Datasmith/CADKernel/Private/CADKernel/Geo/Surfaces/NURBSSurface.h
@@ -0,0 +1,221 @@
+// Copyright Epic Games, Inc. All Rights Reserved.
+#pragma once
+
+#include "CADKernel/Geo/Surfaces/NurbsSurfaceData.h"
+#include "CADKernel/Geo/Surfaces/Surface.h"
+#include "CADKernel/Math/BSpline.h"
+
+namespace CADKernel
+{
+	class FNURBSSurface : public FSurface
+	{
+		friend FEntity;
+
+	protected:
+		int32 PoleUCount;
+		int32 PoleVCount;
+
+		int32 UDegree;
+		int32 VDegree;
+
+		TArray<double> UNodalVector;
+		TArray<double> VNodalVector;
+
+		TArray<double> Weights;
+
+		TArray<FPoint> Poles;
+
+		bool bIsRational;
+
+		/**
+		 * Data generated at initialization which are not serialized
+		 */
+		TArray<double> HomogeneousPoles;
+
+		/**
+		 * Build a Non uniform B-Spline surface
+		 * @param NodalVectorU Its size is the number of poles in U + the surface degree in U + 1 (PoleUNum + UDegre + 1)
+		 * @param NodalVectorV Its size is the number of poles in V + the surface degree in V + 1 (PoleVNum + VDegre + 1)
+		 */
+		FNURBSSurface(const double InToleranceGeometric, int32 InPoleUCount, int32 InPoleVCount, int32 InDegreU, int32 InDegreV, const TArray<double>& InNodalVectorU, const TArray<double>& InNodalVectorV, const TArray<FPoint>& InPoles)
+			: FSurface(InToleranceGeometric)
+			, PoleUCount(InPoleUCount)
+			, PoleVCount(InPoleVCount)
+			, UDegree(InDegreU)
+			, VDegree(InDegreV)
+			, UNodalVector(InNodalVectorU)
+			, VNodalVector(InNodalVectorV)
+			, Poles(InPoles)
+			, bIsRational(false)
+		{
+			Finalize();
+		}
+
+		/**
+		 * Build a Non uniform rational B-Spline surface
+		 * @param NodalVectorU Its size is the number of poles in U + the surface degree in U + 1 (PoleUNum + UDegre + 1)
+		 * @param NodalVectorV Its size is the number of poles in V + the surface degree in V + 1 (PoleVNum + VDegre + 1)
+		 */
+		FNURBSSurface(const double InToleranceGeometric, int32 InPoleUCount, int32 InPoleVCount, int32 InDegreU, int32 InDegreV, const TArray<double>& InNodalVectorU, const TArray<double>& InNodalVectorV, const TArray<FPoint>& InPoles, const TArray<double>& InWeights)
+			: FSurface(InToleranceGeometric)
+			, PoleUCount(InPoleUCount)
+			, PoleVCount(InPoleVCount)
+			, UDegree(InDegreU)
+			, VDegree(InDegreV)
+			, UNodalVector(InNodalVectorU)
+			, VNodalVector(InNodalVectorV)
+			, Weights(InWeights)
+			, Poles(InPoles)
+			, bIsRational(true)
+		{
+			ComputeMinToleranceIso();
+			Finalize();
+		}
+
+		/**
+		 * Build a Non uniform B-Spline surface
+		 * @param NodalVectorU Its size is the number of poles in U + the surface degree in U + 1 (PoleUNum + UDegre + 1)
+		 * @param NodalVectorV Its size is the number of poles in V + the surface degree in V + 1 (PoleVNum + VDegre + 1)
+		 */
+		FNURBSSurface(const double InToleranceGeometric, FNurbsSurfaceData NurbsData)
+			: FSurface(InToleranceGeometric)
+			, PoleUCount(NurbsData.PoleUCount)
+			, PoleVCount(NurbsData.PoleVCount)
+			, UDegree(NurbsData.UDegree)
+			, VDegree(NurbsData.VDegree)
+			, UNodalVector(NurbsData.UNodalVector)
+			, VNodalVector(NurbsData.VNodalVector)
+			, Weights(NurbsData.Weights)
+			, Poles(NurbsData.Poles)
+			, bIsRational(true)
+		{
+			ComputeMinToleranceIso();
+			Finalize();
+		}
+
+		/**
+		 * Build a Non uniform B-Spline surface
+		 * @param NodalVectorU Its size is the number of poles in U + the surface degree in U + 1 (PoleUNum + UDegre + 1)
+		 * @param NodalVectorV Its size is the number of poles in V + the surface degree in V + 1 (PoleVNum + VDegre + 1)
+		 */
+		FNURBSSurface(const double InToleranceGeometric, FNurbsSurfaceHomogeneousData NurbsData)
+			: FSurface(InToleranceGeometric)
+		{
+			FillNurbs(NurbsData);
+		}
+
+		FNURBSSurface() = default;
+
+	public:
+
+		virtual void Serialize(FCADKernelArchive& Ar) override
+		{
+			FSurface::Serialize(Ar);
+			Ar << PoleUCount;
+			Ar << PoleVCount;
+			Ar << UDegree;
+			Ar << VDegree;
+			Ar.Serialize(UNodalVector);
+			Ar.Serialize(VNodalVector);
+			Ar.Serialize(Weights);
+			Ar.Serialize(Poles);
+			Ar << bIsRational;
+
+			if (Ar.IsLoading())
+			{
+				Finalize();
+			}
+		}
+
+		ESurface GetSurfaceType() const
+		{
+			return ESurface::Nurbs;
+		}
+
+#ifdef CADKERNEL_DEV
+		virtual FInfoEntity& GetInfo(FInfoEntity&) const override;
+#endif
+
+		constexpr const int32 GetDegree(EIso Iso) const
+		{
+			switch (Iso)
+			{
+			case EIso::IsoU:
+				return UDegree;
+			case EIso::IsoV:
+			default:
+				return VDegree;
+			}
+		}
+
+		constexpr const int32 GetPoleCount(EIso Iso) const
+		{
+			switch (Iso)
+			{
+			case EIso::IsoU:
+				return PoleUCount;
+			case EIso::IsoV:
+			default:
+				return PoleVCount;
+			}
+		}
+
+		const TArray<FPoint>& GetPoles() const
+		{
+			return Poles;
+		}
+
+		const TArray<double>& GetWeights() const
+		{
+			return Weights;
+		}
+
+		TArray<double> GetHPoles() const
+		{
+			return HomogeneousPoles;
+		}
+
+		constexpr const TArray<double>& GetNodalVector(EIso Iso) const
+		{
+			switch (Iso)
+			{
+			case EIso::IsoU:
+				return UNodalVector;
+			case EIso::IsoV:
+			default:
+				return VNodalVector;
+			}
+		}
+
+		bool IsRational() const
+		{
+			return bIsRational;
+		}
+
+
+		virtual TSharedPtr<FEntityGeom> ApplyMatrix(const FMatrixH& InMatrix) const override;
+
+		virtual void EvaluatePoint(const FPoint2D& InSurfacicCoordinate, FSurfacicPoint& OutPoint3D, int32 InDerivativeOrder = 0) const override
+		{
+			BSpline::EvaluatePoint(*this, InSurfacicCoordinate, OutPoint3D, InDerivativeOrder);
+		}
+
+		void EvaluatePointGrid(const FCoordinateGrid& InSurfacicCoordinates, FSurfacicSampling& OutPoints, bool bComputeNormals) const override
+		{
+			BSpline::EvaluatePointGrid(*this, InSurfacicCoordinates, OutPoints, bComputeNormals);
+		}
+
+		virtual void LinesNotDerivables(const FSurfacicBoundary& InBoundary, int32 InDerivativeOrder, FCoordinateGrid& OutNotDerivableCoordinates) const override
+		{
+			BSpline::FindNotDerivableParameters(*this, InDerivativeOrder, InBoundary, OutNotDerivableCoordinates);
+		}
+
+		void ComputeMinToleranceIso();
+
+	private:
+		void Finalize();
+		void FillNurbs(FNurbsSurfaceHomogeneousData& NurbsData);
+	};
+
+} // namespace CADKernel
+
diff --git a/Engine/Source/Runtime/Datasmith/CADKernel/Private/CADKernel/Geo/Surfaces/OffsetSurface.h b/Engine/Source/Runtime/Datasmith/CADKernel/Private/CADKernel/Geo/Surfaces/OffsetSurface.h
new file mode 100644
--- /dev/null
+++ b/Engine/Source/Runtime/Datasmith/CADKernel/Private/CADKernel/Geo/Surfaces/OffsetSurface.h
@@ -0,0 +1,94 @@
+// Copyright Epic Games, Inc. All Rights Reserved.
+#pragma once
+
+#include "CADKernel/Geo/Surfaces/Surface.h"
+
+namespace CADKernel
+{
+
+	class CADKERNEL_API FOffsetSurface : public FSurface
+	{
+		friend FEntity;
+
+	protected:
+
+		TSharedPtr<FSurface> BaseSurface;
+		double Offset;
+
+		FOffsetSurface(const double InToleranceGeometric, TSharedRef<FSurface> InBaseSurface, double InOffset)
+			: FSurface(InToleranceGeometric)
+			, BaseSurface(InBaseSurface)
+			, Offset(InOffset)
+		{
+			ComputeMinToleranceIso();
+		}
+
+		FOffsetSurface() = default;
+
+		void ComputeMinToleranceIso()
+		{
+			MinToleranceIso = BaseSurface->GetIsoTolerances();
+		}
+
+	public:
+
+		virtual void Serialize(FCADKernelArchive& Ar) override
+		{
+			FSurface::Serialize(Ar);
+			SerializeIdent(Ar, BaseSurface);
+			Ar << Offset;
+		}
+
+		virtual void SpawnIdent(FDatabase& Database) override
+		{
+			if (!FEntity::SetId(Database))
+			{
+				return;
+			}
+
+			BaseSurface->SpawnIdent(Database);
+		}
+
+		virtual void ResetMarkersRecursively() override
+		{
+			ResetMarkers();
+			BaseSurface->ResetMarkersRecursively();
+		}
+
+		ESurface GetSurfaceType() const
+		{
+			return ESurface::Offset;
+		}
+
+		const double GetOffset() const
+		{
+			return Offset;
+		}
+
+		TSharedPtr<FSurface> GetBaseSurface() const
+		{
+			return BaseSurface;
+		}
+
+		virtual void InitBoundary() override;
+
+#ifdef CADKERNEL_DEV
+		virtual FInfoEntity& GetInfo(FInfoEntity&) const override;
+#endif
+
+		virtual TSharedPtr<FEntityGeom> ApplyMatrix(const FMatrixH& InMatrix) const override;
+
+		virtual void LinesNotDerivables(const FSurfacicBoundary& Bounds, int32 InDerivativeOrder, FCoordinateGrid& OutNotDerivableCoordinates) const override;
+
+		/**
+		 * The Gradient and Laplacian are approximated by their value on the carrier surface.
+		 */
+		virtual void EvaluatePoint(const FPoint2D& InSurfacicCoordinate, FSurfacicPoint& OutPoint3D, int32 InDerivativeOrder = 0) const override;
+		virtual void EvaluatePoints(const TArray<FPoint2D>& InSurfacicCoordinates, TArray<FSurfacicPoint>& OutPoint3D, int32 InDerivativeOrder = 0) const override;
+		virtual void EvaluatePointGrid(const FCoordinateGrid& Coordinates, FSurfacicSampling& OutPoints, bool bComputeNormals = false) const override;
+
+		virtual void Presample(const FSurfacicBoundary& InBoundaries, FCoordinateGrid& OutCoordinates) override;
+	};
+
+} // namespace CADKernel
+
diff --git a/Engine/Source/Runtime/Datasmith/CADKernel/Private/CADKernel/Geo/Surfaces/PlaneSurface.cpp b/Engine/Source/Runtime/Datasmith/CADKernel/Private/CADKernel/Geo/Surfaces/PlaneSurface.cpp
--- a/Engine/Source/Runtime/Datasmith/CADKernel/Private/CADKernel/Geo/Surfaces/PlaneSurface.cpp
+++ b/Engine/Source/Runtime/Datasmith/CADKernel/Private/CADKernel/Geo/Surfaces/PlaneSurface.cpp
@@ -6,7 +6,6 @@
 #include "CADKernel/Core/System.h"
 #include "CADKernel/Geo/GeoPoint.h"
 #include "CADKernel/Geo/Sampling/SurfacicSampling.h"
-#include "CADKernel/Math/Plane.h"
 
 namespace CADKernel
 {
diff --git a/Engine/Source/Runtime/Datasmith/CADKernel/Private/CADKernel/Geo/Surfaces/PlaneSurface.h b/Engine/Source/Runtime/Datasmith/CADKernel/Private/CADKernel/Geo/Surfaces/PlaneSurface.h
new file mode 100644
--- /dev/null
+++ b/Engine/Source/Runtime/Datasmith/CADKernel/Private/CADKernel/Geo/Surfaces/PlaneSurface.h
@@ -0,0 +1,110 @@
+// Copyright Epic Games, Inc. All Rights Reserved.
+#pragma once
+
+#include "CADKernel/Geo/Surfaces/Surface.h"
+
+namespace CADKernel
+{
+	class CADKERNEL_API FPlaneSurface : public FSurface
+	{
+		friend FEntity;
+
+	protected:
+
+		FMatrixH Matrix;
+		FMatrixH InverseMatrix;
+
+		/**
+		 * The plane surface is the plane XY.
+		 * The surface is placed at its final position and orientation by the Matrix
+		 */
+		FPlaneSurface(const double InToleranceGeometric, const FMatrixH& InMatrix, const FSurfacicBoundary& InBoundary);
+
+		/**
+		 * The plane surface is the plane XY.
+		 * The surface is placed at its final position and orientation by the Matrix,
+		 * The matrix is calculated from the normal plane at its final position and its distance from the origin along the normal.
+		 */
+		FPlaneSurface(const double InToleranceGeometric, double InDistanceFromOrigin, FPoint InNormal, const FSurfacicBoundary& InBoundary)
+			: FPlaneSurface(InToleranceGeometric, InNormal* InDistanceFromOrigin, InNormal, InBoundary)
+		{
+		}
+
+		/**
+		 * The plane surface is the plane XY.
+		 * The surface is placed at its final position and orientation by the Matrix 
+		 * The matrix is calculated from the plan origin at its final position and its final normal
+		 */
+		FPlaneSurface(const double InToleranceGeometric, const FPoint& Position, FPoint Normal, const FSurfacicBoundary& InBoundary);
+
+		FPlaneSurface() = default;
+		
+		void ComputeMinToleranceIso()
+		{
+			FPoint Origin = Matrix.Multiply(FPoint::ZeroPoint);
+
+			FPoint Point2DU{ 1 , 0, 0 };
+			FPoint Point2DV{ 0, 1, 0 };
+
+			double ToleranceU = Tolerance3D / ComputeScaleAlongAxis(Point2DU, Matrix, Origin);
+			double ToleranceV = Tolerance3D / ComputeScaleAlongAxis(Point2DV, Matrix, Origin);
+
+			MinToleranceIso.Set(ToleranceU, ToleranceV);
+		}
+
+	public:
+
+		virtual void Serialize(FCADKernelArchive& Ar) override
+		{
+			FSurface::Serialize(Ar);
+			Ar << Matrix;
+			Ar << InverseMatrix;
+		}
+
+		ESurface GetSurfaceType() const
+		{
+			return ESurface::Plane;
+		}
+
+		const FMatrixH& GetMatrix() const
+		{
+			return Matrix;
+		}
+
+		FPlane GetPlane() const;
+
+#ifdef CADKERNEL_DEV
+		virtual FInfoEntity& GetInfo(FInfoEntity&) const override;
+#endif
+
+		virtual TSharedPtr<FEntityGeom> ApplyMatrix(const FMatrixH& InMatrix) const override;
+
+		virtual void EvaluatePoint(const FPoint2D& InSurfacicCoordinate, FSurfacicPoint& OutPoint3D, int32 InDerivativeOrder = 0) const override;
+		virtual void EvaluatePoints(const TArray<FPoint2D>& InSurfacicCoordinates, TArray<FSurfacicPoint>& OutPoint3D, int32 InDerivativeOrder = 0) const override;
+		virtual void EvaluatePointGrid(const FCoordinateGrid& Coordinates, FSurfacicSampling& OutPoints, bool bComputeNormals = false) const override;
+
+		virtual FPoint ProjectPoint(const FPoint& InPoint, FPoint* OutProjectedPoint = nullptr) const;// override;
+		virtual void ProjectPoints(const TArray<FPoint>& InPoints, TArray<FPoint>* OutProjectedPointCoordinates, TArray<FPoint>* OutProjectedPoints = nullptr) const;// override;
+
+		virtual void Presample(const FSurfacicBoundary& InBoundaries, FCoordinateGrid& OutCoordinates) override
+		{
+			OutCoordinates[EIso::IsoU].Empty(3);
+			OutCoordinates[EIso::IsoU].Add(InBoundaries[EIso::IsoU].Min);
+			OutCoordinates[EIso::IsoU].Add(InBoundaries[EIso::IsoU].GetMiddle());
+			OutCoordinates[EIso::IsoU].Add(InBoundaries[EIso::IsoU].Max);
+
+			OutCoordinates[EIso::IsoV].Empty(3);
+			OutCoordinates[EIso::IsoV].Add(InBoundaries[EIso::IsoV].Min);
+			OutCoordinates[EIso::IsoV].Add(InBoundaries[EIso::IsoV].GetMiddle());
+			OutCoordinates[EIso::IsoV].Add(InBoundaries[EIso::IsoV].Max);
+		}
+
+		virtual void IsSurfaceClosed(bool& bOutClosedAlongU, bool& bOutClosedAlongV) const override
+		{
+			bOutClosedAlongU = false;
+			bOutClosedAlongV = false;
+		}
+
+	};
+
+} // namespace CADKernel
diff --git a/Engine/Source/Runtime/Datasmith/CADKernel/Private/CADKernel/Geo/Surfaces/RevolutionSurface.h b/Engine/Source/Runtime/Datasmith/CADKernel/Private/CADKernel/Geo/Surfaces/RevolutionSurface.h
new file mode 100644
--- /dev/null
+++ b/Engine/Source/Runtime/Datasmith/CADKernel/Private/CADKernel/Geo/Surfaces/RevolutionSurface.h
@@ -0,0 +1,81 @@
+// Copyright Epic Games, Inc. All Rights Reserved.
+#pragma once
+
+#include "CADKernel/Geo/Surfaces/Surface.h"
+
+namespace CADKernel
+{
+	class FCurve;
+	class FSegmentCurve;
+
+	class CADKERNEL_API FRevolutionSurface : public FSurface
+	{
+		friend FEntity;
+
+	protected:
+
+		TSharedPtr<FSegmentCurve> Axis;
+		TSharedPtr<FCurve> Generatrix;
+
+		FPoint RotationAxis;
+
+		/**
+		 * A surface of revolution is defined by an axis of rotation (Axis) and a generatrix.
+		 * The start and terminate rotation angles can limited the surface. 
+		 * The surface is created by rotating the generatrix about the axis of rotation through the start and terminating angles. 
+		 * 
+		 * The bounds of the surface of revolution are defined as follow:
+		 * Bounds[EIso::IsoU].Min = Start rotation angle;
+		 * Bounds[EIso::IsoU].Max = Terminate rotation angle;
+		 * Bounds[EIso::IsoV].Min = Min Generatrix bounds;
+		 * Bounds[EIso::IsoV].Max = Max Generatrix bounds;
+		 */
+		FRevolutionSurface(const double InToleranceGeometric, TSharedRef<FSegmentCurve> Axe, TSharedRef<FCurve> Generatrix, double MinAngle, double MaxAngle);
+
+		FRevolutionSurface() = default;
+
+	public:
+
+		virtual void Serialize(FCADKernelArchive& Ar) override
+		{
+			FSurface::Serialize(Ar);
+			SerializeIdent(Ar, Axis);
+			SerializeIdent(Ar, Generatrix);
+			Ar << RotationAxis;
+		}
+
+		virtual void SpawnIdent(FDatabase& Database) override;
+		virtual void ResetMarkersRecursively() override;
+
+		ESurface GetSurfaceType() const
+		{
+			return ESurface::Revolution;
+		}
+
+		TSharedPtr<FCurve> GetGeneratrix() const
+		{
+			return Generatrix;
+		}
+
+		virtual TSharedPtr<FSegmentCurve> GetAxe() const
+		{
+			return Axis;
+		}
+
+#ifdef CADKERNEL_DEV
+		virtual FInfoEntity& GetInfo(FInfoEntity&) const override;
+#endif
+
+		virtual TSharedPtr<FEntityGeom> ApplyMatrix(const FMatrixH& InMatrix) const override;
+
+		virtual void EvaluatePoint(const FPoint2D& InSurfacicCoordinate, FSurfacicPoint& OutPoint3D, int32 InDerivativeOrder = 0) const override;
+		virtual void EvaluatePointGrid(const FCoordinateGrid& Coordinates, FSurfacicSampling& OutPoints, bool bComputeNormals = false) const override;
+
+		virtual void LinesNotDerivables(const FSurfacicBoundary& Bounds, int32 InDerivativeOrder, FCoordinateGrid& OutCoordinates) const override;
+
+		virtual void Presample(const FSurfacicBoundary& InBoundaries, FCoordinateGrid& OutCoordinates) override;
+
+	};
+
+} // namespace CADKernel
+
diff --git a/Engine/Source/Runtime/Datasmith/CADKernel/Private/CADKernel/Geo/Surfaces/RuledSurface.h b/Engine/Source/Runtime/Datasmith/CADKernel/Private/CADKernel/Geo/Surfaces/RuledSurface.h
new file mode 100644
--- /dev/null
+++ b/Engine/Source/Runtime/Datasmith/CADKernel/Private/CADKernel/Geo/Surfaces/RuledSurface.h
@@ -0,0 +1,66 @@
+// Copyright Epic Games, Inc. All Rights Reserved.
+#pragma once
+
+#include "CADKernel/Geo/Surfaces/Surface.h"
+
+namespace CADKernel
+{
+	class FCurve;
+	class FNURBSSurface;
+
+	class CADKERNEL_API FRuledSurface : public FSurface
+	{
+		friend FEntity;
+
+	protected:
+		TSharedPtr<FCurve> Curves[2];
+
+		FRuledSurface(const double InToleranceGeometric, TSharedPtr<FCurve> InCurveU, TSharedPtr<FCurve> InCurveV)
+			: FSurface(InToleranceGeometric, 0., 1., 0., 1.)
+		{
+			Curves[0] = InCurveU;
+			Curves[1] = InCurveV;
+			ComputeDefaultMinToleranceIso();
+		}
+
+		FRuledSurface() = default;
+
+	public:
+
+		virtual void Serialize(FCADKernelArchive& Ar) override
+		{
+			FSurface::Serialize(Ar);
+			SerializeIdent(Ar, Curves[0]);
+			SerializeIdent(Ar, Curves[1]);
+		}
+
+		virtual void SpawnIdent(FDatabase& Database) override;
+		virtual void ResetMarkersRecursively() override;
+
+		ESurface GetSurfaceType() const
+		{
+			return ESurface::Ruled;
+		}
+
+		const TSharedPtr<FCurve>& GetCurve(EIso Iso) const
+		{
+			return Curves[Iso];
+		}
+
+#ifdef CADKERNEL_DEV
+		virtual FInfoEntity& GetInfo(FInfoEntity&) const override;
+#endif
+
+		virtual TSharedPtr<FEntityGeom> ApplyMatrix(const FMatrixH& InMatrix) const override;
+
+		virtual void EvaluatePoint(const FPoint2D& InSurfacicCoordinate, FSurfacicPoint& OutPoint3D, int32 InDerivativeOrder = 0) const override;
+		virtual void EvaluatePointGrid(const FCoordinateGrid& Coordinates, FSurfacicSampling& OutPoints, bool bComputeNormals = false) const override;
+
+		virtual void LinesNotDerivables(const FSurfacicBoundary& Bounds, int32 InDerivativeOrder, FCoordinateGrid& OutCoordinates) const override;
+
+		virtual void Presample(const FSurfacicBoundary& InBoundaries, FCoordinateGrid& OutCoordinates) override;
+
+	};
+
+} // namespace CADKernel
+
diff --git a/Engine/Source/Runtime/Datasmith/CADKernel/Private/CADKernel/Geo/Surfaces/SphericalSurface.cpp b/Engine/Source/Runtime/Datasmith/CADKernel/Private/CADKernel/Geo/Surfaces/SphericalSurface.cpp
--- a/Engine/Source/Runtime/Datasmith/CADKernel/Private/CADKernel/Geo/Surfaces/SphericalSurface.cpp
+++ b/Engine/Source/Runtime/Datasmith/CADKernel/Private/CADKernel/Geo/Surfaces/SphericalSurface.cpp
@@ -34,7 +34,7 @@ void FSphericalSurface::EvaluatePointGridInCylindricalSpace(const FCoordinateGri
 	{
 		double CosV = cos(VAngle);
 		double Rho = Radius * CosV;
-		double SwapOrientation = (VAngle < PI && VAngle >= 0) ? 1.0 : -1.0;
+		double SwapOrientation = (VAngle < DOUBLE_PI && VAngle >= 0) ? 1.0 : -1.0;
 
 		for (int32 Undex = 0; Undex < UCount; Undex++)
 		{
diff --git a/Engine/Source/Runtime/Datasmith/CADKernel/Private/CADKernel/Geo/Surfaces/SphericalSurface.h b/Engine/Source/Runtime/Datasmith/CADKernel/Private/CADKernel/Geo/Surfaces/SphericalSurface.h
new file mode 100644
--- /dev/null
+++ b/Engine/Source/Runtime/Datasmith/CADKernel/Private/CADKernel/Geo/Surfaces/SphericalSurface.h
@@ -0,0 +1,140 @@
+// Copyright Epic Games, Inc. All Rights Reserved.
+#pragma once
+
+#include "CADKernel/Geo/Surfaces/Surface.h"
+
+namespace CADKernel
+{
+	class CADKERNEL_API FSphericalSurface : public FSurface
+	{
+		friend FEntity;
+
+	protected:
+		FMatrixH Matrix;
+		double Radius;
+
+		/**
+		 * The spherical surface is defined its radius. 
+		 *
+		 * It's defined as the rotation around Z axis of an semicircle defined in the plan XY centered at the origin. 
+		 *
+		 * The spherical surface is placed at its final position and orientation by the Matrix
+		 */
+		FSphericalSurface(const double InToleranceGeometric, const FMatrixH& InMatrix, double InRadius, double InParallelStartAngle = 0.0, double InParallelEndAngle = DOUBLE_TWO_PI, double InMeridianStartAngle = -DOUBLE_HALF_PI, double InMeridianEndAngle = DOUBLE_HALF_PI)
+			: FSphericalSurface(InToleranceGeometric, InMatrix, InRadius, FSurfacicBoundary(InParallelStartAngle, InParallelEndAngle, InMeridianStartAngle, InMeridianEndAngle))
+		{
+		}
+
+		/**
+		 * The spherical surface is defined its radius.
+		 *
+		 * It's defined as the rotation around Z axis of an semicircle defined in the plan XY centered at the origin.
+		 *
+		 * The spherical surface is placed at its final position and orientation by the Matrix
+		 *
+		 * The bounds of the spherical surface are defined as follow:
+		 * Bounds[EIso::IsoU].Min = MeridianStartAngle
+		 * Bounds[EIso::IsoU].Max = MeridianEndAngle
+		 * Bounds[EIso::IsoV].Min = ParallelStartAngle
+		 * Bounds[EIso::IsoV].Max = ParallelEndAngle
+		 */
+		FSphericalSurface(const double InToleranceGeometric, const FMatrixH& InMatrix, double InRadius, const FSurfacicBoundary& Boundary)
+			: FSurface(InToleranceGeometric, Boundary)
+			, Matrix(InMatrix)
+			, Radius(InRadius)
+		{
+			ComputeMinToleranceIso();
+		}
+
+		FSphericalSurface() = default;
+
+		void ComputeMinToleranceIso()
+		{
+			double Tolerance2D = Tolerance3D / Radius;
+
+			FPoint Origin = Matrix.Multiply(FPoint::ZeroPoint);
+
+			FPoint Point2DU{ 1 , 0, 0 };
+			FPoint Point2DV{ 0, 1, 0 };
+
+			double ToleranceU = Tolerance2D / ComputeScaleAlongAxis(Point2DU, Matrix, Origin);
+			double ToleranceV = Tolerance2D / ComputeScaleAlongAxis(Point2DV, Matrix, Origin);
+
+			MinToleranceIso.Set(ToleranceU, ToleranceV);
+		}
+	
+	public:
+
+		virtual void Serialize(FCADKernelArchive& Ar) override
+		{
+			FSurface::Serialize(Ar);
+			Ar << Matrix;
+			Ar << Radius;
+		}
+
+		ESurface GetSurfaceType() const
+		{
+			return ESurface::Sphere;
+		}
+
+#ifdef CADKERNEL_DEV
+		virtual FInfoEntity& GetInfo(FInfoEntity&) const override;
+#endif
+
+		virtual TSharedPtr<FEntityGeom> ApplyMatrix(const FMatrixH& InMatrix) const override;
+
+		virtual void EvaluatePoint(const FPoint2D& InSurfacicCoordinate, FSurfacicPoint& OutPoint3D, int32 InDerivativeOrder = 0) const override
+		{
+			double CosU = cos(InSurfacicCoordinate.U);
+			double CosV = cos(InSurfacicCoordinate.V);
+
+			double SinU = sin(InSurfacicCoordinate.U);
+			double SinV = sin(InSurfacicCoordinate.V);
+
+			OutPoint3D.DerivativeOrder = InDerivativeOrder;
+			OutPoint3D.Point.Set(Radius * CosV * CosU, Radius * CosV * SinU, Radius * SinV);
+			OutPoint3D.Point = Matrix.Multiply(OutPoint3D.Point);
+
+			if (InDerivativeOrder > 0)
+			{
+				OutPoint3D.GradientU = FPoint(-Radius * CosV * SinU, Radius * CosV * CosU, 0.0);
+				OutPoint3D.GradientV = FPoint(-Radius * SinV * CosU, -Radius * SinV * SinU, Radius * CosV);
+
+				OutPoint3D.GradientU = Matrix.MultiplyVector(OutPoint3D.GradientU);
+				OutPoint3D.GradientV = Matrix.MultiplyVector(OutPoint3D.GradientV);
+			}
+
+			if (InDerivativeOrder > 1)
+			{
+				OutPoint3D.LaplacianU = FPoint(-Radius * CosV * CosU, -Radius * CosV * SinU, 0.0);
+				OutPoint3D.LaplacianV = FPoint(OutPoint3D.LaplacianU.X, OutPoint3D.LaplacianU.Y, -Radius * SinV);
+				OutPoint3D.LaplacianUV = FPoint(Radius * SinV * SinU, -Radius * SinV * CosU, 0.);
+
+				OutPoint3D.LaplacianU = Matrix.MultiplyVector(OutPoint3D.LaplacianU);
+				OutPoint3D.LaplacianV = Matrix.MultiplyVector(OutPoint3D.LaplacianV);
+				OutPoint3D.LaplacianUV = Matrix.MultiplyVector(OutPoint3D.LaplacianUV);
+			}
+		}
+
+		virtual void EvaluatePointGrid(const FCoordinateGrid& Coordinates, FSurfacicSampling& OutPoints, bool bComputeNormals = false) const override;
+
+		virtual FPoint2D EvaluatePointInCylindricalSpace(const FPoint2D& InSurfacicCoordinate) const override
+		{
+			double CosU = cos(InSurfacicCoordinate.U);
+			double CosV = cos(InSurfacicCoordinate.V);
+
+			double SinU = sin(InSurfacicCoordinate.U);
+			double SwapOrientation = (InSurfacicCoordinate.V < DOUBLE_PI && InSurfacicCoordinate.V >= 0) ? 1.0 : -1.0;
+
+			return FPoint(Radius * CosV * CosU * SwapOrientation, Radius * CosV * SinU );
+		}
+
+		virtual void EvaluatePointGridInCylindricalSpace(const FCoordinateGrid& Coordinates, TArray<FPoint2D>&) const override;
+
+		virtual void Presample(const FSurfacicBoundary& InBoundaries, FCoordinateGrid& OutCoordinates) override
+		{
+			PresampleIsoCircle(InBoundaries, OutCoordinates, EIso::IsoU);
+			PresampleIsoCircle(InBoundaries, OutCoordinates, EIso::IsoV);
+		}
+	};
+}
diff --git a/Engine/Source/Runtime/Datasmith/CADKernel/Private/CADKernel/Geo/Surfaces/SurfaceMaker.cpp b/Engine/Source/Runtime/Datasmith/CADKernel/Private/CADKernel/Geo/Surfaces/SurfaceMaker.cpp
new file mode 100644
--- /dev/null
+++ b/Engine/Source/Runtime/Datasmith/CADKernel/Private/CADKernel/Geo/Surfaces/SurfaceMaker.cpp
@@ -0,0 +1,63 @@
+// Copyright Epic Games, Inc. All Rights Reserved.
+
+#include "CADKernel/Geo/Surfaces/Surface.h"
+
+#include "CADKernel/Geo/Surfaces/PlaneSurface.h"
+#include "CADKernel/Geo/Surfaces/BezierSurface.h"
+#include "CADKernel/Geo/Surfaces/NURBSSurface.h"
+#include "CADKernel/Geo/Surfaces/RuledSurface.h"
+#include "CADKernel/Geo/Surfaces/RevolutionSurface.h"
+#include "CADKernel/Geo/Surfaces/TabulatedCylinderSurface.h"
+#include "CADKernel/Geo/Surfaces/CylinderSurface.h"
+#include "CADKernel/Geo/Surfaces/OffsetSurface.h"
+#include "CADKernel/Geo/Surfaces/CompositeSurface.h"
+#include "CADKernel/Geo/Surfaces/CoonsSurface.h"
+#include "CADKernel/Geo/Surfaces/SphericalSurface.h"
+#include "CADKernel/Geo/Surfaces/TorusSurface.h"
+#include "CADKernel/Geo/Surfaces/ConeSurface.h"
+
+
+namespace CADKernel
+{
+TSharedPtr<FSurface> FSurface::MakeBezierSurface(const double InToleranceGeometric, int32 InUDegre, int32 InVDegre, const TArray<FPoint>& InPoles)
+{
+	return CADKernel::FEntity::MakeShared<CADKernel::FBezierSurface>(InToleranceGeometric, InUDegre, InVDegre, InPoles);
+}
+
+TSharedPtr<FSurface> FSurface::MakeConeSurface(const double InToleranceGeometric, const FMatrixH& InMatrix, double InStartRadius, double InConeAngle, const FSurfacicBoundary& InBoundary)
+{
+	return CADKernel::FEntity::MakeShared<CADKernel::FConeSurface>(InToleranceGeometric, InMatrix, InStartRadius, InConeAngle, InBoundary);
+}
+
+TSharedPtr<FSurface> FSurface::MakeCylinderSurface(const double InToleranceGeometric, const FMatrixH& InMatrix, const double InRadius, const FSurfacicBoundary& InBoundary)
+{
+	return CADKernel::FEntity::MakeShared<CADKernel::FCylinderSurface>(InToleranceGeometric, InMatrix, InRadius, InBoundary);
+}
+
+TSharedPtr<FSurface> FSurface::MakeNurbsSurface(const double InToleranceGeometric, const FNurbsSurfaceHomogeneousData& NurbsData)
+{
+	return CADKernel::FEntity::MakeShared<CADKernel::FNURBSSurface>(InToleranceGeometric, NurbsData);
+}
+
+TSharedPtr<FSurface> FSurface::MakeNurbsSurface(const double InToleranceGeometric, const FNurbsSurfaceData& NurbsData)
+{
+	return CADKernel::FEntity::MakeShared<CADKernel::FNURBSSurface>(InToleranceGeometric, NurbsData);
+}
+
+TSharedPtr<FSurface> FSurface::MakePlaneSurface(const double InToleranceGeometric, const FMatrixH& InMatrix, const FSurfacicBoundary& InBoundary)
+{
+	return CADKernel::FEntity::MakeShared<CADKernel::FPlaneSurface>(InToleranceGeometric, InMatrix, InBoundary);
+}
+
+TSharedPtr<FSurface> FSurface::MakeSphericalSurface(const double InToleranceGeometric, const FMatrixH& InMatrix, double InRadius, const FSurfacicBoundary& InBoundary)
+{
+	return CADKernel::FEntity::MakeShared<CADKernel::FSphericalSurface>(InToleranceGeometric, InMatrix, InRadius, InBoundary);
+}
+
+TSharedPtr<FSurface> FSurface::MakeTorusSurface(const double InToleranceGeometric, const FMatrixH& InMatrix, double InMajorRadius, double InMinorRadius, const FSurfacicBoundary& InBoundary)
+{
+	return CADKernel::FEntity::MakeShared<CADKernel::FTorusSurface>(InToleranceGeometric, InMatrix, InMajorRadius, InMinorRadius, InBoundary);
+}
+
+} // namespace CADKernel
+
diff --git a/Engine/Source/Runtime/Datasmith/CADKernel/Private/CADKernel/Geo/Surfaces/TabulatedCylinderSurface.h b/Engine/Source/Runtime/Datasmith/CADKernel/Private/CADKernel/Geo/Surfaces/TabulatedCylinderSurface.h
new file mode 100644
--- /dev/null
+++ b/Engine/Source/Runtime/Datasmith/CADKernel/Private/CADKernel/Geo/Surfaces/TabulatedCylinderSurface.h
@@ -0,0 +1,80 @@
+// Copyright Epic Games, Inc. All Rights Reserved.
+#pragma once
+
+#include "CADKernel/Geo/Surfaces/Surface.h"
+#include "CADKernel/Math/Point.h"
+
+namespace CADKernel
+{
+	class FCurve;
+
+	class CADKERNEL_API FTabulatedCylinderSurface : public FSurface
+	{
+		friend FEntity;
+
+	protected:
+		TSharedPtr<FCurve> GuideCurve;
+		FPoint DirectorVector;
+
+		/**
+		 * A tabulated cylinder is a surface formed by moving a curve segment called GuideCurve parallel to
+		 * itself along a vector called DirectorVector. 
+		 *
+		 * The bounds of the cylinder are defined as follow:
+		 * Bounds[EIso::IsoU] = GuideCurve bounds;
+		 * Bounds[EIso::IsoV].Min = InVMin;
+		 * Bounds[EIso::IsoV].Max = InVMax;
+		 */
+		FTabulatedCylinderSurface(const double InToleranceGeometric, TSharedPtr<FCurve> InGuideCurve, const FPoint& InDirectorVector, double InVMin = 0., double InVMax = 1.)
+			: FSurface(InToleranceGeometric, 0.0, 1.0, InVMin, InVMax)
+			, GuideCurve(InGuideCurve)
+			, DirectorVector(InDirectorVector)
+		{
+			ComputeDefaultMinToleranceIso();
+		}
+
+		FTabulatedCylinderSurface() = default;
+
+	public:
+
+		virtual void Serialize(FCADKernelArchive& Ar) override
+		{
+			FSurface::Serialize(Ar);
+			SerializeIdent(Ar, GuideCurve);
+			Ar << DirectorVector;
+		}	
+
+		virtual void SpawnIdent(FDatabase& Database) override;
+		virtual void ResetMarkersRecursively() override;
+		
+		ESurface GetSurfaceType() const
+		{
+			return ESurface::TabulatedCylinder;
+		}
+
+		TSharedPtr<FCurve> GetGuideCurve() const
+		{
+			return GuideCurve;
+		}
+
+		FPoint GetDirectorVector() const
+		{
+			return DirectorVector;
+		}
+
+		virtual void EvaluatePoint(const FPoint2D& InSurfacicCoordinate, FSurfacicPoint& OutPoint3D, int32 InDerivativeOrder = 0) const override;
+		virtual void EvaluatePointGrid(const FCoordinateGrid& Coordinates, FSurfacicSampling& OutPoints, bool bComputeNormals = false) const override;
+
+		virtual void LinesNotDerivables(const FSurfacicBoundary& Bounds, int32 InDerivativeOrder, FCoordinateGrid& OutNotDerivables) const override;
+
+		virtual void Presample(const FSurfacicBoundary& InBoundaries, FCoordinateGrid& OutCoordinates) override;
+
+		virtual TSharedPtr<FEntityGeom> ApplyMatrix(const FMatrixH& InMatrix) const override;
+
+#ifdef CADKERNEL_DEV
+		virtual FInfoEntity& GetInfo(FInfoEntity&) const override;
+#endif
+
+	};
+
+} // namespace CADKernel
diff --git a/Engine/Source/Runtime/Datasmith/CADKernel/Private/CADKernel/Geo/Surfaces/TorusSurface.cpp b/Engine/Source/Runtime/Datasmith/CADKernel/Private/CADKernel/Geo/Surfaces/TorusSurface.cpp
--- a/Engine/Source/Runtime/Datasmith/CADKernel/Private/CADKernel/Geo/Surfaces/TorusSurface.cpp
+++ b/Engine/Source/Runtime/Datasmith/CADKernel/Private/CADKernel/Geo/Surfaces/TorusSurface.cpp
@@ -56,7 +56,7 @@ void FTorusSurface::EvaluatePointGridInCylindricalSpace(const FCoordinateGrid& C
 		double CosV = cos(Angle);
 		double Rho = MajorRadius + MinorRadius * CosV;
 
-		double SwapOrientation = (Angle < PI && Angle >= 0) ? 1.0 : -1.0;
+		double SwapOrientation = (Angle < DOUBLE_PI && Angle >= 0) ? 1.0 : -1.0;
 		for (int32 Undex = 0; Undex < UCount; Undex++)
 		{
 			OutPoints.Emplace(Rho * CosU[Undex] * SwapOrientation, Rho * SinU[Undex]);
diff --git a/Engine/Source/Runtime/Datasmith/CADKernel/Private/CADKernel/Geo/Surfaces/TorusSurface.h b/Engine/Source/Runtime/Datasmith/CADKernel/Private/CADKernel/Geo/Surfaces/TorusSurface.h
new file mode 100644
--- /dev/null
+++ b/Engine/Source/Runtime/Datasmith/CADKernel/Private/CADKernel/Geo/Surfaces/TorusSurface.h
@@ -0,0 +1,153 @@
+// Copyright Epic Games, Inc. All Rights Reserved.
+#pragma once
+
+#include "CADKernel/Geo/Surfaces/Surface.h"
+
+namespace CADKernel
+{
+	class CADKERNEL_API FTorusSurface : public FSurface
+	{
+		friend FEntity;
+
+	protected:
+		FMatrixH Matrix;
+		double MajorRadius;
+		double MinorRadius;
+
+		/**
+		 * A torus is the solid formed by revolving a circular disc about a specified coplanar axis. 
+		 * MajorRadius is the distance from the axis to the center of the defining disc, and MinorRadius is the radius of the defining disc,
+		 * where MajorRadius > MinorRadius > 0.0. 
+		 *
+		 * The torus computed at the origin with Z axis.
+		 * It is placed at its final position and orientation by the Matrix
+		 */
+		FTorusSurface(const double InToleranceGeometric, const FMatrixH& InMatrix, double InMajorRadius, double InMinorRadius, double InMajorStartAngle = 0.0, double InMajorEndAngle = DOUBLE_TWO_PI, double InMinorStartAngle = 0.0, double InMinorEndAngle = DOUBLE_TWO_PI)
+			: FTorusSurface(InToleranceGeometric, InMatrix, InMajorRadius, InMinorRadius, FSurfacicBoundary(InMajorStartAngle, InMajorEndAngle, InMinorStartAngle, InMinorEndAngle))
+		{
+		}
+
+		/**
+		 * A torus is the solid formed by revolving a circular disc about a specified coplanar axis.
+		 * MajorRadius is the distance from the axis to the center of the defining disc, and MinorRadius is the radius of the defining disc,
+		 * where MajorRadius > MinorRadius > 0.0.
+		 *
+		 * The torus computed at the origin with Z axis.
+		 * It is placed at its final position and orientation by the Matrix
+		 *
+		 * The bounds of the cone are defined as follow:
+		 * Bounds[EIso::IsoU].Min = MajorStartAngle
+		 * Bounds[EIso::IsoU].Max = MajorEndAngle
+		 * Bounds[EIso::IsoV].Min = MinorStartAngle
+		 * Bounds[EIso::IsoV].Max = MinorEndAngle
+		 */
+		FTorusSurface(const double InToleranceGeometric, const FMatrixH & InMatrix, double InMajorRadius, double InMinorRadius, const FSurfacicBoundary& InBoundary)
+			: FSurface(InToleranceGeometric, InBoundary)
+			, Matrix(InMatrix)
+			, MajorRadius(InMajorRadius)
+			, MinorRadius(InMinorRadius)
+		{
+			ComputeMinToleranceIso();
+		}
+
+		FTorusSurface() = default;
+
+		void ComputeMinToleranceIso()
+		{
+			double Tolerance2DU = Tolerance3D / MajorRadius;
+			double Tolerance2DV = Tolerance3D / MinorRadius;
+
+			FPoint Origin = Matrix.Multiply(FPoint::ZeroPoint);
+
+			FPoint Point2DU{ 1 , 0, 0 };
+			FPoint Point2DV{ 0, 1, 0 };
+
+			Tolerance2DU /= ComputeScaleAlongAxis(Point2DU, Matrix, Origin);
+			Tolerance2DV /= ComputeScaleAlongAxis(Point2DV, Matrix, Origin);
+
+			MinToleranceIso.Set(FMath::Abs(Tolerance2DU), FMath::Abs(Tolerance2DV));
+		}
+		
+	public:
+
+		virtual void Serialize(FCADKernelArchive& Ar) override
+		{
+			FSurface::Serialize(Ar);
+			Ar << Matrix;
+			Ar << MajorRadius;
+			Ar << MinorRadius;
+		}
+
+		ESurface GetSurfaceType() const
+		{
+			return ESurface::Torus;
+		}
+
+#ifdef CADKERNEL_DEV
+		virtual FInfoEntity& GetInfo(FInfoEntity&) const override;
+#endif
+
+		virtual TSharedPtr<FEntityGeom> ApplyMatrix(const FMatrixH& InMatrix) const override;
+
+		virtual void EvaluatePoint(const FPoint2D& InSurfacicCoordinate, FSurfacicPoint& OutPoint3D, int32 InDerivativeOrder = 0) const override
+		{
+			const double CosU = cos(InSurfacicCoordinate.U);
+			const double CosV = cos(InSurfacicCoordinate.V);
+
+			const double SinU = sin(InSurfacicCoordinate.U);
+			const double SinV = sin(InSurfacicCoordinate.V);
+			double Rho = (MajorRadius + MinorRadius * CosV);
+
+			OutPoint3D.DerivativeOrder = InDerivativeOrder;
+			OutPoint3D.Point.Set(Rho * CosU, Rho * SinU, MinorRadius * SinV);
+
+			OutPoint3D.Point = Matrix.Multiply(OutPoint3D.Point);
+
+			if (InDerivativeOrder > 0)
+			{
+				OutPoint3D.GradientU = FPoint((MajorRadius + MinorRadius * CosV) * -SinU, (MajorRadius + MinorRadius * CosV) * CosU, 0.0);
+				OutPoint3D.GradientV = FPoint((MinorRadius * -SinV) * CosU, (MinorRadius * -SinV) * SinU, MinorRadius * CosV);
+
+				OutPoint3D.GradientU = Matrix.MultiplyVector(OutPoint3D.GradientU);
+				OutPoint3D.GradientV = Matrix.MultiplyVector(OutPoint3D.GradientV);
+
+				if (InDerivativeOrder > 1)
+				{
+					OutPoint3D.LaplacianU = FPoint((MajorRadius + MinorRadius * CosV) * -CosU, (MajorRadius + MinorRadius * CosV) * -SinU, 0.0);
+					OutPoint3D.LaplacianV = FPoint((MinorRadius * -CosV) * CosU, (MinorRadius * -CosV) * SinU, MinorRadius * -SinV);
+					OutPoint3D.LaplacianUV = FPoint((MinorRadius * -SinV) * -SinU, (MinorRadius * -SinV) * CosU, 0.0);
+
+					OutPoint3D.LaplacianU = Matrix.MultiplyVector(OutPoint3D.LaplacianU);
+					OutPoint3D.LaplacianV = Matrix.MultiplyVector(OutPoint3D.LaplacianV);
+					OutPoint3D.LaplacianUV = Matrix.MultiplyVector(OutPoint3D.LaplacianUV);
+				}
+			}
+		}
+	
+		virtual void EvaluatePointGrid(const FCoordinateGrid& Coordinates, FSurfacicSampling& OutPoints, bool bComputeNormals = false) const override;
+
+		virtual void EvaluatePointGridInCylindricalSpace(const FCoordinateGrid& Coordinates, TArray<FPoint2D>&) const override;
+
+		virtual FPoint2D EvaluatePointInCylindricalSpace(const FPoint2D& InSurfacicCoordinate) const override
+		{
+			const double CosU = cos(InSurfacicCoordinate.U);
+			const double CosV = cos(InSurfacicCoordinate.V);
+
+			const double SinU = sin(InSurfacicCoordinate.U);
+
+			double Rho = (MajorRadius + MinorRadius * CosV);
+
+			double SwapOrientation = (InSurfacicCoordinate.V < DOUBLE_PI&& InSurfacicCoordinate.V >=0) ? 1.0 : -1.0;
+
+			return FPoint2D(Rho * CosU * SwapOrientation, Rho * SinU);
+		}
+
+		virtual void Presample(const FSurfacicBoundary& InBoundaries, FCoordinateGrid& OutCoordinates) override
+		{
+			PresampleIsoCircle(InBoundaries, OutCoordinates, EIso::IsoU);
+			PresampleIsoCircle(InBoundaries, OutCoordinates, EIso::IsoV);
+		}
+
+	};
+}
+
diff --git a/Engine/Source/Runtime/Datasmith/CADKernel/Private/CADKernel/Math/BSpline.cpp b/Engine/Source/Runtime/Datasmith/CADKernel/Private/CADKernel/Math/BSpline.cpp
--- a/Engine/Source/Runtime/Datasmith/CADKernel/Private/CADKernel/Math/BSpline.cpp
+++ b/Engine/Source/Runtime/Datasmith/CADKernel/Private/CADKernel/Math/BSpline.cpp
@@ -1414,7 +1414,7 @@ namespace CADKernel
 				// definir l'ordre de multiplicite du noeud
 				while (i < Degree + PoleCount)
 				{
-					if (FMath::Abs(NodalVector[i + 1] - udeb) > SMALL_NUMBER) break;
+					if (FMath::Abs(NodalVector[i + 1] - udeb) > DOUBLE_SMALL_NUMBER) break;
 					multiplicite++;
 					i++;
 				}
@@ -1467,12 +1467,12 @@ namespace CADKernel
 				}
 
 				// passer sur les noeuds multiples
-				if (FMath::Abs(nodalVector[i + 1] - nodalVector[i]) <SMALL_NUMBER) continue;
+				if (FMath::Abs(nodalVector[i + 1] - nodalVector[i]) <DOUBLE_SMALL_NUMBER) continue;
 				else uCourant = nodalVector[i];
 
 				// definir les bornes de l'intervalle a discretiser
-				if (uCourant > u1 + SMALL_NUMBER) u1 = uCourant;
-				if (nodalVector[i + 1] <u2 - SMALL_NUMBER) u2 = nodalVector[i + 1];
+				if (uCourant > u1 + DOUBLE_SMALL_NUMBER) u1 = uCourant;
+				if (nodalVector[i + 1] <u2 - DOUBLE_SMALL_NUMBER) u2 = nodalVector[i + 1];
 
 				// repartir les parametres dans l'intervalle courant*/
 				if (degre <= 1)
@@ -1498,7 +1498,7 @@ namespace CADKernel
 				}
 
 				// tester s'il s'agit du dernier intervalle
-				if (u2 > UMax - SMALL_NUMBER) break;
+				if (u2 > UMax - DOUBLE_SMALL_NUMBER) break;
 
 				// passer a l'intervalle suivant
 				u1 = u2;
diff --git a/Engine/Source/Runtime/Datasmith/CADKernel/Private/CADKernel/Math/Geometry.cpp b/Engine/Source/Runtime/Datasmith/CADKernel/Private/CADKernel/Math/Geometry.cpp
--- a/Engine/Source/Runtime/Datasmith/CADKernel/Private/CADKernel/Math/Geometry.cpp
+++ b/Engine/Source/Runtime/Datasmith/CADKernel/Private/CADKernel/Math/Geometry.cpp
@@ -72,8 +72,8 @@ namespace CADKernel
 
 	bool IntersectSegments2D(const TSegment<FPoint2D>& SegmentAB, const TSegment<FPoint2D>& SegmentCD)
 	{
-		constexpr const double Min = -SMALL_NUMBER;
-		constexpr const double Max = 1. + SMALL_NUMBER;
+		constexpr const double Min = -DOUBLE_SMALL_NUMBER;
+		constexpr const double Max = 1. + DOUBLE_SMALL_NUMBER;
 
 		TFunction<bool(double, double, double, double)> Intersect = [](double A, double B, double C, double D)
 		{
diff --git a/Engine/Source/Runtime/Datasmith/CADKernel/Private/CADKernel/Math/MatrixH.cpp b/Engine/Source/Runtime/Datasmith/CADKernel/Private/CADKernel/Math/MatrixH.cpp
--- a/Engine/Source/Runtime/Datasmith/CADKernel/Private/CADKernel/Math/MatrixH.cpp
+++ b/Engine/Source/Runtime/Datasmith/CADKernel/Private/CADKernel/Math/MatrixH.cpp
@@ -40,18 +40,18 @@ namespace CADKernel
 
 		//on cherche le vecteur ox
 		Xaxis.Normalize();
-		if (FMath::Abs(Xaxis.Length()) < SMALL_NUMBER)
+		if (FMath::Abs(Xaxis.Length()) < DOUBLE_SMALL_NUMBER)
 		{
 			Zaxis = FPoint(1, 0, 0);
 			Xaxis = Axis ^ Zaxis;
 			Xaxis.Normalize();
-			if (FMath::Abs(Xaxis.Length()) < SMALL_NUMBER)
+			if (FMath::Abs(Xaxis.Length()) < DOUBLE_SMALL_NUMBER)
 			{
 				Zaxis = FPoint(0, 0, 1);
 				Xaxis = Axis ^ Zaxis;
 
 				Xaxis.Normalize();
-				ensureCADKernel(FMath::Abs(Xaxis.Length()) > SMALL_NUMBER);
+				ensureCADKernel(FMath::Abs(Xaxis.Length()) > DOUBLE_SMALL_NUMBER);
 			}
 		}
 
@@ -69,7 +69,7 @@ namespace CADKernel
 		FMatrixH Matrix;
 		Matrix.SetIdentity();
 
-		ensureCADKernel(Axe.Length() > SMALL_NUMBER);
+		ensureCADKernel(Axe.Length() > DOUBLE_SMALL_NUMBER);
 		Axe.Normalize();
 
 		Matrix.Get(0, 0) = Axe[0] * Axe[0] + (double)cos(Angle) * ((double)1.0 - Axe[0] * Axe[0]);
diff --git a/Engine/Source/Runtime/Datasmith/CADKernel/Private/CADKernel/Math/Point.cpp b/Engine/Source/Runtime/Datasmith/CADKernel/Private/CADKernel/Math/Point.cpp
--- a/Engine/Source/Runtime/Datasmith/CADKernel/Private/CADKernel/Math/Point.cpp
+++ b/Engine/Source/Runtime/Datasmith/CADKernel/Private/CADKernel/Math/Point.cpp
@@ -37,14 +37,14 @@ double FPoint::SignedAngle(const FPoint & Other, const FPoint & Normal) const
 
 	double ScalarProduct = Vector1 * Vector2;
 
-	if (ScalarProduct >= 1 - SMALL_NUMBER)
+	if (ScalarProduct >= 1 - DOUBLE_SMALL_NUMBER)
 	{
 		return 0.;
 	}
 
-	if (ScalarProduct <= -1 + SMALL_NUMBER)
+	if (ScalarProduct <= -1 + DOUBLE_SMALL_NUMBER)
 	{
-		return PI;
+		return DOUBLE_PI;
 	}
 
 	return MixedTripleProduct(Vector1, Vector2, Vector3) > 0 ? acos(ScalarProduct) : -acos(ScalarProduct);
diff --git a/Engine/Source/Runtime/Datasmith/CADKernel/Private/CADKernel/Math/SlopeUtils.h b/Engine/Source/Runtime/Datasmith/CADKernel/Private/CADKernel/Math/SlopeUtils.h
--- a/Engine/Source/Runtime/Datasmith/CADKernel/Private/CADKernel/Math/SlopeUtils.h
+++ b/Engine/Source/Runtime/Datasmith/CADKernel/Private/CADKernel/Math/SlopeUtils.h
@@ -69,14 +69,14 @@ namespace CADKernel
 		double DeltaV = EndPoint.V - StartPoint.V;
 		double Delta;
 
-		if (FMath::Abs(DeltaU) < SMALL_NUMBER && FMath::Abs(DeltaV) < SMALL_NUMBER)
+		if (FMath::Abs(DeltaU) < DOUBLE_SMALL_NUMBER && FMath::Abs(DeltaV) < DOUBLE_SMALL_NUMBER)
 		{
 			return 0;
 		}
 
-		if (DeltaU > SMALL_NUMBER)
+		if (DeltaU > DOUBLE_SMALL_NUMBER)
 		{
-			if (DeltaV > SMALL_NUMBER)
+			if (DeltaV > DOUBLE_SMALL_NUMBER)
 			{
 				if (DeltaU > DeltaV)
 				{
@@ -101,7 +101,7 @@ namespace CADKernel
 		}
 		else
 		{
-			if (DeltaV > SMALL_NUMBER)
+			if (DeltaV > DOUBLE_SMALL_NUMBER)
 			{
 				if (-DeltaU > DeltaV)
 				{
@@ -277,7 +277,7 @@ namespace CADKernel
 	 *
 	 * Return true if all of the segment BPi is inside the sector defined the half-lines [BA) and [BC) in the counterclockwise.
 	 */
-	inline bool ArePointsInsideSectorABC(const FPoint2D& PointA, const FPoint2D& PointB, const FPoint2D& PointC, const TArray<const FPoint2D*>& Points, const double FlatAngle = -SMALL_NUMBER)
+	inline bool ArePointsInsideSectorABC(const FPoint2D& PointA, const FPoint2D& PointB, const FPoint2D& PointC, const TArray<const FPoint2D*>& Points, const double FlatAngle = -DOUBLE_SMALL_NUMBER)
 	{
 		double SlopWithNextBoundary = ComputeSlope(PointB, PointC);
 		double BoundaryDeltaSlope = ComputePositiveSlope(PointB, PointA, SlopWithNextBoundary);
@@ -286,7 +286,7 @@ namespace CADKernel
 			double DeltaU = PointB.U - PointP->U;
 			double DeltaV = PointB.V - PointP->V;
 
-			if (FMath::Abs(DeltaU) < SMALL_NUMBER && FMath::Abs(DeltaV) < SMALL_NUMBER)
+			if (FMath::Abs(DeltaU) < DOUBLE_SMALL_NUMBER && FMath::Abs(DeltaV) < DOUBLE_SMALL_NUMBER)
 			{
 				return true;
 			}
diff --git a/Engine/Source/Runtime/Datasmith/CADKernel/Private/CADKernel/Mesh/Criteria/CriteriaGrid.cpp b/Engine/Source/Runtime/Datasmith/CADKernel/Private/CADKernel/Mesh/Criteria/CriteriaGrid.cpp
--- a/Engine/Source/Runtime/Datasmith/CADKernel/Private/CADKernel/Mesh/Criteria/CriteriaGrid.cpp
+++ b/Engine/Source/Runtime/Datasmith/CADKernel/Private/CADKernel/Mesh/Criteria/CriteriaGrid.cpp
@@ -6,7 +6,10 @@
 #include "CADKernel/Topo/TopologicalFace.h"
 #include "CADKernel/Geo/Surfaces/Surface.h"
 #include "CADKernel/UI/Display.h"
+
+#ifdef CADKERNEL_DEV
 #include "CADKernel/Mesh/Meshers/IsoTriangulator/DefineForDebug.h"
+#endif
 
 namespace CADKernel
 {
@@ -44,18 +47,14 @@ void FCriteriaGrid::Init()
 	TrueUcoorindateCount = CoordinateGrid2[EIso::IsoU].Num();
 }
 
-FCriteriaGrid::FCriteriaGrid(FTopologicalFace& InSurface)
-	: Face(InSurface)
-	, CoordinateGrid(InSurface.GetCrossingPointCoordinates())
+FCriteriaGrid::FCriteriaGrid(FTopologicalFace& InFace)
+	: Face(InFace)
+	, CoordinateGrid(InFace.GetCrossingPointCoordinates())
 {
 	Face.Presample();
 
-	const FCoordinateGrid& FaceGrid = Face.GetCrossingPointCoordinates();
-
-	// If the size of the grid is very huge for a pre-sample step means that the surface is very poorly build
-	// Canceled the face meshing is preferred than waiting a too long time for a dirty result.
 	constexpr int32 MaxGrid = 1000000;
-	if (FaceGrid[EIso::IsoU].Num() * FaceGrid[EIso::IsoV].Num() > MaxGrid)
+	if (CoordinateGrid[EIso::IsoU].Num() * CoordinateGrid[EIso::IsoV].Num() > MaxGrid)
 	{
 		Face.SetDeleted();
 		return;
diff --git a/Engine/Source/Runtime/Datasmith/CADKernel/Private/CADKernel/Mesh/Criteria/Criterion.cpp b/Engine/Source/Runtime/Datasmith/CADKernel/Private/CADKernel/Mesh/Criteria/Criterion.cpp
--- a/Engine/Source/Runtime/Datasmith/CADKernel/Private/CADKernel/Mesh/Criteria/Criterion.cpp
+++ b/Engine/Source/Runtime/Datasmith/CADKernel/Private/CADKernel/Mesh/Criteria/Criterion.cpp
@@ -39,14 +39,14 @@ namespace CADKernel
 
 			double DeltaU = Coordinates[Index] - Coordinates[PreviousIndex];
 
-			if (ChordLength > NumericPrecision && Sag > SMALL_NUMBER)
+			if (ChordLength > NumericPrecision && Sag > DOUBLE_SMALL_NUMBER)
 			{
 				double DeltaUMax = ComputeDeltaU(ChordLength, DeltaU, Sag);
 				if (DeltaUMax < DeltaUMaxs[PreviousIndex]) DeltaUMaxs[PreviousIndex] = DeltaUMax;
 			}
 
 			Sag = EvaluateSag(Points[2 * Index].Point, Points[2 * PreviousIndex].Point, Points[MiddleIndex].Point, ChordLength);
-			if (ChordLength > NumericPrecision && Sag > SMALL_NUMBER)
+			if (ChordLength > NumericPrecision && Sag > DOUBLE_SMALL_NUMBER)
 			{
 				double DeltaUMax = ComputeDeltaU(ChordLength, DeltaU, Sag);
 				if (DeltaUMax < DeltaUMaxs[PreviousIndex]) 
diff --git a/Engine/Source/Runtime/Datasmith/CADKernel/Private/CADKernel/Mesh/Criteria/CurvatureCriterion.h b/Engine/Source/Runtime/Datasmith/CADKernel/Private/CADKernel/Mesh/Criteria/CurvatureCriterion.h
--- a/Engine/Source/Runtime/Datasmith/CADKernel/Private/CADKernel/Mesh/Criteria/CurvatureCriterion.h
+++ b/Engine/Source/Runtime/Datasmith/CADKernel/Private/CADKernel/Mesh/Criteria/CurvatureCriterion.h
@@ -29,7 +29,7 @@ namespace CADKernel
 		 */
 		virtual void UpdateDelta(double InDeltaU, double InUSag, double InDiagonalSag, double InVSag, double ChordLength, double DiagonalLength, double& OutSagDeltaUMax, double& OutSagDeltaUMin, FIsoCurvature& IsoCurvature) const override
 		{
-			if (DiagonalLength > SMALL_NUMBER)
+			if (DiagonalLength > DOUBLE_SMALL_NUMBER)
 			{
 				double Curvature = 8 * InUSag / FMath::Square(DiagonalLength);
 				if (IsoCurvature.Max < Curvature)
diff --git a/Engine/Source/Runtime/Datasmith/CADKernel/Private/CADKernel/Mesh/Criteria/SizeCriterion.cpp b/Engine/Source/Runtime/Datasmith/CADKernel/Private/CADKernel/Mesh/Criteria/SizeCriterion.cpp
--- a/Engine/Source/Runtime/Datasmith/CADKernel/Private/CADKernel/Mesh/Criteria/SizeCriterion.cpp
+++ b/Engine/Source/Runtime/Datasmith/CADKernel/Private/CADKernel/Mesh/Criteria/SizeCriterion.cpp
@@ -60,7 +60,7 @@ void FSizeCriterion::ApplyOnParameters(const TArray<double>& Coordinates, const
 
 void FMinSizeCriterion::UpdateDelta(double InDeltaU, double InUSag, double InDiagonalSag, double InVSag, double ChordLength, double DiagonalLength, double& OutSagDeltaUMax, double& OutSagDeltaUMin, FIsoCurvature& SurfaceCurvature) const
 {
-	if (ChordLength < KINDA_SMALL_NUMBER)
+	if (ChordLength < DOUBLE_KINDA_SMALL_NUMBER)
 	{
 		return;
 	}
@@ -74,7 +74,7 @@ void FMinSizeCriterion::UpdateDelta(double InDeltaU, double InUSag, double InDia
 
 void FMaxSizeCriterion::UpdateDelta(double InDeltaU, double InUSag, double InDiagonalSag, double InVSag, double ChordLength, double DiagonalLength, double& OutSagDeltaUMax, double& OutSagDeltaUMin, FIsoCurvature& SurfaceCurvature) const
 {
-	if (ChordLength < KINDA_SMALL_NUMBER)
+	if (ChordLength < DOUBLE_KINDA_SMALL_NUMBER)
 	{
 		return;
 	}
diff --git a/Engine/Source/Runtime/Datasmith/CADKernel/Private/CADKernel/Mesh/Meshers/BowyerWatsonTriangulator.h b/Engine/Source/Runtime/Datasmith/CADKernel/Private/CADKernel/Mesh/Meshers/BowyerWatsonTriangulator.h
--- a/Engine/Source/Runtime/Datasmith/CADKernel/Private/CADKernel/Mesh/Meshers/BowyerWatsonTriangulator.h
+++ b/Engine/Source/Runtime/Datasmith/CADKernel/Private/CADKernel/Mesh/Meshers/BowyerWatsonTriangulator.h
@@ -5,7 +5,10 @@
 #include "CADKernel/Math/SlopeUtils.h"
 #include "CADKernel/Mesh/Meshers/IsoTriangulator/IsoCell.h"
 #include "CADKernel/UI/Display.h"
+
+#ifdef CADKERNEL_DEV
 #include "CADKernel/Mesh/Meshers/IsoTriangulator/DefineForDebug.h"
+#endif
 
 //#define DEBUG_BOWYERWATSON
 namespace CADKernel
diff --git a/Engine/Source/Runtime/Datasmith/CADKernel/Private/CADKernel/Mesh/Meshers/IsoTriangulator.cpp b/Engine/Source/Runtime/Datasmith/CADKernel/Private/CADKernel/Mesh/Meshers/IsoTriangulator.cpp
--- a/Engine/Source/Runtime/Datasmith/CADKernel/Private/CADKernel/Mesh/Meshers/IsoTriangulator.cpp
+++ b/Engine/Source/Runtime/Datasmith/CADKernel/Private/CADKernel/Mesh/Meshers/IsoTriangulator.cpp
@@ -24,7 +24,11 @@ namespace CADKernel
 
 const double FIsoTriangulator::GeometricToMeshingToleranceFactor = 10.;
 
-FIsoTriangulator::FIsoTriangulator(FGrid& InGrid, TSharedRef<FFaceMesh> EntityMesh, FMesherReport& InMesherReport)
+#ifdef DEBUG_BOWYERWATSON
+bool FBowyerWatsonTriangulator::bDisplay = false;
+#endif
+
+FIsoTriangulator::FIsoTriangulator(FGrid& InGrid, TSharedRef<FFaceMesh> EntityMesh)
 	: Grid(InGrid)
 	, Mesh(EntityMesh)
 	, LoopSegmentsIntersectionTool(InGrid)
@@ -36,7 +40,6 @@ FIsoTriangulator::FIsoTriangulator(FGrid& InGrid, TSharedRef<FFaceMesh> EntityMe
 	, SquareGeometricTolerance2(2. * SquareGeometricTolerance)
 	, MeshingTolerance(GeometricTolerance* GeometricToMeshingToleranceFactor)
 	, SquareMeshingTolerance(FMath::Square(MeshingTolerance))
-	, MesherReport(InMesherReport)
 {
 	FinalInnerSegments.Reserve(3 * Grid.InnerNodesCount());
 	IndexOfLowerLeftInnerNodeSurroundingALoop.Reserve(Grid.GetLoopCount());
@@ -82,7 +85,9 @@ bool FIsoTriangulator::Triangulate()
 	FLoopCleaner LoopCleaner(*this);
 	if (!LoopCleaner.Run())
 	{
-		MesherReport.Logs.AddDegeneratedLoop();
+#ifdef CADKERNEL_DEV
+		MesherReport->Logs.AddDegeneratedLoop();
+#endif
 		FMessage::Printf(EVerboseLevel::Log, TEXT("The meshing of the surface %d failed due to a degenerated loop\n"), Grid.GetFace().GetId());
 		return false;
 	}
@@ -463,7 +468,7 @@ void FIsoTriangulator::BuildInnerSegments()
 				{
 					if (FMath::IsNearlyEqual(LoopPoint.V, Node1.V, GeometricTolerance))
 					{
-						if (Node1.U - SMALL_NUMBER < LoopPoint.U && LoopPoint.U < Node2.U + SMALL_NUMBER)
+						if (Node1.U - DOUBLE_SMALL_NUMBER < LoopPoint.U && LoopPoint.U < Node2.U + DOUBLE_SMALL_NUMBER)
 						{
 							return true;
 						}
@@ -479,7 +484,7 @@ void FIsoTriangulator::BuildInnerSegments()
 				{
 					if (FMath::IsNearlyEqual(LoopPoint.U, Node1.U, GeometricTolerance))
 					{
-						if (Node1.V - SMALL_NUMBER < LoopPoint.V && LoopPoint.V < Node2.V + SMALL_NUMBER)
+						if (Node1.V - DOUBLE_SMALL_NUMBER < LoopPoint.V && LoopPoint.V < Node2.V + DOUBLE_SMALL_NUMBER)
 						{
 							return true;
 						}
@@ -1590,7 +1595,7 @@ void FIsoTriangulator::FindIsoCandidateSegmentInCycle(TArray<FIsoNode*> CycleNod
 #endif
 
 	const double FlatSlope = 0.10; // ~5 deg: The segment must make an angle less than 10 deg with the Iso
-	double MinSlope = FlatSlope + SMALL_NUMBER; //.25; // ~10 deg: The segment must make an angle less than 10 deg with the Iso
+	double MinSlope = FlatSlope + DOUBLE_SMALL_NUMBER; //.25; // ~10 deg: The segment must make an angle less than 10 deg with the Iso
 
 	FIsoNode* CandidateA = nullptr;
 	FIsoNode* CandidateB = nullptr;
@@ -1638,7 +1643,7 @@ void FIsoTriangulator::FindIsoCandidateSegmentInCycle(TArray<FIsoNode*> CycleNod
 			Grid.DisplayIsoSegment(UniformScaled, Segment, 0, EVisuProperty::YellowCurve);
 			Wait(bDisplay);
 #endif
-			MinSlope = FlatSlope + SMALL_NUMBER;
+			MinSlope = FlatSlope + DOUBLE_SMALL_NUMBER;
 		}
 	}
 
@@ -1777,7 +1782,9 @@ void FIsoTriangulator::MeshCycle(const EGridSpace Space, const TArray<FIsoSegmen
 	// Check if the cycle is in self intersecting and fix it. 
 	if (!CanCycleBeMeshed(Cycle, CycleIntersectionTool))
 	{
-		MesherReport.Logs.AddCycleMeshingFailure();
+#ifdef CADKERNEL_DEV
+		MesherReport->Logs.AddCycleMeshingFailure();
+#endif
 		return;
 	}
 
@@ -2002,7 +2009,7 @@ void FIsoTriangulator::MeshCycle(const EGridSpace Space, const TArray<FIsoSegmen
 			if (&EndNode != &PreviousNode)
 			{
 				// Case of probable auto-intersection cycle at PreviousNode, cancel is preferred 
-				if (FMath::IsNearlyEqual(StartMaxSlope, 8., (double)KINDA_SMALL_NUMBER))
+				if (FMath::IsNearlyEqual(StartMaxSlope, 8., DOUBLE_KINDA_SMALL_NUMBER))
 				{
 					return;
 				}
@@ -2014,7 +2021,7 @@ void FIsoTriangulator::MeshCycle(const EGridSpace Space, const TArray<FIsoSegmen
 			if (&StartNode != &NextNode)
 			{
 				// Case of probable auto-intersection cycle at PreviousNode, cancel is preferred 
-				if (FMath::IsNearlyEqual(EndMaxSlope, 8., (double)KINDA_SMALL_NUMBER))
+				if (FMath::IsNearlyEqual(EndMaxSlope, 8., DOUBLE_KINDA_SMALL_NUMBER))
 				{
 					return;
 				}
@@ -2087,7 +2094,7 @@ void FIsoTriangulator::MeshCycle(const EGridSpace Space, const TArray<FIsoSegmen
 				double CosAngle = FMath::Abs(ComputeCosinus(NodeNormal, Node->GetNormal(Grid)));
 
 				// the criteria is weighted according to the cosine of the angle between the normal of the candidate triangle and the normal at the tested point
-				if (CosAngle > SMALL_NUMBER)
+				if (CosAngle > DOUBLE_SMALL_NUMBER)
 				{
 					PointCriteria /= CosAngle;
 				}
@@ -2255,7 +2262,7 @@ void FIsoTriangulator::FindInnerGridCellSurroundingSmallLoop()
 		int32 IndexU = 0;
 		for (; IndexU < NumU - 1; ++IndexU)
 		{
-			if ((FirstPoint.U > UCoordinates[IndexU]) && (FirstPoint.U < UCoordinates[IndexU + 1] + SMALL_NUMBER))
+			if ((FirstPoint.U > UCoordinates[IndexU]) && (FirstPoint.U < UCoordinates[IndexU + 1] + DOUBLE_SMALL_NUMBER))
 			{
 				break;
 			}
@@ -2264,16 +2271,16 @@ void FIsoTriangulator::FindInnerGridCellSurroundingSmallLoop()
 		int32 IndexV = 0;
 		for (; IndexV < NumV - 1; ++IndexV)
 		{
-			if ((FirstPoint.V > VCoordinates[IndexV]) && (FirstPoint.V < VCoordinates[IndexV + 1] + SMALL_NUMBER))
+			if ((FirstPoint.V > VCoordinates[IndexV]) && (FirstPoint.V < VCoordinates[IndexV + 1] + DOUBLE_SMALL_NUMBER))
 			{
 				break;
 			}
 		}
 
 		double UMin = UCoordinates[IndexU];
-		double UMax = UCoordinates[IndexU + 1] + SMALL_NUMBER;
+		double UMax = UCoordinates[IndexU + 1] + DOUBLE_SMALL_NUMBER;
 		double VMin = VCoordinates[IndexV];
-		double VMax = VCoordinates[IndexV + 1] + SMALL_NUMBER;
+		double VMax = VCoordinates[IndexV + 1] + DOUBLE_SMALL_NUMBER;
 
 		bool bBoudardyIsSurrounded = true;
 		for (const FPoint2D& LoopPoint : Loops[LoopIndex])
@@ -3057,7 +3064,7 @@ void FIsoTriangulator::TryToConnectVertexSubLoopWithTheMostIsoSegment(FCell& Cel
 #endif
 
 	const double FlatSlope = 0.10; // ~5 deg: The segment must make an angle less than 10 deg with the Iso
-	double MinSlope = FlatSlope + SMALL_NUMBER; //.25; // ~10 deg: The segment must make an angle less than 10 deg with the Iso
+	double MinSlope = FlatSlope + DOUBLE_SMALL_NUMBER; //.25; // ~10 deg: The segment must make an angle less than 10 deg with the Iso
 
 
 	if (Loop.Num() <= 2)
@@ -3116,7 +3123,7 @@ void FIsoTriangulator::TryToConnectVertexSubLoopWithTheMostIsoSegment(FCell& Cel
 			Grid.DisplayIsoSegment(UniformScaled, *CandidateA, *CandidateB, 0, EVisuProperty::RedCurve);
 #endif
 			TryToCreateSegment(Cell, CandidateA, ACoordinates, CandidateB, BCoordinates, 0.1);
-			MinSlope = FlatSlope + SMALL_NUMBER;
+			MinSlope = FlatSlope + DOUBLE_SMALL_NUMBER;
 		}
 	}
 
@@ -3133,7 +3140,7 @@ void FIsoTriangulator::TryToConnectTwoSubLoopsWithTheMostIsoSegment(FCell& Cell,
 #endif
 
 	const double FlatSlope = 0.10; // ~5 deg: The segment must make an angle less than 10 deg with the Iso
-	double MinSlope = FlatSlope + SMALL_NUMBER;// 0.25; // ~15 deg: The segment must make an angle less than 10 deg with the Iso
+	double MinSlope = FlatSlope + DOUBLE_SMALL_NUMBER;// 0.25; // ~15 deg: The segment must make an angle less than 10 deg with the Iso
 
 	for (FLoopNode* CandidateA : LoopA)
 	{
@@ -3160,7 +3167,7 @@ void FIsoTriangulator::TryToConnectTwoSubLoopsWithTheMostIsoSegment(FCell& Cell,
 			Grid.DisplayIsoSegment(EGridSpace::UniformScaled, *CandidateA, *CandidateB, 0, EVisuProperty::BlueCurve);
 #endif			
 			TryToCreateSegment(Cell, CandidateA, ACoordinates, CandidateB, BCoordinates, 0.1);
-			MinSlope = FlatSlope + SMALL_NUMBER;
+			MinSlope = FlatSlope + DOUBLE_SMALL_NUMBER;
 		}
 	}
 }
diff --git a/Engine/Source/Runtime/Datasmith/CADKernel/Private/CADKernel/Mesh/Meshers/IsoTriangulator.h b/Engine/Source/Runtime/Datasmith/CADKernel/Private/CADKernel/Mesh/Meshers/IsoTriangulator.h
--- a/Engine/Source/Runtime/Datasmith/CADKernel/Private/CADKernel/Mesh/Meshers/IsoTriangulator.h
+++ b/Engine/Source/Runtime/Datasmith/CADKernel/Private/CADKernel/Mesh/Meshers/IsoTriangulator.h
@@ -8,10 +8,12 @@
 #include "CADKernel/Mesh/Meshers/IsoTriangulator/IntersectionSegmentTool.h"
 #include "CADKernel/Mesh/Meshers/IsoTriangulator/IsoNode.h"
 #include "CADKernel/Mesh/Meshers/IsoTriangulator/IsoSegment.h"
-#include "CADKernel/Mesh/Meshers/MesherReport.h"
 #include "CADKernel/UI/Visu.h"
 
+#ifdef CADKERNEL_DEV
 #include "CADKernel/Mesh/Meshers/IsoTriangulator/DefineForDebug.h"
+#include "CADKernel/Mesh/Meshers/MesherReport.h"
+#endif
 
 //#define NEED_TO_CHECK_USEFULNESS
 //#define DEBUG_DELAUNAY
@@ -160,11 +162,20 @@ protected:
 	const double MeshingTolerance;
 	const double SquareMeshingTolerance;
 
-	FMesherReport& MesherReport;
+#ifdef CADKERNEL_DEV
+	FMesherReport* MesherReport;
+#endif
 
 public:
 
-	FIsoTriangulator(FGrid& InGrid, TSharedRef<FFaceMesh> EntityMesh, FMesherReport& InMesherReport);
+	FIsoTriangulator(FGrid& InGrid, TSharedRef<FFaceMesh> EntityMesh);
+
+#ifdef CADKERNEL_DEV
+	void SetMesherReport(FMesherReport& InMesherReport)
+	{
+		MesherReport = &InMesherReport;
+	}
+#endif
 
 	/**
 	 * Main method
@@ -408,7 +419,7 @@ inline double CotangentCriteria(const FPoint& APoint, const FPoint& BPoint, cons
 	double NormOFScalarProduct = sqrt(OutNormal * OutNormal);
 
 	// PPoint is aligned with (A,B)
-	if (NormOFScalarProduct < SMALL_NUMBER)
+	if (NormOFScalarProduct < DOUBLE_SMALL_NUMBER)
 	{
 		return BigValue;
 	}
@@ -431,7 +442,7 @@ inline double CotangentCriteria(const FPoint2D& APoint, const FPoint2D& BPoint,
 	double OutNormal = PA ^ PB;
 	double NormOFPointProduct = FMath::Abs(OutNormal);
 
-	if (NormOFPointProduct < SMALL_NUMBER)
+	if (NormOFPointProduct < DOUBLE_SMALL_NUMBER)
 	{
 		// PPoint is aligned with (A,B)
 		return BigValue;
diff --git a/Engine/Source/Runtime/Datasmith/CADKernel/Private/CADKernel/Mesh/Meshers/IsoTriangulator/DefineForDebug.h b/Engine/Source/Runtime/Datasmith/CADKernel/Private/CADKernel/Mesh/Meshers/IsoTriangulator/DefineForDebug.h
deleted file mode 100644
diff --git a/Engine/Source/Runtime/Datasmith/CADKernel/Private/CADKernel/Mesh/Meshers/IsoTriangulator/IsoTriangulatorDisplay.cpp b/Engine/Source/Runtime/Datasmith/CADKernel/Private/CADKernel/Mesh/Meshers/IsoTriangulator/IsoTriangulatorDisplay.cpp
deleted file mode 100644
diff --git a/Engine/Source/Runtime/Datasmith/CADKernel/Private/CADKernel/Mesh/Meshers/MesherTools.cpp b/Engine/Source/Runtime/Datasmith/CADKernel/Private/CADKernel/Mesh/Meshers/MesherTools.cpp
--- a/Engine/Source/Runtime/Datasmith/CADKernel/Private/CADKernel/Mesh/Meshers/MesherTools.cpp
+++ b/Engine/Source/Runtime/Datasmith/CADKernel/Private/CADKernel/Mesh/Meshers/MesherTools.cpp
@@ -40,7 +40,7 @@ namespace CADKernel
 			TravelTime += (NextStep - LastStep) / DeltaUs[CrossingIndex];
 			CrossingIndex++;
 			LastStep = NextStep;
-		} while (LastStep + SMALL_NUMBER < Boundary.Max);
+		} while (LastStep + DOUBLE_SMALL_NUMBER < Boundary.Max);
 
 
 		// Step 2: Define the number of step(each step lasts one) as it must be an integer bigger than 1
@@ -69,7 +69,7 @@ namespace CADKernel
 				for (; UIndex < DeltaUs.Num(); ++UIndex)
 				{
 					double DeltaU = DeltaUs[UIndex];
-					if (PreferredCuttingPoints[NeighbourIndex].Coordinate - SMALL_NUMBER > CrossingUs[UIndex + 1])
+					if (PreferredCuttingPoints[NeighbourIndex].Coordinate - DOUBLE_SMALL_NUMBER > CrossingUs[UIndex + 1])
 					{
 						double DeltaTravelTime = (CrossingUs[UIndex + 1] - LastStep) / DeltaU;
 						LastStep = CrossingUs[UIndex + 1];
@@ -77,7 +77,7 @@ namespace CADKernel
 						continue;
 					}
 
-					if (CrossingUs[UIndex] - SMALL_NUMBER < PreferredCuttingPoints[NeighbourIndex].Coordinate && PreferredCuttingPoints[NeighbourIndex].Coordinate < CrossingUs[UIndex + 1] + SMALL_NUMBER)
+					if (CrossingUs[UIndex] - DOUBLE_SMALL_NUMBER < PreferredCuttingPoints[NeighbourIndex].Coordinate && PreferredCuttingPoints[NeighbourIndex].Coordinate < CrossingUs[UIndex + 1] + DOUBLE_SMALL_NUMBER)
 					{
 						double DeltaTravelTime = (PreferredCuttingPoints[NeighbourIndex].Coordinate - LastStep) / DeltaU;
 						double TravelTimeToNeighbourU = TravelTime + DeltaTravelTime;
@@ -116,7 +116,7 @@ namespace CADKernel
 		// Define the next Target time step i.e. a Preferred Cutting Points or a round step   
 		TFunction<void(double&, double&)> IncreaseTargetTimeAccordingToNeighbors = [&](double& TargetTime, double& TimeDelta)
 		{
-			if (NeighbourIndex < TravelTimeOfNeighborsU.Num() && TargetTime + TimeDelta + SMALL_NUMBER > TravelTimeOfNeighborsU[NeighbourIndex])
+			if (NeighbourIndex < TravelTimeOfNeighborsU.Num() && TargetTime + TimeDelta + DOUBLE_SMALL_NUMBER > TravelTimeOfNeighborsU[NeighbourIndex])
 			{
 				TimeDelta = CorrectedFinalTravelTime - TargetTime;
 				for (; NeighbourIndex < TravelTimeOfNeighborsU.Num(); ++NeighbourIndex)
@@ -168,7 +168,7 @@ namespace CADKernel
 			double& DeltaU = DeltaUs[CrossingIndex];
 			double DeltaTravelTime = (CrossingUs[CrossingIndex + 1] - LastStep) / DeltaU;
 			double MaxTravelTime = TravelTime + DeltaTravelTime;
-			while (TimeDelta > SMALL_NUMBER && MaxTravelTime > TargetTime - SMALL_NUMBER)
+			while (TimeDelta > DOUBLE_SMALL_NUMBER && MaxTravelTime > TargetTime - DOUBLE_SMALL_NUMBER)
 			{
 				DeltaTravelTime = TargetTime - TravelTime;
 				LastStep += DeltaU * DeltaTravelTime;
@@ -178,7 +178,7 @@ namespace CADKernel
 				IncreaseTargetTimeAccordingToNeighbors(TargetTime, TimeDelta);
 			}
 			
-			if (OutCuttingPoints.Last() + SMALL_NUMBER > Boundary.GetMax())
+			if (OutCuttingPoints.Last() + DOUBLE_SMALL_NUMBER > Boundary.GetMax())
 			{
 				break;
 			}
@@ -229,7 +229,7 @@ namespace CADKernel
 				TravelTimeByStep.Add(StepTime);
 				TravelTime += StepTime;
 				LastStep = Step;
-				if (LastStep + SMALL_NUMBER > EndU)
+				if (LastStep + DOUBLE_SMALL_NUMBER > EndU)
 				{
 					break;
 				}
@@ -250,12 +250,12 @@ namespace CADKernel
 			{
 				double CorrectedStepTime = StepTime / CorrectiveFactor;
 				double CorrectedDeltaU = DeltaUs[CrossingIndex] * CorrectiveFactor;
-				while (TravelTime + CorrectedStepTime + SMALL_NUMBER > Step)
+				while (TravelTime + CorrectedStepTime + DOUBLE_SMALL_NUMBER > Step)
 				{
 					double Time = Step - TravelTime;
 					OutCuttingPointUs.Add(TravelLength + CorrectedDeltaU * Time);
 					Step++;
-					if (Step + SMALL_NUMBER > NbStep)
+					if (Step + DOUBLE_SMALL_NUMBER > NbStep)
 					{
 						OutCuttingPointUs.Add(EndU);
 						return;
diff --git a/Engine/Source/Runtime/Datasmith/CADKernel/Private/CADKernel/Mesh/Meshers/ParametricMesher.cpp b/Engine/Source/Runtime/Datasmith/CADKernel/Private/CADKernel/Mesh/Meshers/ParametricMesher.cpp
--- a/Engine/Source/Runtime/Datasmith/CADKernel/Private/CADKernel/Mesh/Meshers/ParametricMesher.cpp
+++ b/Engine/Source/Runtime/Datasmith/CADKernel/Private/CADKernel/Mesh/Meshers/ParametricMesher.cpp
@@ -2,7 +2,6 @@
 
 #include "CADKernel/Mesh/Meshers/ParametricMesher.h"
 
-#include "CADKernel/Core/KernelParameters.h"
 #include "CADKernel/Mesh/Criteria/CriteriaGrid.h"
 #include "CADKernel/Mesh/Criteria/Criterion.h"
 #include "CADKernel/Mesh/Meshers/IsoTriangulator.h"
@@ -17,15 +16,15 @@
 #include "CADKernel/Topo/TopologicalFace.h"
 #include "CADKernel/Utils/Util.h"
 
+#ifdef CADKERNEL_DEV
+#include "CADKernel/Mesh/Meshers/MesherReport.h"
+#endif
+
 //#define DEBUG_MESH_EDGE
 //#define DEBUG_GETPREFERREDUVCOORDINATESFROMNEIGHBOURS
 namespace CADKernel
 {
 
-FMesherParameters::FMesherParameters()
-	: InconsistencyAngle(TEXT("inconsistencyAngle"), 20., *this)
-{}
-
 FParametricMesher::FParametricMesher(FModelMesh& InMeshModel)
 	: MeshModel(InMeshModel)
 {
@@ -86,8 +85,6 @@ void FParametricMesher::MeshEntities(TArray<FTopologicalShapeEntity*>& InEntitie
 	}
 
 	MeshEntities();
-
-	MesherReport.Print();
 }
 
 void FParametricMesher::MeshEntities()
@@ -113,7 +110,9 @@ void FParametricMesher::MeshEntities()
 		ApplyFaceCriteria(*Face);
 	}
 
+#ifdef CADKERNEL_DEV
 	MesherReport.Chronos.ApplyCriteriaDuration = FChrono::Elapse(ApplyCriteriaStartTime);
+#endif
 
 	FTimePoint MeshingStartTime = FChrono::Now();
 
@@ -132,7 +131,10 @@ void FParametricMesher::MeshEntities()
 
 		IsolateQuadFace(QuadTrimmedSurfaceSet, OtherEntities);
 
+#ifdef CADKERNEL_DEV
 		MesherReport.Chronos.IsolateQuadPatchDuration = FChrono::Elapse(IsolateQuadPatchStartTime);
+#endif
+
 		FMessage::Printf(Log, TEXT("  %d Quad Surfaces found\n"), QuadTrimmedSurfaceSet.Num());
 	}
 
@@ -142,8 +144,10 @@ void FParametricMesher::MeshEntities()
 
 	FTimePoint MeshStartTime = FChrono::Now();
 	MeshSurfaceByFront(QuadTrimmedSurfaceSet);
+#ifdef CADKERNEL_DEV
 	MesherReport.Chronos.GlobalMeshDuration = FChrono::Elapse(MeshStartTime);
 	MesherReport.Chronos.GlobalDuration = FChrono::Elapse(StartTime);
+#endif
 
 	//Chronos.PrintTimeElapse();
 }
@@ -229,7 +233,9 @@ void FParametricMesher::Mesh(FTopologicalFace& Face)
 
 	if (Grid.IsDegenerated())
 	{
+#ifdef CADKERNEL_DEV
 		MesherReport.Logs.AddDegeneratedGrid();
+ #endif
 		FMessage::Printf(EVerboseLevel::Log, TEXT("The meshing of the surface %d failed due to a degenerated grid\n"), Face.GetId());
 		Face.SetMeshed();
 		return;
@@ -239,7 +245,11 @@ void FParametricMesher::Mesh(FTopologicalFace& Face)
 
 	TSharedRef<FFaceMesh> SurfaceMesh = StaticCastSharedRef<FFaceMesh>(Face.GetOrCreateMesh(MeshModel));
 
-	FIsoTriangulator IsoTrianguler(Grid, SurfaceMesh, MesherReport);
+	FIsoTriangulator IsoTrianguler(Grid, SurfaceMesh);
+#ifdef CADKERNEL_DEV
+	MesherReport
+#endif
+
 	if (IsoTrianguler.Triangulate())
 	{
 		if (Face.IsBackOriented())
@@ -289,9 +299,13 @@ void FParametricMesher::GenerateCloud(FGrid& Grid)
 #endif
 			FTimePoint MeshThinZonesTime = FChrono::Now();
 			MeshThinZoneEdges(Grid);
+#ifdef CADKERNEL_DEV
 			MesherReport.Chronos.GlobalMeshThinZones += FChrono::Elapse(MeshThinZonesTime);
+#endif
 		}
+#ifdef CADKERNEL_DEV
 		MesherReport.Chronos.GlobalThinZones += FChrono::Elapse(StartTime);
+#endif
 	}
 
 #ifdef DEBUG_THIN_ZONES
@@ -304,7 +318,9 @@ void FParametricMesher::GenerateCloud(FGrid& Grid)
 
 	Grid.ProcessPointCloud();
 
+#ifdef CADKERNEL_DEV
 	MesherReport.Chronos.GlobalMeshAndGetLoopNodes += FChrono::Elapse(StartTime);
+#endif
 }
 
 void FParametricMesher::MeshFaceLoops(FGrid& Grid)
@@ -321,7 +337,9 @@ void FParametricMesher::MeshFaceLoops(FGrid& Grid)
 		}
 	}
 
+#ifdef CADKERNEL_DEV
 	MesherReport.Chronos.GlobalMeshEdges += FChrono::Elapse(StartTime);
+#endif
 }
 
 static void FillImposedIsoCuttingPoints(TArray<double>& UEdgeSetOfIntersectionWithIso, ECoordinateType CoordinateType, double EdgeToleranceGeo, const FTopologicalEdge& Edge, TArray<FCuttingPoint>& OutImposedIsoVertexSet)
@@ -359,7 +377,7 @@ static void FillImposedIsoCuttingPoints(TArray<double>& UEdgeSetOfIntersectionWi
 	case 1:
 	{
 		int32 CuttingPointIndex = 0;
-		while (CuttingPointIndex < Edge.GetCrossingPointUs().Num() && Edge.GetCrossingPointUs()[CuttingPointIndex] + SMALL_NUMBER <= OutImposedIsoVertexSet[StartIndex].Coordinate)
+		while (CuttingPointIndex < Edge.GetCrossingPointUs().Num() && Edge.GetCrossingPointUs()[CuttingPointIndex] + DOUBLE_SMALL_NUMBER <= OutImposedIsoVertexSet[StartIndex].Coordinate)
 		{
 			++CuttingPointIndex;
 		};
@@ -1370,7 +1388,9 @@ void FParametricMesher::MeshThinZoneEdges(FGrid& Grid)
 	}
 	Close3DDebugSession();
 #endif
+#ifdef CADKERNEL_DEV
 	MesherReport.Chronos.GlobalMeshThinZones += FChrono::Elapse(MeshStartTime);
+#endif
 
 }
 
@@ -1616,11 +1636,6 @@ void FParametricMesher::MeshThinZoneSide(const FThinZoneSide& Side)
 
 }
 
-void FParametricMesher::PrintReport()
-{
-	MesherReport.Print();
-}
-
 #ifdef DEBUG_INTERSECTEDGEISOS
 TMap<int32, int32> SurfaceDrawed;
 bool bDisplayIsoCurve = true;
diff --git a/Engine/Source/Runtime/Datasmith/CADKernel/Private/CADKernel/Mesh/Structure/Grid.cpp b/Engine/Source/Runtime/Datasmith/CADKernel/Private/CADKernel/Mesh/Structure/Grid.cpp
--- a/Engine/Source/Runtime/Datasmith/CADKernel/Private/CADKernel/Mesh/Structure/Grid.cpp
+++ b/Engine/Source/Runtime/Datasmith/CADKernel/Private/CADKernel/Mesh/Structure/Grid.cpp
@@ -422,7 +422,7 @@ void FGrid::FindPointsCloseToLoop()
 		Index[Iso] = 1;
 		for (; Index[Iso] < CuttingCount[Iso] - 1; ++Index[Iso])
 		{
-			if (UniformCuttingCoordinates[Iso][Index[Iso]] + SMALL_NUMBER > (*PointA)[Iso])
+			if (UniformCuttingCoordinates[Iso][Index[Iso]] + DOUBLE_SMALL_NUMBER > (*PointA)[Iso])
 			{
 				break;
 			}
@@ -480,12 +480,12 @@ void FGrid::FindPointsCloseToLoop()
 
 	TFunction<bool(const double, const double)> IsReallyBigger = [](const double FirstValue, const double SecondValue) ->bool
 	{
-		return FirstValue - SMALL_NUMBER > SecondValue;
+		return FirstValue - DOUBLE_SMALL_NUMBER > SecondValue;
 	};
 
 	TFunction<bool(const double, const double)> IsReallySmaller = [](const double FirstValue, const double SecondValue) ->bool
 	{
-		return FirstValue + SMALL_NUMBER < SecondValue;
+		return FirstValue + DOUBLE_SMALL_NUMBER < SecondValue;
 	};
 
 	double Slop;
@@ -1016,7 +1016,7 @@ void FGrid::GetMeshOfLoop(const FTopologicalLoop& Loop)
 				{
 					double Diff = CuttingPolyline.Coordinates[Index++];
 					Diff -= CuttingPolyline.Coordinates[Index];
-					if (Diff > -SMALL_NUMBER)
+					if (Diff > -DOUBLE_SMALL_NUMBER)
 					{
 						bProjectionFailed = true;
 						break;
@@ -1487,8 +1487,8 @@ void FGrid::FindInnerFacePoints()
 
 	// Loop node too close to one of CoordinateU or CoordinateV are moved a little to avoid floating error of comparison 
 	// This step is necessary instead of all points could be considered outside...
-	const double SmallToleranceU = SMALL_NUMBER;
-	const double SmallToleranceV = SMALL_NUMBER;
+	const double SmallToleranceU = DOUBLE_SMALL_NUMBER;
+	const double SmallToleranceV = DOUBLE_SMALL_NUMBER;
 
 	{
 		int32 IndexV = 0;
@@ -1636,11 +1636,11 @@ void FGrid::FindInnerFacePoints()
 					else
 					{
 						double APvectAB = UniformCuttingCoordinates[EIso::IsoV][IndexV] * ABu - UniformCuttingCoordinates[EIso::IsoU][IndexU] * ABv + AuABVMinusAvABu;
-						if (APvectAB > SMALL_NUMBER)
+						if (APvectAB > DOUBLE_SMALL_NUMBER)
 						{
 							NbIntersectVForward[Index] = NbIntersectVForward[Index] > 0 ? 0 : 1;
 						}
-						else if (APvectAB < SMALL_NUMBER)
+						else if (APvectAB < DOUBLE_SMALL_NUMBER)
 						{
 							NbIntersectVBackward[Index] = NbIntersectVBackward[Index] > 0 ? 0 : 1;
 						}
@@ -1721,11 +1721,11 @@ void FGrid::FindInnerFacePoints()
 					else
 					{
 						double APvectAB = UniformCuttingCoordinates[EIso::IsoV][IndexV] * ABu - UniformCuttingCoordinates[EIso::IsoU][IndexU] * ABv + AuABVMinusAvABu;
-						if (APvectAB > SMALL_NUMBER)
+						if (APvectAB > DOUBLE_SMALL_NUMBER)
 						{
 							NbIntersectUBackward[Index] = NbIntersectUBackward[Index] > 0 ? 0 : 1;
 						}
-						else if (APvectAB < SMALL_NUMBER)
+						else if (APvectAB < DOUBLE_SMALL_NUMBER)
 						{
 							NbIntersectUForward[Index] = NbIntersectUForward[Index] > 0 ? 0 : 1;
 						}
diff --git a/Engine/Source/Runtime/Datasmith/CADKernel/Private/CADKernel/Mesh/Structure/Grid.h b/Engine/Source/Runtime/Datasmith/CADKernel/Private/CADKernel/Mesh/Structure/Grid.h
--- a/Engine/Source/Runtime/Datasmith/CADKernel/Private/CADKernel/Mesh/Structure/Grid.h
+++ b/Engine/Source/Runtime/Datasmith/CADKernel/Private/CADKernel/Mesh/Structure/Grid.h
@@ -6,10 +6,14 @@
 #include "CADKernel/Geo/GeoPoint.h"
 #include "CADKernel/Math/Point.h"
 #include "CADKernel/Mesh/MeshEnum.h"
-#include "CADKernel/Mesh/Meshers/IsoTriangulator/DefineForDebug.h"
 #include "CADKernel/Mesh/Structure/ThinZone2D.h"
 #include "CADKernel/UI/Display.h"
 
+#ifdef CADKERNEL_DEV
+#include "CADKernel/Mesh/Meshers/IsoTriangulator/DefineForDebug.h"
+#endif
+
+
 namespace CADKernel
 {
 
@@ -565,7 +569,6 @@ public:
 	// ======================================================================================================================================================================================================================
 #ifdef CADKERNEL_DEV
 	bool bDisplay = false;
-	const double DisplayScale = 100.;
 
 	void DisplayIsoNode(EGridSpace Space, const int32 PointIndex, FIdent Ident = 0, EVisuProperty Property = EVisuProperty::BluePoint) const;
 	void DisplayIsoNode(EGridSpace Space, const FIsoNode& Node, FIdent Ident = 0, EVisuProperty Property = EVisuProperty::BluePoint) const;
diff --git a/Engine/Source/Runtime/Datasmith/CADKernel/Private/CADKernel/Mesh/Structure/LoopCleaner.cpp b/Engine/Source/Runtime/Datasmith/CADKernel/Private/CADKernel/Mesh/Structure/LoopCleaner.cpp
--- a/Engine/Source/Runtime/Datasmith/CADKernel/Private/CADKernel/Mesh/Structure/LoopCleaner.cpp
+++ b/Engine/Source/Runtime/Datasmith/CADKernel/Private/CADKernel/Mesh/Structure/LoopCleaner.cpp
@@ -16,7 +16,6 @@ FLoopCleaner::FLoopCleaner(FIsoTriangulator& Triangulator)
 	, LoopNodes(Triangulator.LoopNodes)
 	, LoopSegments(Triangulator.LoopSegments)
 	, IsoSegmentFactory(Triangulator.IsoSegmentFactory)
-	, MesherReport(Triangulator.MesherReport)
 	, bDisplay(Triangulator.bDisplay)
 	, LoopSegmentsIntersectionTool(Grid)
 	, GeometricTolerance(Triangulator.GeometricTolerance)
@@ -27,7 +26,9 @@ FLoopCleaner::FLoopCleaner(FIsoTriangulator& Triangulator)
 	, GetFirst(LoopCleanerImpl::GetFirstNode)
 	, GetSecond(LoopCleanerImpl::GetSecondNode)
 {
-
+#ifdef CADKERNEL_DEV
+	SetMesherReport(*Triangulator.MesherReport);
+#endif
 }
 
 bool FLoopCleaner::CleanLoops()
@@ -534,7 +535,7 @@ bool FLoopCleaner::TryToRemoveSelfIntersectionByMovingTheClosedOusidePoint(const
 	{
 		double Coordinate;
 		ProjectedPoints[OtherSegmentIndex][OtherSegmentNodeIndex] = ProjectPointOnSegment(Points[OtherSegmentIndex][OtherSegmentNodeIndex], Points[SegmentIndex][0], Points[SegmentIndex][1], Coordinate, false);
-		bool bProjectedPoint1IsInside = (Coordinate >= -SMALL_NUMBER && Coordinate <= 1 + SMALL_NUMBER);
+		bool bProjectedPoint1IsInside = (Coordinate >= -DOUBLE_SMALL_NUMBER && Coordinate <= 1 + DOUBLE_SMALL_NUMBER);
 		if (bProjectedPoint1IsInside)
 		{
 			Distance[OtherSegmentIndex][OtherSegmentNodeIndex] = ProjectedPoints[OtherSegmentIndex][OtherSegmentNodeIndex].SquareDistance(Points[OtherSegmentIndex][OtherSegmentNodeIndex]);
@@ -663,7 +664,7 @@ bool FLoopCleaner::TryToRemoveIntersectionByMovingTheClosedOusidePoint(const FIs
 	{
 		double Coordinate;
 		ProjectedPoints[OtherSegmentIndex] = ProjectPointOnSegment(Points[OtherSegmentIndex][OtherSegmentNodeIndex], Points[SegmentIndex][0], Points[SegmentIndex][1], Coordinate, false);
-		bool bProjectedPoint1IsInside = (Coordinate >= -SMALL_NUMBER && Coordinate <= 1 + SMALL_NUMBER);
+		bool bProjectedPoint1IsInside = (Coordinate >= -DOUBLE_SMALL_NUMBER && Coordinate <= 1 + DOUBLE_SMALL_NUMBER);
 		if (bProjectedPoint1IsInside)
 		{
 			Distance[OtherSegmentIndex] = ProjectedPoints[OtherSegmentIndex].Distance(Points[OtherSegmentIndex][OtherSegmentNodeIndex]);
diff --git a/Engine/Source/Runtime/Datasmith/CADKernel/Private/CADKernel/Mesh/Structure/LoopCleaner.h b/Engine/Source/Runtime/Datasmith/CADKernel/Private/CADKernel/Mesh/Structure/LoopCleaner.h
--- a/Engine/Source/Runtime/Datasmith/CADKernel/Private/CADKernel/Mesh/Structure/LoopCleaner.h
+++ b/Engine/Source/Runtime/Datasmith/CADKernel/Private/CADKernel/Mesh/Structure/LoopCleaner.h
@@ -7,7 +7,11 @@
 #include "CADKernel/Mesh/Meshers/IsoTriangulator/IntersectionSegmentTool.h"
 #include "CADKernel/Mesh/Meshers/IsoTriangulator/IsoNode.h"
 #include "CADKernel/Mesh/Meshers/IsoTriangulator/IsoSegment.h"
+
+#ifdef CADKERNEL_DEV
 #include "CADKernel/Mesh/Meshers/MesherReport.h"
+#endif
+
 
 namespace CADKernel
 {
@@ -96,7 +100,10 @@ private:
 	TArray<FLoopNode>& LoopNodes;
 	TArray<FIsoSegment*>& LoopSegments;
 	TFactory<FIsoSegment>& IsoSegmentFactory;
-	FMesherReport& MesherReport;
+
+#ifdef CADKERNEL_DEV
+	FMesherReport* MesherReport;
+#endif
 
 	bool bDisplay;
 
@@ -154,13 +161,17 @@ public:
 	{
 		if (!CleanLoops())
 		{
-			MesherReport.Logs.AddRemoveSelfIntersectionFailure();
+#ifdef CADKERNEL_DEV
+			MesherReport->Logs.AddRemoveSelfIntersectionFailure();
+#endif
 			return false;
 		}
 		
 		if (!UncrossLoops())
 		{
-			MesherReport.Logs.AddRemoveCrossingLoopsFailure();
+#ifdef CADKERNEL_DEV
+			MesherReport->Logs.AddRemoveCrossingLoopsFailure();
+#endif
 			return false;
 		}
 
@@ -169,6 +180,13 @@ public:
 
 private:
 
+#ifdef CADKERNEL_DEV
+	void SetMesherReport(FMesherReport& MesherReport)
+	{
+		MesherReport = &InMesherReport;
+	}
+#endif
+
 	bool CleanLoops();
 	bool UncrossLoops();
 
diff --git a/Engine/Source/Runtime/Datasmith/CADKernel/Private/CADKernel/Topo/TopologicalFace.cpp b/Engine/Source/Runtime/Datasmith/CADKernel/Private/CADKernel/Topo/TopologicalFace.cpp
--- a/Engine/Source/Runtime/Datasmith/CADKernel/Private/CADKernel/Topo/TopologicalFace.cpp
+++ b/Engine/Source/Runtime/Datasmith/CADKernel/Private/CADKernel/Topo/TopologicalFace.cpp
@@ -418,10 +418,10 @@ TSharedRef<FFaceMesh> FTopologicalFace::GetOrCreateMesh(FModelMesh& MeshModel)
 
 void FTopologicalFace::InitDeltaUs()
 {
-	CrossingPointDeltaMins[EIso::IsoU].Init(SMALL_NUMBER, CrossingCoordinates[EIso::IsoU].Num() - 1);
+	CrossingPointDeltaMins[EIso::IsoU].Init(DOUBLE_SMALL_NUMBER, CrossingCoordinates[EIso::IsoU].Num() - 1);
 	CrossingPointDeltaMaxs[EIso::IsoU].Init(HUGE_VALUE, CrossingCoordinates[EIso::IsoU].Num() - 1);
 
-	CrossingPointDeltaMins[EIso::IsoV].Init(SMALL_NUMBER, CrossingCoordinates[EIso::IsoV].Num() - 1);
+	CrossingPointDeltaMins[EIso::IsoV].Init(DOUBLE_SMALL_NUMBER, CrossingCoordinates[EIso::IsoV].Num() - 1);
 	CrossingPointDeltaMaxs[EIso::IsoV].Init(HUGE_VALUE, CrossingCoordinates[EIso::IsoV].Num() - 1);
 }
 
diff --git a/Engine/Source/Runtime/Datasmith/CADKernel/Private/CADKernel/Topo/Topomaker.cpp b/Engine/Source/Runtime/Datasmith/CADKernel/Private/CADKernel/Topo/Topomaker.cpp
--- a/Engine/Source/Runtime/Datasmith/CADKernel/Private/CADKernel/Topo/Topomaker.cpp
+++ b/Engine/Source/Runtime/Datasmith/CADKernel/Private/CADKernel/Topo/Topomaker.cpp
@@ -18,12 +18,14 @@
 #include "CADKernel/UI/Message.h"
 #include "CADKernel/Utils/Util.h"
 
+#include "Math/UnrealMathUtility.h"
+
 namespace CADKernel
 {
 
 FTopomaker::FTopomaker(FSession& InSession, double InTolerance)
 	: Session(InSession)
-	, SewTolerance(InTolerance* UE_SQRT_2)
+	, SewTolerance(InTolerance * UE_DOUBLE_SQRT_2)
 	, SewToleranceSquare(FMath::Square(SewTolerance))
 {
 	int32 ShellCount = 0;
@@ -47,7 +49,7 @@ FTopomaker::FTopomaker(FSession& InSession, double InTolerance)
 FTopomaker::FTopomaker(FSession& InSession, const TArray<TSharedPtr<FTopologicalFace>>& InFaces, double InTolerance)
 	: Session(InSession)
 	, Faces(InFaces)
-	, SewTolerance(InTolerance * UE_SQRT_2)
+	, SewTolerance(InTolerance * UE_DOUBLE_SQRT_2)
 	, SewToleranceSquare(FMath::Square(SewTolerance))
 {
 }
@@ -153,8 +155,10 @@ void FTopomaker::Sew()
 
 	MergeUnconnectedAdjacentEdges();
 
+#ifdef CADKERNEL_DEV
 	Report.SewDuration = FChrono::Elapse(StartJoinTime);
 	FChrono::PrintClockElapse(EVerboseLevel::Log, TEXT(""), TEXT("Sew"), Report.SewDuration);
+#endif
 }
 
 void FTopomaker::GetVertices(TArray<FTopologicalVertex*>& Vertices)
@@ -407,20 +411,22 @@ void FTopomaker::MergeCoincidentEdges(TArray<FTopologicalVertex*>& VerticesToPro
 		for (int32 EdgeI = 0; EdgeI < ConnectedEdgeCount - 1; ++EdgeI)
 		{
 			FTopologicalEdge* Edge = ConnectedEdges[EdgeI];
-			if (!Edge->IsActiveEntity())
+			if (!Edge->IsActiveEntity() || Edge->IsDegenerated())
 			{
 				continue;
 			}
+
 			bool bFirstEdgeBorder = Edge->IsBorder();
 			FTopologicalVertex& EndVertex = *Edge->GetOtherVertex(*Vertex)->GetLinkActiveEntity();
 
 			for (int32 EdgeJ = EdgeI + 1; EdgeJ < ConnectedEdgeCount; ++EdgeJ)
 			{
 				FTopologicalEdge* SecondEdge = ConnectedEdges[EdgeJ];
-				if (!SecondEdge->IsActiveEntity())
+				if (!SecondEdge->IsActiveEntity() || SecondEdge->IsDegenerated())
 				{
 					continue;
 				}
+
 				bool bSecondEdgeBorder = Edge->IsBorder();
 
 				// Process only if at least one edge is Border
@@ -1118,11 +1124,15 @@ void FTopomaker::OrientShells()
 	for (FShell* Shell : Shells)
 	{
 		int32 FaceSwapCount = Shell->Orient();
+#ifdef CADKERNEL_DEV
 		Report.AddSwappedFaceCount(FaceSwapCount);
+#endif
 	}
 
-	Report.OrientationDuration = FChrono::Elapse(StartTime);
+#ifdef CADKERNEL_DEV
+ 	Report.OrientationDuration = FChrono::Elapse(StartTime);
 	FChrono::PrintClockElapse(EVerboseLevel::Log, TEXT(""), TEXT("Orient"), Report.OrientationDuration);
+#endif
 }
 
 }
diff --git a/Engine/Source/Runtime/Datasmith/CADKernel/Private/CADKernel/Topo/TopomakerReport.cpp b/Engine/Source/Runtime/Datasmith/CADKernel/Private/CADKernel/Topo/TopomakerReport.cpp
deleted file mode 100644
diff --git a/Engine/Source/Runtime/Datasmith/CADKernel/Private/CADKernel/UI/Display.cpp b/Engine/Source/Runtime/Datasmith/CADKernel/Private/CADKernel/UI/Display.cpp
--- a/Engine/Source/Runtime/Datasmith/CADKernel/Private/CADKernel/UI/Display.cpp
+++ b/Engine/Source/Runtime/Datasmith/CADKernel/Private/CADKernel/UI/Display.cpp
@@ -16,7 +16,6 @@
 #include "CADKernel/Geo/Surfaces/Surface.h"
 #include "CADKernel/Math/Aabb.h"
 #include "CADKernel/Math/Boundary.h"
-#include "CADKernel/Math/Plane.h"
 #include "CADKernel/Mesh/Structure/EdgeMesh.h"
 #include "CADKernel/Mesh/Structure/ModelMesh.h"
 #include "CADKernel/Mesh/Structure/FaceMesh.h"
@@ -27,6 +26,7 @@
 #include "CADKernel/Topo/Shell.h"
 #include "CADKernel/Topo/TopologicalFace.h"
 
+#include "Math/Plane.h"
 
 namespace CADKernel
 {
@@ -342,25 +342,25 @@ void DisplayEntity2D(const FEntity& Entity)
 void Display(const FPlane& Plane, FIdent Ident)
 {
 #ifdef CADKERNEL_DEV
-	FPoint Normal = Plane.GetNormal();
-	if (Normal.Length() < SMALL_NUMBER)
+	FVector Normal = Plane.GetNormal();
+	if (Normal.Size() < DOUBLE_SMALL_NUMBER)
 	{
 		return;
 	}
 
-	FPoint UAxis;
+	FVector UAxis;
 	// Find the first frame axis non collinear to the plane normal
 	for (int32 Index = 0; Index < 3; ++Index)
 	{
-		FPoint Axis = FPoint::ZeroPoint;
+		FVector Axis = FVector::ZeroVector;
 		Axis[Index] = 1;
 		UAxis = Normal ^ Axis;
-		if (UAxis.Length() > SMALL_NUMBER)
+		if (UAxis.Size() > DOUBLE_SMALL_NUMBER)
 		{
 			break;
 		}
 	}
-	FPoint VAxis = UAxis ^ Normal;
+	FVector VAxis = UAxis ^ Normal;
 
 	UAxis.Normalize();
 	UAxis *= 10;
@@ -368,13 +368,25 @@ void Display(const FPlane& Plane, FIdent Ident)
 	VAxis *= 10;
 
 	TArray<FPoint> Points;
-	Points.Init(FPoint(), 5);
-	const FPoint& Point = Plane.GetPoint();
-	Points[0] = Point + UAxis + VAxis;
-	Points[1] = Point + UAxis - VAxis;
-	Points[2] = Point - UAxis - VAxis;
-	Points[3] = Point - UAxis + VAxis;
-	Points[4] = Point + UAxis + VAxis;
+
+	const FVector Point = Plane.GetOrigin();
+	
+	FVector VTemp;
+
+	VTemp = Point + UAxis + VAxis;
+	Points.Emplace(VTemp[0], VTemp[1], VTemp[2]);
+
+	VTemp = Point + UAxis - VAxis;
+	Points.Emplace(VTemp[0], VTemp[1], VTemp[2]);
+
+	VTemp = Point - UAxis - VAxis;
+	Points.Emplace(VTemp[0], VTemp[1], VTemp[2]);
+
+	VTemp = Point - UAxis + VAxis;
+	Points.Emplace(VTemp[0], VTemp[1], VTemp[2]);
+
+	VTemp = Point + UAxis + VAxis;
+	Points.Emplace(VTemp[0], VTemp[1], VTemp[2]);
 
 	F3DDebugSegment G(Ident);
 	DrawElement(2, Points);
@@ -623,7 +635,7 @@ void DisplayControlPolygon(const FCurve& Curve)
 	{
 		for (const FPoint& Pole : Poles)
 		{
-			Display(Pole, EVisuProperty::GreenPoint);
+			DisplayPoint(Pole, EVisuProperty::GreenPoint);
 		}
 
 		for (int32 Index = 1; Index < Poles.Num(); Index++)
@@ -657,7 +669,7 @@ void DisplayControlPolygon(const FSurface& Surface)
 	{
 		for (int32 Index = 0; Index < Poles.Num(); Index++)
 		{
-			Display(Poles[Index], EVisuProperty::GreenPoint);
+			DisplayPoint(Poles[Index], EVisuProperty::GreenPoint);
 		}
 
 		for (int32 IndexV = 0, Index = 0; IndexV < PoleVNum; IndexV++)
diff --git a/Engine/Source/Runtime/Datasmith/CADKernel/Public/CADKernel/Geo/Curves/BezierCurve.h b/Engine/Source/Runtime/Datasmith/CADKernel/Public/CADKernel/Geo/Curves/BezierCurve.h
deleted file mode 100644
diff --git a/Engine/Source/Runtime/Datasmith/CADKernel/Public/CADKernel/Geo/Curves/BoundedCurve.h b/Engine/Source/Runtime/Datasmith/CADKernel/Public/CADKernel/Geo/Curves/BoundedCurve.h
deleted file mode 100644
diff --git a/Engine/Source/Runtime/Datasmith/CADKernel/Public/CADKernel/Geo/Curves/CompositeCurve.h b/Engine/Source/Runtime/Datasmith/CADKernel/Public/CADKernel/Geo/Curves/CompositeCurve.h
deleted file mode 100644
diff --git a/Engine/Source/Runtime/Datasmith/CADKernel/Public/CADKernel/Geo/Curves/Curve.h b/Engine/Source/Runtime/Datasmith/CADKernel/Public/CADKernel/Geo/Curves/Curve.h
--- a/Engine/Source/Runtime/Datasmith/CADKernel/Public/CADKernel/Geo/Curves/Curve.h
+++ b/Engine/Source/Runtime/Datasmith/CADKernel/Public/CADKernel/Geo/Curves/Curve.h
@@ -13,10 +13,11 @@
 
 namespace CADKernel
 {
-	class FSurface;
+	//class FSurface;
 
 	struct FCurvePoint;
 	struct FCurvePoint2D;
+	struct FNurbsCurveData;
 
 	class CADKERNEL_API FCurve : public FEntityGeom
 	{
@@ -45,6 +46,15 @@ namespace CADKernel
 
 	public:
 
+		static TSharedPtr<FCurve> MakeNurbsCurve(FNurbsCurveData& InNurbsData);
+		// TODO
+		//static TSharedPtr<FCurve> MakeBezierCurve(const double InToleranceGeometric, int32 InUDegre, int32 InVDegre, const TArray<FPoint>& InPoles);
+		//static TSharedPtr<FCurve> MakeConeCurve(const double InToleranceGeometric, const FMatrixH& InMatrix, double InStartRadius, double InConeAngle, const FSurfacicBoundary& InBoundary);
+		//static TSharedPtr<FCurve> MakeCylinderCurve(const double InToleranceGeometric, const FMatrixH& InMatrix, const double InRadius, const FSurfacicBoundary& InBoundary);
+		//static TSharedPtr<FCurve> MakePlaneCurve(const double InToleranceGeometric, const FMatrixH& InMatrix, const FSurfacicBoundary& InBoundary);
+		//static TSharedPtr<FCurve> MakeSphericalCurve(const double InToleranceGeometric, const FMatrixH& InMatrix, double InRadius, const FSurfacicBoundary& InBoundary);
+		//static TSharedPtr<FCurve> MakeTorusCurve(const double InToleranceGeometric, const FMatrixH& InMatrix, double InMajorRadius, double InMinorRadius, const FSurfacicBoundary& InBoundary);
+
 		virtual void Serialize(FCADKernelArchive& Ar) override
 		{
 			// Curve's type is serialize because it is used to instantiate the correct entity on deserialization (@see Deserialize(FCADKernelArchive& Archive)) 
diff --git a/Engine/Source/Runtime/Datasmith/CADKernel/Public/CADKernel/Geo/Curves/EllipseCurve.h b/Engine/Source/Runtime/Datasmith/CADKernel/Public/CADKernel/Geo/Curves/EllipseCurve.h
deleted file mode 100644
diff --git a/Engine/Source/Runtime/Datasmith/CADKernel/Public/CADKernel/Geo/Curves/HyperbolaCurve.h b/Engine/Source/Runtime/Datasmith/CADKernel/Public/CADKernel/Geo/Curves/HyperbolaCurve.h
deleted file mode 100644
diff --git a/Engine/Source/Runtime/Datasmith/CADKernel/Public/CADKernel/Geo/Curves/NURBSCurve.h b/Engine/Source/Runtime/Datasmith/CADKernel/Public/CADKernel/Geo/Curves/NURBSCurve.h
deleted file mode 100644
diff --git a/Engine/Source/Runtime/Datasmith/CADKernel/Public/CADKernel/Geo/Curves/NURBSCurveData.h b/Engine/Source/Runtime/Datasmith/CADKernel/Public/CADKernel/Geo/Curves/NURBSCurveData.h
new file mode 100644
--- /dev/null
+++ b/Engine/Source/Runtime/Datasmith/CADKernel/Public/CADKernel/Geo/Curves/NURBSCurveData.h
@@ -0,0 +1,21 @@
+// Copyright Epic Games, Inc. All Rights Reserved.
+#pragma once
+
+#include "CADKernel/Core/Types.h"
+#include "CADKernel/Math/Point.h"
+
+namespace CADKernel
+{
+	struct FNurbsCurveData
+	{
+		bool bIsRational;
+		int32 Dimension = 0;
+
+		int32 Degree = 0;
+		TArray<double> NodalVector;
+
+		TArray<double> Weights;
+		TArray<FPoint> Poles;
+	};
+} // namespace CADKernel
+
diff --git a/Engine/Source/Runtime/Datasmith/CADKernel/Public/CADKernel/Geo/Curves/ParabolaCurve.h b/Engine/Source/Runtime/Datasmith/CADKernel/Public/CADKernel/Geo/Curves/ParabolaCurve.h
deleted file mode 100644
diff --git a/Engine/Source/Runtime/Datasmith/CADKernel/Public/CADKernel/Geo/Curves/PolylineCurve.h b/Engine/Source/Runtime/Datasmith/CADKernel/Public/CADKernel/Geo/Curves/PolylineCurve.h
deleted file mode 100644
diff --git a/Engine/Source/Runtime/Datasmith/CADKernel/Public/CADKernel/Geo/Curves/SegmentCurve.h b/Engine/Source/Runtime/Datasmith/CADKernel/Public/CADKernel/Geo/Curves/SegmentCurve.h
deleted file mode 100644
diff --git a/Engine/Source/Runtime/Datasmith/CADKernel/Public/CADKernel/Geo/Sampling/PolylineTools.h b/Engine/Source/Runtime/Datasmith/CADKernel/Public/CADKernel/Geo/Sampling/PolylineTools.h
--- a/Engine/Source/Runtime/Datasmith/CADKernel/Public/CADKernel/Geo/Sampling/PolylineTools.h
+++ b/Engine/Source/Runtime/Datasmith/CADKernel/Public/CADKernel/Geo/Sampling/PolylineTools.h
@@ -135,7 +135,7 @@ template<class PointType>
 PointType ComputePoint(const TArray<double>& PolylineCoordinates, const TArray<PointType>& PolylinePoints, const int32 Index, const double PointCoordinate)
 {
 	double Delta = PolylineCoordinates[Index + 1] - PolylineCoordinates[Index];
-	if (FMath::IsNearlyZero(Delta, (double)SMALL_NUMBER))
+	if (FMath::IsNearlyZero(Delta, (double)DOUBLE_SMALL_NUMBER))
 	{
 		return PolylinePoints[Index];
 	}
@@ -201,7 +201,7 @@ protected:
 	PointType ComputePoint(const int32 Index, const double PointCoordinate) const
 	{
 		double Delta = PolylineCoordinates[Index + 1] - PolylineCoordinates[Index];
-		if (FMath::IsNearlyZero(Delta, (double)SMALL_NUMBER))
+		if (FMath::IsNearlyZero(Delta, (double)DOUBLE_SMALL_NUMBER))
 		{
 			return PolylinePoints[Index];
 		}
@@ -425,12 +425,12 @@ public:
 		double LastCoordinate = InBoundary.Min;
 		for (int32 IndexCurvilinear = 1, IndexCoordinate = BoundaryIndices[0] + 1; IndexCurvilinear < CurvilinearCoordinates.Num(); ++IndexCurvilinear, ++IndexCoordinate)
 		{
-			while (CurvilinearLength < CurvilinearCoordinates[IndexCurvilinear] + SMALL_NUMBER)
+			while (CurvilinearLength < CurvilinearCoordinates[IndexCurvilinear] + DOUBLE_SMALL_NUMBER)
 			{
 				double Coordinate = ComputeSamplePointCoordinate(IndexCurvilinear, IndexCoordinate, CurvilinearLength, LastCoordinate);
 				OutCoordinates.Add(Coordinate);
 				CurvilinearLength += SectionLength;
-				if (CurvilinearLength + SMALL_NUMBER > CurveLength)
+				if (CurvilinearLength + DOUBLE_SMALL_NUMBER > CurveLength)
 				{
 					OutCoordinates.Add(InBoundary.Max);
 					break;
@@ -553,11 +553,11 @@ public:
 
 		int32 PolylineStartIndex = BoundaryIndices[0];
 		int32 PolylineEndIndex = BoundaryIndices[1];
-		if (FMath::IsNearlyEqual(PolylineCoordinates[BoundaryIndices[0] + 1], InBoundary.Min, (double)SMALL_NUMBER))
+		if (FMath::IsNearlyEqual(PolylineCoordinates[BoundaryIndices[0] + 1], InBoundary.Min, (double)DOUBLE_SMALL_NUMBER))
 		{
 			PolylineStartIndex++;
 		}
-		if (FMath::IsNearlyEqual(PolylineCoordinates[BoundaryIndices[1]], InBoundary.Max, (double)SMALL_NUMBER))
+		if (FMath::IsNearlyEqual(PolylineCoordinates[BoundaryIndices[1]], InBoundary.Max, (double)DOUBLE_SMALL_NUMBER))
 		{
 			PolylineEndIndex--;
 		}
@@ -596,12 +596,12 @@ public:
 		OutCoordinates.Empty(NewSize);
 		OutPoints.Empty(NewSize);
 
-		if (FMath::IsNearlyEqual(PolylineCoordinates[BoundaryIndices[0] + 1], InBoundary.Min, (double)SMALL_NUMBER))
+		if (FMath::IsNearlyEqual(PolylineCoordinates[BoundaryIndices[0] + 1], InBoundary.Min, (double)DOUBLE_SMALL_NUMBER))
 		{
 			BoundaryIndices[0]++;
 		}
 
-		if (FMath::IsNearlyEqual(PolylineCoordinates[BoundaryIndices[1]], InBoundary.Max, (double)SMALL_NUMBER))
+		if (FMath::IsNearlyEqual(PolylineCoordinates[BoundaryIndices[1]], InBoundary.Max, (double)DOUBLE_SMALL_NUMBER))
 		{
 			BoundaryIndices[1]--;
 		}
@@ -625,12 +625,12 @@ public:
 
 		if(BoundaryIndices[1] - BoundaryIndices[0] > 0)
 		{
-			if (FMath::IsNearlyEqual(PolylineCoordinates[BoundaryIndices[0] + 1], InBoundary.Min, (double)SMALL_NUMBER))
+			if (FMath::IsNearlyEqual(PolylineCoordinates[BoundaryIndices[0] + 1], InBoundary.Min, (double)DOUBLE_SMALL_NUMBER))
 			{
 				BoundaryIndices[0]++;
 			}
 
-			if (FMath::IsNearlyEqual(PolylineCoordinates[BoundaryIndices[1]], InBoundary.Max, (double)SMALL_NUMBER))
+			if (FMath::IsNearlyEqual(PolylineCoordinates[BoundaryIndices[1]], InBoundary.Max, (double)DOUBLE_SMALL_NUMBER))
 			{
 				BoundaryIndices[1]--;
 			}
diff --git a/Engine/Source/Runtime/Datasmith/CADKernel/Public/CADKernel/Geo/Sampling/SurfacicPolyline.h b/Engine/Source/Runtime/Datasmith/CADKernel/Public/CADKernel/Geo/Sampling/SurfacicPolyline.h
--- a/Engine/Source/Runtime/Datasmith/CADKernel/Public/CADKernel/Geo/Sampling/SurfacicPolyline.h
+++ b/Engine/Source/Runtime/Datasmith/CADKernel/Public/CADKernel/Geo/Sampling/SurfacicPolyline.h
@@ -97,7 +97,7 @@ public:
 	FSurfacicTolerance ComputeTolerance(const double Tolerance3D, const FSurfacicTolerance& MinToleranceIso, const int32 Index) const
 	{
 		double Distance3D = Points3D[Index].Distance(Points3D[Index + 1]);
-		if (FMath::IsNearlyZero(Distance3D, (double)SMALL_NUMBER))
+		if (FMath::IsNearlyZero(Distance3D, (double)DOUBLE_SMALL_NUMBER))
 		{
 			return FPoint2D::FarawayPoint;
 		}
@@ -111,7 +111,7 @@ public:
 	double ComputeLinearToleranceAt(const double Tolerance3D, const double MinLinearTolerance, const int32 Index) const
 	{
 		double Distance3D = Points3D[Index].Distance(Points3D[Index + 1]);
-		if (FMath::IsNearlyZero(Distance3D, (double)SMALL_NUMBER))
+		if (FMath::IsNearlyZero(Distance3D, (double)DOUBLE_SMALL_NUMBER))
 		{
 			return (Coordinates.Last() - Coordinates[0]) / 10.;
 		}
diff --git a/Engine/Source/Runtime/Datasmith/CADKernel/Public/CADKernel/Geo/Surfaces/BezierSurface.h b/Engine/Source/Runtime/Datasmith/CADKernel/Public/CADKernel/Geo/Surfaces/BezierSurface.h
deleted file mode 100644
diff --git a/Engine/Source/Runtime/Datasmith/CADKernel/Public/CADKernel/Geo/Surfaces/CompositeSurface.h b/Engine/Source/Runtime/Datasmith/CADKernel/Public/CADKernel/Geo/Surfaces/CompositeSurface.h
deleted file mode 100644
diff --git a/Engine/Source/Runtime/Datasmith/CADKernel/Public/CADKernel/Geo/Surfaces/ConeSurface.h b/Engine/Source/Runtime/Datasmith/CADKernel/Public/CADKernel/Geo/Surfaces/ConeSurface.h
deleted file mode 100644
diff --git a/Engine/Source/Runtime/Datasmith/CADKernel/Public/CADKernel/Geo/Surfaces/CoonsSurface.h b/Engine/Source/Runtime/Datasmith/CADKernel/Public/CADKernel/Geo/Surfaces/CoonsSurface.h
deleted file mode 100644
diff --git a/Engine/Source/Runtime/Datasmith/CADKernel/Public/CADKernel/Geo/Surfaces/CylinderSurface.h b/Engine/Source/Runtime/Datasmith/CADKernel/Public/CADKernel/Geo/Surfaces/CylinderSurface.h
deleted file mode 100644
diff --git a/Engine/Source/Runtime/Datasmith/CADKernel/Public/CADKernel/Geo/Surfaces/NURBSSurface.h b/Engine/Source/Runtime/Datasmith/CADKernel/Public/CADKernel/Geo/Surfaces/NURBSSurface.h
deleted file mode 100644
diff --git a/Engine/Source/Runtime/Datasmith/CADKernel/Public/CADKernel/Geo/Surfaces/NurbsSurfaceData.h b/Engine/Source/Runtime/Datasmith/CADKernel/Public/CADKernel/Geo/Surfaces/NurbsSurfaceData.h
new file mode 100644
--- /dev/null
+++ b/Engine/Source/Runtime/Datasmith/CADKernel/Public/CADKernel/Geo/Surfaces/NurbsSurfaceData.h
@@ -0,0 +1,43 @@
+// Copyright Epic Games, Inc. All Rights Reserved.
+#pragma once
+
+#include "CADKernel/Core/Types.h"
+#include "CADKernel/Math/Point.h"
+
+namespace CADKernel
+{
+
+struct FNurbsSurfaceHomogeneousData
+{
+	bool bSwapUV;
+	int32 PoleUCount;
+	int32 PoleVCount;
+
+	int32 UDegree;
+	int32 VDegree;
+
+	TArray<double> UNodalVector;
+	TArray<double> VNodalVector;
+
+	bool bIsRational;
+	TArray<double> HomogeneousPoles;
+};
+
+struct FNurbsSurfaceData
+{
+	bool bSwapUV;
+	int32 PoleUCount;
+	int32 PoleVCount;
+
+	int32 UDegree;
+	int32 VDegree;
+
+	TArray<double> UNodalVector;
+	TArray<double> VNodalVector;
+
+	// if Weights.Num() == 0  => bIsRational = false
+	TArray<double> Weights;
+	TArray<FPoint> Poles;
+};
+
+}
\ No newline at end of file
diff --git a/Engine/Source/Runtime/Datasmith/CADKernel/Public/CADKernel/Geo/Surfaces/OffsetSurface.h b/Engine/Source/Runtime/Datasmith/CADKernel/Public/CADKernel/Geo/Surfaces/OffsetSurface.h
deleted file mode 100644
diff --git a/Engine/Source/Runtime/Datasmith/CADKernel/Public/CADKernel/Geo/Surfaces/PlaneSurface.h b/Engine/Source/Runtime/Datasmith/CADKernel/Public/CADKernel/Geo/Surfaces/PlaneSurface.h
deleted file mode 100644
diff --git a/Engine/Source/Runtime/Datasmith/CADKernel/Public/CADKernel/Geo/Surfaces/RevolutionSurface.h b/Engine/Source/Runtime/Datasmith/CADKernel/Public/CADKernel/Geo/Surfaces/RevolutionSurface.h
deleted file mode 100644
diff --git a/Engine/Source/Runtime/Datasmith/CADKernel/Public/CADKernel/Geo/Surfaces/RuledSurface.h b/Engine/Source/Runtime/Datasmith/CADKernel/Public/CADKernel/Geo/Surfaces/RuledSurface.h
deleted file mode 100644
diff --git a/Engine/Source/Runtime/Datasmith/CADKernel/Public/CADKernel/Geo/Surfaces/SphericalSurface.h b/Engine/Source/Runtime/Datasmith/CADKernel/Public/CADKernel/Geo/Surfaces/SphericalSurface.h
deleted file mode 100644
diff --git a/Engine/Source/Runtime/Datasmith/CADKernel/Public/CADKernel/Geo/Surfaces/Surface.h b/Engine/Source/Runtime/Datasmith/CADKernel/Public/CADKernel/Geo/Surfaces/Surface.h
--- a/Engine/Source/Runtime/Datasmith/CADKernel/Public/CADKernel/Geo/Surfaces/Surface.h
+++ b/Engine/Source/Runtime/Datasmith/CADKernel/Public/CADKernel/Geo/Surfaces/Surface.h
@@ -20,6 +20,8 @@ namespace CADKernel
 	struct FCurvePoint;
 	struct FCurvePoint2D;
 	struct FSurfacicSampling;
+	struct FNurbsSurfaceData;
+	struct FNurbsSurfaceHomogeneousData;
 
 	class CADKERNEL_API FSurface : public FEntityGeom
 	{
@@ -79,6 +81,15 @@ namespace CADKernel
 
 	public:
 
+		static TSharedPtr<FSurface> MakeBezierSurface(const double InToleranceGeometric, int32 InUDegre, int32 InVDegre, const TArray<FPoint>& InPoles);
+		static TSharedPtr<FSurface> MakeConeSurface(const double InToleranceGeometric, const FMatrixH& InMatrix, double InStartRadius, double InConeAngle, const FSurfacicBoundary& InBoundary);
+		static TSharedPtr<FSurface> MakeCylinderSurface(const double InToleranceGeometric, const FMatrixH& InMatrix, const double InRadius, const FSurfacicBoundary& InBoundary);
+		static TSharedPtr<FSurface> MakeNurbsSurface(const double InToleranceGeometric, const FNurbsSurfaceData& NurbsData);
+		static TSharedPtr<FSurface> MakeNurbsSurface(const double InToleranceGeometric, const FNurbsSurfaceHomogeneousData& NurbsData);
+		static TSharedPtr<FSurface> MakePlaneSurface(const double InToleranceGeometric, const FMatrixH& InMatrix, const FSurfacicBoundary& InBoundary);
+		static TSharedPtr<FSurface> MakeSphericalSurface(const double InToleranceGeometric, const FMatrixH& InMatrix, double InRadius, const FSurfacicBoundary& InBoundary);
+		static TSharedPtr<FSurface> MakeTorusSurface(const double InToleranceGeometric, const FMatrixH& InMatrix, double InMajorRadius, double InMinorRadius, const FSurfacicBoundary& InBoundary);
+
 		virtual void Serialize(FCADKernelArchive& Ar) override
 		{
 			// Surface's type is serialize because it is used to instantiate the correct entity on de-serialization (@see Deserialize(FCADKernelArchive& Archive))
@@ -237,7 +248,7 @@ namespace CADKernel
 
 		void PresampleIsoCircle(const FSurfacicBoundary& InBoundaries, FCoordinateGrid& Coordinates, EIso Iso)
 		{
-			constexpr const double SixOverPi = 6 / PI;
+			constexpr double SixOverPi = 6 / DOUBLE_PI;
 			double Delta = InBoundaries.Length(Iso);
 			int32 SampleCount = /*FMath::Max(3, */ (int32)(Delta * SixOverPi + 1)/*)*/;
 			Delta /= SampleCount;
diff --git a/Engine/Source/Runtime/Datasmith/CADKernel/Public/CADKernel/Geo/Surfaces/TabulatedCylinderSurface.h b/Engine/Source/Runtime/Datasmith/CADKernel/Public/CADKernel/Geo/Surfaces/TabulatedCylinderSurface.h
deleted file mode 100644
diff --git a/Engine/Source/Runtime/Datasmith/CADKernel/Public/CADKernel/Geo/Surfaces/TorusSurface.h b/Engine/Source/Runtime/Datasmith/CADKernel/Public/CADKernel/Geo/Surfaces/TorusSurface.h
deleted file mode 100644
diff --git a/Engine/Source/Runtime/Datasmith/CADKernel/Public/CADKernel/Math/Aabb.h b/Engine/Source/Runtime/Datasmith/CADKernel/Public/CADKernel/Math/Aabb.h
--- a/Engine/Source/Runtime/Datasmith/CADKernel/Public/CADKernel/Math/Aabb.h
+++ b/Engine/Source/Runtime/Datasmith/CADKernel/Public/CADKernel/Math/Aabb.h
@@ -18,9 +18,6 @@ ENUM_CLASS_FLAGS(EAABBBoundary);
 
 namespace CADKernel
 {
-	class FSphere;
-	class FPlane;
-
 	template<class PointType>
 	class TAABB
 	{
diff --git a/Engine/Source/Runtime/Datasmith/CADKernel/Public/CADKernel/Math/Boundary.h b/Engine/Source/Runtime/Datasmith/CADKernel/Public/CADKernel/Math/Boundary.h
--- a/Engine/Source/Runtime/Datasmith/CADKernel/Public/CADKernel/Math/Boundary.h
+++ b/Engine/Source/Runtime/Datasmith/CADKernel/Public/CADKernel/Math/Boundary.h
@@ -117,7 +117,7 @@ namespace CADKernel
 		bool IsDegenerated() const
 		{
 			double DeltaU = (Max - Min);
-			return (DeltaU < KINDA_SMALL_NUMBER);
+			return (DeltaU < DOUBLE_KINDA_SMALL_NUMBER);
 		}
 
 		/**
@@ -134,7 +134,7 @@ namespace CADKernel
 		/**
 		 * If a coordinate is outside the bounds, set the coordinate at the closed limit
 		 */
-		void MoveInsideIfNot(double& Coordinate, double Tolerance = SMALL_NUMBER) const
+		void MoveInsideIfNot(double& Coordinate, double Tolerance = DOUBLE_SMALL_NUMBER) const
 		{
 			if (Coordinate <= Min)
 			{
@@ -206,14 +206,14 @@ namespace CADKernel
 		}
 
 		/**
-		 * If the boundary width is near or equal to zero, it's widened by +/- SMALL_NUMBER
+		 * If the boundary width is near or equal to zero, it's widened by +/- DOUBLE_SMALL_NUMBER
 		 */
 		void WidenIfDegenerated()
 		{
 			if (FMath::IsNearlyEqual(Min, Max))
 			{
-				Min -= SMALL_NUMBER;
-				Max += SMALL_NUMBER;
+				Min -= DOUBLE_SMALL_NUMBER;
+				Max += DOUBLE_SMALL_NUMBER;
 			}
 		}
 
@@ -349,7 +349,7 @@ namespace CADKernel
 		}
 
 		/**
-		 * If Along each axis, the bound width is near equal to zero, it's widened by +/- SMALL_NUMBER
+		 * If Along each axis, the bound width is near equal to zero, it's widened by +/- DOUBLE_SMALL_NUMBER
 		 */
 		void WidenIfDegenerated()
 		{
@@ -360,7 +360,7 @@ namespace CADKernel
 		/**
 		 * If a point is outside the bounds, set the coordinate to insert the point inside the bounds
 		 */
-		void MoveInsideIfNot(FPoint& Point, double Tolerance = SMALL_NUMBER) const
+		void MoveInsideIfNot(FPoint& Point, double Tolerance = DOUBLE_SMALL_NUMBER) const
 		{
 			UVBoundaries[EIso::IsoU].MoveInsideIfNot(Point.X, Tolerance);
 			UVBoundaries[EIso::IsoV].MoveInsideIfNot(Point.Y, Tolerance);
@@ -369,7 +369,7 @@ namespace CADKernel
 		/**
 		 * If a point is outside the bounds, set the coordinate to insert the point inside the bounds
 		 */
-		void MoveInsideIfNot(FPoint2D& Point, double Tolerance = SMALL_NUMBER) const
+		void MoveInsideIfNot(FPoint2D& Point, double Tolerance = DOUBLE_SMALL_NUMBER) const
 		{
 			UVBoundaries[EIso::IsoU].MoveInsideIfNot(Point.U, Tolerance);
 			UVBoundaries[EIso::IsoV].MoveInsideIfNot(Point.V, Tolerance);
diff --git a/Engine/Source/Runtime/Datasmith/CADKernel/Public/CADKernel/Math/Geometry.h b/Engine/Source/Runtime/Datasmith/CADKernel/Public/CADKernel/Math/Geometry.h
--- a/Engine/Source/Runtime/Datasmith/CADKernel/Public/CADKernel/Math/Geometry.h
+++ b/Engine/Source/Runtime/Datasmith/CADKernel/Public/CADKernel/Math/Geometry.h
@@ -531,20 +531,17 @@ inline bool FastIntersectSegments2D(const TSegment<FPoint2D>& SegmentAB, const T
  */
 inline FPoint2D FindIntersectionOfSegments2D(const TSegment<FPoint2D>& SegmentAB, const TSegment<FPoint2D>& SegmentCD, double& OutABIntersectionCoordinate)
 {
-	constexpr const double Min = -DOUBLE_SMALL_NUMBER;
-	constexpr const double Max = 1. + DOUBLE_SMALL_NUMBER;
-
-	FPoint2D AB = SegmentAB[1] - SegmentAB[0];
-	FPoint2D DC = SegmentCD[0] - SegmentCD[1];
-	FPoint2D AC = SegmentCD[0] - SegmentAB[0];
+	const FPoint2D AB = SegmentAB[1] - SegmentAB[0];
+	const FPoint2D DC = SegmentCD[0] - SegmentCD[1];
+	const FPoint2D AC = SegmentCD[0] - SegmentAB[0];
 
 	double ParallelCoef = DC ^ AB;
 	if (FMath::IsNearlyZero(ParallelCoef))
 	{
-		double SquareAB = AB * AB;
+		const double SquareAB = AB * AB;
 		double CCoordinate = (AB * AC) / SquareAB;
 
-		FPoint2D AD = SegmentCD[1] - SegmentAB[0];
+		const FPoint2D AD = SegmentCD[1] - SegmentAB[0];
 		double DCoordinate = (AB * AD) / SquareAB;
 
 		if (CCoordinate >= -DOUBLE_KINDA_SMALL_NUMBER && CCoordinate <= 1 + DOUBLE_KINDA_SMALL_NUMBER)
@@ -588,7 +585,7 @@ inline FPoint2D FindIntersectionOfSegments2D(const TSegment<FPoint2D>& SegmentAB
 }
 
 /**
- * @return false if the lines are parallele
+ * @return false if the lines are parallele 
  */
 inline bool FindIntersectionOfLines2D(const TSegment<FPoint2D>& LineAB, const TSegment<FPoint2D>& LineCD, FPoint2D& OutIntersectionPoint)
 {
@@ -606,7 +603,7 @@ inline bool FindIntersectionOfLines2D(const TSegment<FPoint2D>& LineAB, const TS
 	}
 
 	double OutABIntersectionCoordinate = (DC ^ AC) / ParallelCoef;
-	OutIntersectionPoint = LineAB[0] + AB * OutABIntersectionCoordinate;
+	OutIntersectionPoint =  LineAB[0] + AB * OutABIntersectionCoordinate;
 	return true;
 }
 
diff --git a/Engine/Source/Runtime/Datasmith/CADKernel/Public/CADKernel/Math/MathConst.h b/Engine/Source/Runtime/Datasmith/CADKernel/Public/CADKernel/Math/MathConst.h
--- a/Engine/Source/Runtime/Datasmith/CADKernel/Public/CADKernel/Math/MathConst.h
+++ b/Engine/Source/Runtime/Datasmith/CADKernel/Public/CADKernel/Math/MathConst.h
@@ -33,7 +33,7 @@ namespace CADKernel
 		}
 	}
 
-	inline int32 RealCompare(const double Value1, const double Value2, const double Tolerance = SMALL_NUMBER)
+	inline int32 RealCompare(const double Value1, const double Value2, const double Tolerance = DOUBLE_SMALL_NUMBER)
 	{
 		double Difference = Value1 - Value2;
 		if (Difference < -Tolerance) 
diff --git a/Engine/Source/Runtime/Datasmith/CADKernel/Public/CADKernel/Mesh/Criteria/Criterion.h b/Engine/Source/Runtime/Datasmith/CADKernel/Public/CADKernel/Mesh/Criteria/Criterion.h
--- a/Engine/Source/Runtime/Datasmith/CADKernel/Public/CADKernel/Mesh/Criteria/Criterion.h
+++ b/Engine/Source/Runtime/Datasmith/CADKernel/Public/CADKernel/Mesh/Criteria/Criterion.h
@@ -81,7 +81,7 @@ namespace CADKernel
 			double NormSqrVec = ChordVec * ChordVec;
 			Length = sqrt(NormSqrVec);
 
-			if (NormSqrVec > SMALL_NUMBER)
+			if (NormSqrVec > DOUBLE_SMALL_NUMBER)
 			{
 				FPoint VecSag = ChordVec ^ MiddleVec;
 				double NormSqrSag = VecSag * VecSag;
@@ -118,7 +118,7 @@ namespace CADKernel
 			double DiagonalSagU = FMath::Abs(InDiagonalSag - InVSag) * ChordLength / DiagonalLength;
 			InUSag = FMath::Max(InUSag, DiagonalSagU);
 
-			if (InUSag > SMALL_NUMBER)
+			if (InUSag > DOUBLE_SMALL_NUMBER)
 			{
 				double DeltaUMax = ComputeDeltaU(ChordLength, InDeltaU, InUSag);
 				OutSagDeltaUMax = (OutSagDeltaUMax < DeltaUMax) ? OutSagDeltaUMax : DeltaUMax;
diff --git a/Engine/Source/Runtime/Datasmith/CADKernel/Public/CADKernel/Mesh/Meshers/ParametricMesher.h b/Engine/Source/Runtime/Datasmith/CADKernel/Public/CADKernel/Mesh/Meshers/ParametricMesher.h
--- a/Engine/Source/Runtime/Datasmith/CADKernel/Public/CADKernel/Mesh/Meshers/ParametricMesher.h
+++ b/Engine/Source/Runtime/Datasmith/CADKernel/Public/CADKernel/Mesh/Meshers/ParametricMesher.h
@@ -3,11 +3,14 @@
 
 #include "CADKernel/Core/Chrono.h"
 #include "CADKernel/Core/Types.h"
-#include "CADKernel/Mesh/Meshers/MesherReport.h"
 #include "CADKernel/Topo/TopologicalEdge.h"
 #include "CADKernel/Topo/Shell.h"
 #include "CADKernel/UI/Progress.h"
 
+#ifdef CADKERNEL_DEV
+#include "CADKernel/Mesh/Meshers/MesherReport.h"
+#endif
+
 namespace CADKernel
 {
 class FCriterion;
@@ -28,19 +31,6 @@ struct FCostToFace
 	}
 };
 
-class CADKERNEL_API FMesherParameters : public FParameters
-{
-public:
-	FParameter InconsistencyAngle;
-
-	FMesherParameters();
-
-	void SetInconsistencyAngle(const double Value)
-	{
-		InconsistencyAngle = Value;
-	}
-};
-
 class CADKERNEL_API FParametricMesher
 {
 protected:
@@ -54,13 +44,14 @@ protected:
 
 	TArray<FTopologicalFace*> Faces;
 
+#ifdef CADKERNEL_DEV
 	FMesherReport MesherReport;
-
 	bool bDisplay = false;
+#endif
 
 public:
 
-	FParametricMesher(FModelMesh& MeshModel);
+	FParametricMesher(FModelMesh& InMeshModel);
 
 	const FModelMesh& GetMeshModel() const
 	{
@@ -93,7 +84,18 @@ public:
 
 	void GenerateCloud(FGrid& Grid);
 
-	void PrintReport();
+
+#ifdef CADKERNEL_DEV
+	void SetMeshReport(const FMesherReport& InMesherReport)
+	{
+		MesherReport = InMesherReport;
+	}
+
+	const FMesherReport& GetMeshReport() const 
+	{
+		return MesherReport;
+	}
+#endif
 
 protected:
 
diff --git a/Engine/Source/Runtime/Datasmith/CADKernel/Public/CADKernel/Mesh/Structure/ModelMesh.h b/Engine/Source/Runtime/Datasmith/CADKernel/Public/CADKernel/Mesh/Structure/ModelMesh.h
--- a/Engine/Source/Runtime/Datasmith/CADKernel/Public/CADKernel/Mesh/Structure/ModelMesh.h
+++ b/Engine/Source/Runtime/Datasmith/CADKernel/Public/CADKernel/Mesh/Structure/ModelMesh.h
@@ -30,9 +30,9 @@ namespace CADKernel
 		TArray<TSharedPtr<FFaceMesh>> FaceMeshes;
 
 		bool QuadAnalyse = false;
-		double MinSize = SMALL_NUMBER;
+		double MinSize = DOUBLE_SMALL_NUMBER;
 		double MaxSize = HUGE_VALUE;
-		double MaxAngle = PI;
+		double MaxAngle = DOUBLE_PI;
 		double Sag = HUGE_VALUE;
 
 		FModelMesh()
diff --git a/Engine/Source/Runtime/Datasmith/CADKernel/Public/CADKernel/Topo/TopologicalEdge.h b/Engine/Source/Runtime/Datasmith/CADKernel/Public/CADKernel/Topo/TopologicalEdge.h
--- a/Engine/Source/Runtime/Datasmith/CADKernel/Public/CADKernel/Topo/TopologicalEdge.h
+++ b/Engine/Source/Runtime/Datasmith/CADKernel/Public/CADKernel/Topo/TopologicalEdge.h
@@ -3,6 +3,7 @@
 
 #include "CADKernel/Geo/Curves/Curve.h"
 #include "CADKernel/Geo/Curves/RestrictionCurve.h"
+#include "CADKernel/Geo/Sampling/SurfacicPolyline.h"
 #include "CADKernel/Geo/GeoEnum.h"
 #include "CADKernel/Math/Boundary.h"
 #include "CADKernel/Topo/Linkable.h"
@@ -34,7 +35,7 @@ struct FCuttingPoint;
 struct FImposedCuttingPoint;
 
 class FModelMesh;
-class FRestrictionCurve;
+class FEdgeMesh;
 class FSurface;
 class FTopologicalLoop;
 class FTopologicalVertex;
@@ -419,7 +420,7 @@ public:
 		int32 Size = CrossingPointUs.Num();
 		ensureCADKernel(Size >= 2);
 		CrossingPointUs.SetNum(Size);
-		CrossingPointDeltaUMins.Init(SMALL_NUMBER, Size - 1);
+		CrossingPointDeltaUMins.Init(DOUBLE_SMALL_NUMBER, Size - 1);
 		CrossingPointDeltaUMaxs.Init(2.0 * (GetEndCurvilinearCoordinates() - GetStartCurvilinearCoordinates()), Size - 1);
 	}
 
diff --git a/Engine/Source/Runtime/Datasmith/CADKernel/Public/CADKernel/Topo/TopologicalFace.h b/Engine/Source/Runtime/Datasmith/CADKernel/Public/CADKernel/Topo/TopologicalFace.h
--- a/Engine/Source/Runtime/Datasmith/CADKernel/Public/CADKernel/Topo/TopologicalFace.h
+++ b/Engine/Source/Runtime/Datasmith/CADKernel/Public/CADKernel/Topo/TopologicalFace.h
@@ -600,21 +600,21 @@ struct FBBoxWithNormal
 		for (int32 Index = 0; Index < 3; ++Index)
 		{
 			double DotProduct = MaxPointNormals[Index] | BBoxNormals[Index];
-			if (DotProduct > KINDA_SMALL_NUMBER)
+			if (DotProduct > DOUBLE_KINDA_SMALL_NUMBER)
 			{
 				GoodOrientation++;
 			}
-			else if(DotProduct < KINDA_SMALL_NUMBER)
+			else if(DotProduct < DOUBLE_KINDA_SMALL_NUMBER)
 			{
 				WrongOrientation++;
 			}
 
 			DotProduct = MinPointNormals[Index] | BBoxNormals[Index];
-			if (DotProduct < KINDA_SMALL_NUMBER)
+			if (DotProduct < DOUBLE_KINDA_SMALL_NUMBER)
 			{
 				GoodOrientation++;
 			}
-			else if (DotProduct > KINDA_SMALL_NUMBER)
+			else if (DotProduct > DOUBLE_KINDA_SMALL_NUMBER)
 			{
 				WrongOrientation++;
 			}
diff --git a/Engine/Source/Runtime/Datasmith/CADKernel/Public/CADKernel/Topo/TopologicalLoop.h b/Engine/Source/Runtime/Datasmith/CADKernel/Public/CADKernel/Topo/TopologicalLoop.h
--- a/Engine/Source/Runtime/Datasmith/CADKernel/Public/CADKernel/Topo/TopologicalLoop.h
+++ b/Engine/Source/Runtime/Datasmith/CADKernel/Public/CADKernel/Topo/TopologicalLoop.h
@@ -2,7 +2,6 @@
 #pragma once
 
 #include "CADKernel/Core/OrientedEntity.h"
-#include "CADKernel/Geo/Curves/CompositeCurve.h"
 #include "CADKernel/Geo/GeoEnum.h"
 #include "CADKernel/Topo/TopologicalEdge.h"
 #include "CADKernel/Topo/TopologicalEntity.h"
diff --git a/Engine/Source/Runtime/Datasmith/CADKernel/Public/CADKernel/Topo/TopologicalShapeEntity.h b/Engine/Source/Runtime/Datasmith/CADKernel/Public/CADKernel/Topo/TopologicalShapeEntity.h
--- a/Engine/Source/Runtime/Datasmith/CADKernel/Public/CADKernel/Topo/TopologicalShapeEntity.h
+++ b/Engine/Source/Runtime/Datasmith/CADKernel/Public/CADKernel/Topo/TopologicalShapeEntity.h
@@ -3,7 +3,6 @@
 
 #include "CADKernel/Core/MetadataDictionary.h"
 #include "CADKernel/Topo/TopologicalEntity.h"
-#include "CADKernel/Topo/TopologicalEntity.h"
 
 namespace CADKernel
 {
diff --git a/Engine/Source/Runtime/Datasmith/CADKernel/Public/CADKernel/Topo/Topomaker.h b/Engine/Source/Runtime/Datasmith/CADKernel/Public/CADKernel/Topo/Topomaker.h
--- a/Engine/Source/Runtime/Datasmith/CADKernel/Public/CADKernel/Topo/Topomaker.h
+++ b/Engine/Source/Runtime/Datasmith/CADKernel/Public/CADKernel/Topo/Topomaker.h
@@ -3,7 +3,9 @@
 
 #include "CADKernel/Core/Types.h"
 
+#ifdef CADKERNEL_DEV
 #include "CADKernel/Topo/TopomakerReport.h"
+#endif
 
 namespace CADKernel
 {
@@ -30,7 +32,9 @@ protected:
 	const double SewTolerance;
 	const double SewToleranceSquare;
 
+#ifdef CADKERNEL_DEV
 	FTopomakerReport Report;
+#endif
 
 public:
 
@@ -61,10 +65,12 @@ public:
 	 */
 	void UnlinkNonManifoldVertex();
 
+#ifdef CADKERNEL_DEV
 	void PrintReport()
 	{
 		Report.Print();
 	}
+#endif
 
 private:
 
diff --git a/Engine/Source/Runtime/Datasmith/CADKernel/Public/CADKernel/Topo/TopomakerReport.h b/Engine/Source/Runtime/Datasmith/CADKernel/Public/CADKernel/Topo/TopomakerReport.h
deleted file mode 100644
diff --git a/Engine/Source/Runtime/Datasmith/CADKernel/Public/CADKernel/Utils/IndexOfCoordinateFinder.h b/Engine/Source/Runtime/Datasmith/CADKernel/Public/CADKernel/Utils/IndexOfCoordinateFinder.h
--- a/Engine/Source/Runtime/Datasmith/CADKernel/Public/CADKernel/Utils/IndexOfCoordinateFinder.h
+++ b/Engine/Source/Runtime/Datasmith/CADKernel/Public/CADKernel/Utils/IndexOfCoordinateFinder.h
@@ -60,7 +60,7 @@ namespace CADKernel
 				while(Index > 0)
 				{
 					--Index;
-					if (InCoordinate + SMALL_NUMBER >= Coordinates[Index])
+					if (InCoordinate + DOUBLE_SMALL_NUMBER >= Coordinates[Index])
 					{
 						return Index;
 					}
