[CADKernel] Remove unused file

diff --git a/Engine/Source/Runtime/Datasmith/CADKernel/Private/CADKernel/Mesh/Structure/ModelMeshReport.cpp b/Engine/Source/Runtime/Datasmith/CADKernel/Private/CADKernel/Mesh/Structure/ModelMeshReport.cpp
deleted file mode 100644
diff --git a/Engine/Source/Runtime/Datasmith/CADKernel/Private/CADKernel/Mesh/Structure/ModelMeshReport.h b/Engine/Source/Runtime/Datasmith/CADKernel/Private/CADKernel/Mesh/Structure/ModelMeshReport.h
deleted file mode 100644
diff --git a/Engine/Source/Runtime/Datasmith/CADKernel/Private/CADKernel/Topo/Body.cpp b/Engine/Source/Runtime/Datasmith/CADKernel/Private/CADKernel/Topo/Body.cpp
--- a/Engine/Source/Runtime/Datasmith/CADKernel/Private/CADKernel/Topo/Body.cpp
+++ b/Engine/Source/Runtime/Datasmith/CADKernel/Private/CADKernel/Topo/Body.cpp
@@ -3,8 +3,10 @@
 
 #include "CADKernel/Topo/Model.h"
 #include "CADKernel/Topo/Shell.h"
-#include "CADKernel/Topo/TopologyReport.h"
 
+#ifdef CADKERNEL_DEV
+#include "CADKernel/Topo/TopologyReport.h"
+#endif
 
 namespace CADKernel
 {
@@ -46,7 +48,6 @@ FInfoEntity& FBody::GetInfo(FInfoEntity& Info) const
 {
 	return FTopologicalShapeEntity::GetInfo(Info).Add(TEXT("Shells"), Shells);
 }
-#endif
 
 void FBody::FillTopologyReport(FTopologyReport& Report) const
 {
@@ -56,5 +57,6 @@ void FBody::FillTopologyReport(FTopologyReport& Report) const
 		Shell->FillTopologyReport(Report);
 	}
 }
+#endif
 
 }
diff --git a/Engine/Source/Runtime/Datasmith/CADKernel/Private/CADKernel/Topo/Model.cpp b/Engine/Source/Runtime/Datasmith/CADKernel/Private/CADKernel/Topo/Model.cpp
--- a/Engine/Source/Runtime/Datasmith/CADKernel/Private/CADKernel/Topo/Model.cpp
+++ b/Engine/Source/Runtime/Datasmith/CADKernel/Private/CADKernel/Topo/Model.cpp
@@ -6,7 +6,10 @@
 #include "CADKernel/Topo/TopologicalEdge.h"
 #include "CADKernel/Topo/TopologicalFace.h"
 #include "CADKernel/Topo/TopologicalVertex.h"
+
+#ifdef CADKERNEL_DEV
 #include "CADKernel/Topo/TopologyReport.h"
+#endif
 
 namespace CADKernel
 {
@@ -193,6 +196,7 @@ for (const FFaceSubset& Subset : SubShells)
 	}
 }
 
+#ifdef CADKERNEL_DEV
 void FModel::FillTopologyReport(FTopologyReport& Report) const
 {
 	for (TSharedPtr<FBody> Body : Bodies)
@@ -205,6 +209,7 @@ void FModel::FillTopologyReport(FTopologyReport& Report) const
 		Face->FillTopologyReport(Report);
 	}
 }
+#endif
 
 void FModel::Orient()
 {
diff --git a/Engine/Source/Runtime/Datasmith/CADKernel/Private/CADKernel/Topo/Shell.cpp b/Engine/Source/Runtime/Datasmith/CADKernel/Private/CADKernel/Topo/Shell.cpp
--- a/Engine/Source/Runtime/Datasmith/CADKernel/Private/CADKernel/Topo/Shell.cpp
+++ b/Engine/Source/Runtime/Datasmith/CADKernel/Private/CADKernel/Topo/Shell.cpp
@@ -4,7 +4,10 @@
 #include "CADKernel/Geo/GeoEnum.h"
 #include "CADKernel/Topo/Body.h"
 #include "CADKernel/Topo/TopologicalFace.h"
+
+#ifdef CADKERNEL_DEV
 #include "CADKernel/Topo/TopologyReport.h"
+#endif
 
 namespace CADKernel
 {
@@ -249,7 +252,8 @@ void FShell::CheckTopology(TArray<FFaceSubset>& Subshells)
 	ResetMarkersRecursively();
 }
 
-void FShell::FillTopologyReport(FTopologyReport& Report) const 
+#ifdef CADKERNEL_DEV
+void FShell::FillTopologyReport(FTopologyReport& Report) const
 {
 	Report.Add(this);
 	for (const FOrientedFace& OrientedFace : GetFaces())
@@ -257,6 +261,7 @@ void FShell::FillTopologyReport(FTopologyReport& Report) const
 		OrientedFace.Entity->FillTopologyReport(Report);
 	}
 }
+#endif
 
 int32 FShell::Orient()
 {
diff --git a/Engine/Source/Runtime/Datasmith/CADKernel/Private/CADKernel/Topo/TopologicalFace.cpp b/Engine/Source/Runtime/Datasmith/CADKernel/Private/CADKernel/Topo/TopologicalFace.cpp
--- a/Engine/Source/Runtime/Datasmith/CADKernel/Private/CADKernel/Topo/TopologicalFace.cpp
+++ b/Engine/Source/Runtime/Datasmith/CADKernel/Private/CADKernel/Topo/TopologicalFace.cpp
@@ -14,7 +14,10 @@
 #include "CADKernel/Mesh/Structure/Grid.h"
 #include "CADKernel/Topo/Shell.h"
 #include "CADKernel/Topo/TopologicalEdge.h"
+
+#ifdef CADKERNEL_DEV
 #include "CADKernel/Topo/TopologyReport.h"
+#endif
 
 namespace CADKernel
 {
@@ -337,6 +340,7 @@ const FTopologicalEdge* FTopologicalFace::GetLinkedEdge(const FTopologicalEdge&
 	return nullptr;
 }
 
+#ifdef CADKERNEL_DEV
 void FTopologicalFace::FillTopologyReport(FTopologyReport& Report) const
 {
 	Report.Add(this);
@@ -349,6 +353,7 @@ void FTopologicalFace::FillTopologyReport(FTopologyReport& Report) const
 		}
 	}
 }
+#endif
 
 void FTopologicalFace::GetEdgeIndex(const FTopologicalEdge& Edge, int32& OutBoundaryIndex, int32& OutEdgeIndex) const
 {
diff --git a/Engine/Source/Runtime/Datasmith/CADKernel/Private/CADKernel/Topo/TopologyReport.cpp b/Engine/Source/Runtime/Datasmith/CADKernel/Private/CADKernel/Topo/TopologyReport.cpp
deleted file mode 100644
diff --git a/Engine/Source/Runtime/Datasmith/CADKernel/Private/CADKernel/Topo/TopologyReport.h b/Engine/Source/Runtime/Datasmith/CADKernel/Private/CADKernel/Topo/TopologyReport.h
deleted file mode 100644
diff --git a/Engine/Source/Runtime/Datasmith/CADKernel/Public/CADKernel/Mesh/Meshers/MesherReport.h b/Engine/Source/Runtime/Datasmith/CADKernel/Public/CADKernel/Mesh/Meshers/MesherReport.h
deleted file mode 100644
diff --git a/Engine/Source/Runtime/Datasmith/CADKernel/Public/CADKernel/Topo/Body.h b/Engine/Source/Runtime/Datasmith/CADKernel/Public/CADKernel/Topo/Body.h
--- a/Engine/Source/Runtime/Datasmith/CADKernel/Public/CADKernel/Topo/Body.h
+++ b/Engine/Source/Runtime/Datasmith/CADKernel/Public/CADKernel/Topo/Body.h
@@ -114,7 +114,9 @@ public:
 		}
 	}
 
+#ifdef CADKERNEL_DEV
 	virtual void FillTopologyReport(FTopologyReport& Report) const override;
+#endif
 
 };
 
diff --git a/Engine/Source/Runtime/Datasmith/CADKernel/Public/CADKernel/Topo/Model.h b/Engine/Source/Runtime/Datasmith/CADKernel/Public/CADKernel/Topo/Model.h
--- a/Engine/Source/Runtime/Datasmith/CADKernel/Public/CADKernel/Topo/Model.h
+++ b/Engine/Source/Runtime/Datasmith/CADKernel/Public/CADKernel/Topo/Model.h
@@ -183,7 +183,9 @@ namespace CADKernel
 		 */
 		void CheckTopology();
 
+#ifdef CADKERNEL_DEV
 		virtual void FillTopologyReport(FTopologyReport& Report) const override;
+#endif
 
 		/**
 		 * Fore each body
diff --git a/Engine/Source/Runtime/Datasmith/CADKernel/Public/CADKernel/Topo/Shell.h b/Engine/Source/Runtime/Datasmith/CADKernel/Public/CADKernel/Topo/Shell.h
--- a/Engine/Source/Runtime/Datasmith/CADKernel/Public/CADKernel/Topo/Shell.h
+++ b/Engine/Source/Runtime/Datasmith/CADKernel/Public/CADKernel/Topo/Shell.h
@@ -132,7 +132,10 @@ public:
 	virtual void UpdateShellOrientation();
 
 	void CheckTopology(TArray<FFaceSubset>& Subshells);
+
+#ifdef CADKERNEL_DEV
 	virtual void FillTopologyReport(FTopologyReport& Report) const override;
+#endif
 
 	/**
 	 * @return true if the shell has at least one border edge
diff --git a/Engine/Source/Runtime/Datasmith/CADKernel/Public/CADKernel/Topo/TopologicalFace.h b/Engine/Source/Runtime/Datasmith/CADKernel/Public/CADKernel/Topo/TopologicalFace.h
--- a/Engine/Source/Runtime/Datasmith/CADKernel/Public/CADKernel/Topo/TopologicalFace.h
+++ b/Engine/Source/Runtime/Datasmith/CADKernel/Public/CADKernel/Topo/TopologicalFace.h
@@ -276,7 +276,9 @@ public:
 	 */
 	void RemoveLinksWithNeighbours();
 
+#ifdef CADKERNEL_DEV
 	virtual void FillTopologyReport(FTopologyReport& Report) const override;
+#endif
 
 	// ======   Meshing Function   ======
 
diff --git a/Engine/Source/Runtime/Datasmith/CADKernel/Public/CADKernel/Topo/TopologicalShapeEntity.h b/Engine/Source/Runtime/Datasmith/CADKernel/Public/CADKernel/Topo/TopologicalShapeEntity.h
--- a/Engine/Source/Runtime/Datasmith/CADKernel/Public/CADKernel/Topo/TopologicalShapeEntity.h
+++ b/Engine/Source/Runtime/Datasmith/CADKernel/Public/CADKernel/Topo/TopologicalShapeEntity.h
@@ -46,7 +46,9 @@ namespace CADKernel
 		 */
 		virtual void SpreadBodyOrientation() = 0;
 
+#ifdef CADKERNEL_DEV
 		virtual void FillTopologyReport(FTopologyReport& Report) const = 0;
+#endif
 
 		FTopologicalShapeEntity* GetHost()
 		{
