Datasmith: Fix skipped values for Lights elements

diff --git a/Engine/Source/Runtime/Datasmith/DatasmithCore/Private/DatasmithSceneXmlWriter.cpp b/Engine/Source/Runtime/Datasmith/DatasmithCore/Private/DatasmithSceneXmlWriter.cpp
--- a/Engine/Source/Runtime/Datasmith/DatasmithCore/Private/DatasmithSceneXmlWriter.cpp
+++ b/Engine/Source/Runtime/Datasmith/DatasmithCore/Private/DatasmithSceneXmlWriter.cpp
@@ -958,20 +958,9 @@ void FDatasmithSceneXmlWriterImpl::WriteLightActorElement(const TSharedPtr< IDat
 
 void FDatasmithSceneXmlWriterImpl::WritePointLightElement(const TSharedRef< IDatasmithPointLightElement >& PointLightElement, FArchive& Archive, int32 Indent)
 {
-	if ( PointLightElement->GetSourceRadius() > 0.0f )
-	{
-		WriteValue( Archive, Indent + 1, DATASMITH_LIGHTSOURCESIZENAME, PointLightElement->GetSourceRadius() );
-	}
-
-	if ( PointLightElement->GetSourceLength() > 0.0f )
-	{
-		WriteValue( Archive, Indent + 1, DATASMITH_LIGHTSOURCELENGTHNAME, PointLightElement->GetSourceLength() );
-	}
-
-	if ( PointLightElement->GetAttenuationRadius() > 20.0f && PointLightElement->GetAttenuationRadius() < 2000.0f ) // 20 cm min 20m max
-	{
-		WriteValue( Archive, Indent + 1, DATASMITH_LIGHTATTENUATIONRADIUSNAME, PointLightElement->GetAttenuationRadius() );
-	}
+	WriteValue( Archive, Indent + 1, DATASMITH_LIGHTSOURCESIZENAME, PointLightElement->GetSourceRadius() );
+	WriteValue( Archive, Indent + 1, DATASMITH_LIGHTSOURCELENGTHNAME, PointLightElement->GetSourceLength() );
+	WriteValue( Archive, Indent + 1, DATASMITH_LIGHTATTENUATIONRADIUSNAME, PointLightElement->GetAttenuationRadius() );
 
 	FString UnitsString;
 	switch ( PointLightElement->GetIntensityUnits() )
