Datasmith: Expose the CL number in our product version

diff --git a/Engine/Source/Runtime/Datasmith/DatasmithCore/Private/DatasmithUtils.cpp b/Engine/Source/Runtime/Datasmith/DatasmithCore/Private/DatasmithUtils.cpp
--- a/Engine/Source/Runtime/Datasmith/DatasmithCore/Private/DatasmithUtils.cpp
+++ b/Engine/Source/Runtime/Datasmith/DatasmithCore/Private/DatasmithUtils.cpp
@@ -175,7 +175,7 @@ int32 FDatasmithUtils::GetEnterpriseVersionAsInt()
 
 FString FDatasmithUtils::GetEnterpriseVersionAsString()
 {
-	return FEngineVersion::Current().ToString( EVersionComponent::Patch );
+	return FEngineVersion::Current().ToString( EVersionComponent::Changelist );
 }
 
 float FDatasmithUtils::GetDatasmithFormatVersionAsFloat()
