Datasmith: Fix light serialization

diff --git a/Engine/Source/Runtime/Datasmith/DatasmithCore/Private/DatasmithSceneXmlWriter.cpp b/Engine/Source/Runtime/Datasmith/DatasmithCore/Private/DatasmithSceneXmlWriter.cpp
--- a/Engine/Source/Runtime/Datasmith/DatasmithCore/Private/DatasmithSceneXmlWriter.cpp
+++ b/Engine/Source/Runtime/Datasmith/DatasmithCore/Private/DatasmithSceneXmlWriter.cpp
@@ -976,9 +976,20 @@ void FDatasmithSceneXmlWriterImpl::WriteLightActorElement(const TSharedPtr< IDat
 
 void FDatasmithSceneXmlWriterImpl::WritePointLightElement(const TSharedRef< IDatasmithPointLightElement >& PointLightElement, FArchive& Archive, int32 Indent)
 {
-	WriteValue( Archive, Indent + 1, DATASMITH_LIGHTSOURCESIZENAME, PointLightElement->GetSourceRadius() );
-	WriteValue( Archive, Indent + 1, DATASMITH_LIGHTSOURCELENGTHNAME, PointLightElement->GetSourceLength() );
-	WriteValue( Archive, Indent + 1, DATASMITH_LIGHTATTENUATIONRADIUSNAME, PointLightElement->GetAttenuationRadius() );
+	if (PointLightElement->GetSourceRadius() >= 0.f)
+	{
+		WriteValue( Archive, Indent + 1, DATASMITH_LIGHTSOURCESIZENAME, PointLightElement->GetSourceRadius() );
+	}
+
+	if (PointLightElement->GetSourceLength() >= 0.f)
+	{
+		WriteValue( Archive, Indent + 1, DATASMITH_LIGHTSOURCELENGTHNAME, PointLightElement->GetSourceLength() );
+	}
+
+	if (PointLightElement->GetAttenuationRadius() >= 0.f)
+	{
+		WriteValue( Archive, Indent + 1, DATASMITH_LIGHTATTENUATIONRADIUSNAME, PointLightElement->GetAttenuationRadius() );
+	}
 
 	FString UnitsString;
 	switch ( PointLightElement->GetIntensityUnits() )
