Datasmith: Fix missing CL info when building plugins on build machines

diff --git a/Engine/Source/Runtime/Datasmith/DatasmithCore/Private/DatasmithSceneXmlWriter.cpp b/Engine/Source/Runtime/Datasmith/DatasmithCore/Private/DatasmithSceneXmlWriter.cpp
--- a/Engine/Source/Runtime/Datasmith/DatasmithCore/Private/DatasmithSceneXmlWriter.cpp
+++ b/Engine/Source/Runtime/Datasmith/DatasmithCore/Private/DatasmithSceneXmlWriter.cpp
@@ -1901,7 +1901,7 @@ void FDatasmithSceneXmlWriter::Serialize( TSharedRef< IDatasmithScene > Datasmit
 
 	XmlString = FString::Printf( TEXT("<%s>%s</%s>"), DATASMITH_EXPORTERVERSION, *FDatasmithUtils::GetDatasmithFormatVersionAsString(), DATASMITH_EXPORTERVERSION ) + LINE_TERMINATOR;
 	FDatasmithSceneXmlWriterImpl::AppendIndent( XmlString, 1 );
-	XmlString += FString::Printf( TEXT("<%s>%s</%s>"), DATASMITH_EXPORTERSDKVERSION, *FDatasmithUtils::GetEnterpriseVersionAsString(), DATASMITH_EXPORTERSDKVERSION ) + LINE_TERMINATOR;
+	XmlString += FString::Printf( TEXT("<%s>%s</%s>"), DATASMITH_EXPORTERSDKVERSION, *FDatasmithUtils::GetEnterpriseVersionAsString(true), DATASMITH_EXPORTERSDKVERSION ) + LINE_TERMINATOR;
 	FDatasmithSceneXmlWriterImpl::SerializeToArchive( Archive, XmlString );
 
 	FDatasmithSceneXmlWriterImpl::WriteIndent(Archive, 1);
diff --git a/Engine/Source/Runtime/Datasmith/DatasmithCore/Private/DatasmithUtils.cpp b/Engine/Source/Runtime/Datasmith/DatasmithCore/Private/DatasmithUtils.cpp
--- a/Engine/Source/Runtime/Datasmith/DatasmithCore/Private/DatasmithUtils.cpp
+++ b/Engine/Source/Runtime/Datasmith/DatasmithCore/Private/DatasmithUtils.cpp
@@ -173,9 +173,10 @@ int32 FDatasmithUtils::GetEnterpriseVersionAsInt()
 	return MajorVersion + MinorVersion + PatchVersion;
 }
 
-FString FDatasmithUtils::GetEnterpriseVersionAsString()
+FString FDatasmithUtils::GetEnterpriseVersionAsString(bool bWithChangelist)
 {
-	return FEngineVersion::Current().ToString( EVersionComponent::Changelist );
+	EVersionComponent Format = (bWithChangelist && FEngineVersion::Current().HasChangelist()) ? EVersionComponent::Changelist : EVersionComponent::Patch;
+	return FEngineVersion::Current().ToString(Format);
 }
 
 float FDatasmithUtils::GetDatasmithFormatVersionAsFloat()
diff --git a/Engine/Source/Runtime/Datasmith/DatasmithCore/Public/DatasmithUtils.h b/Engine/Source/Runtime/Datasmith/DatasmithCore/Public/DatasmithUtils.h
--- a/Engine/Source/Runtime/Datasmith/DatasmithCore/Public/DatasmithUtils.h
+++ b/Engine/Source/Runtime/Datasmith/DatasmithCore/Public/DatasmithUtils.h
@@ -36,7 +36,7 @@ public:
 	static void SanitizeStringInplace(FString& InString);
 
 	static int32 GetEnterpriseVersionAsInt();
-	static FString GetEnterpriseVersionAsString();
+	static FString GetEnterpriseVersionAsString(bool bWithChangelist=false);
 
 	/** Returns the Datasmith data format version */
 	static float GetDatasmithFormatVersionAsFloat();
