Datasmith: fix duplicated vertices and missing vertex color

diff --git a/Engine/Source/Runtime/Datasmith/DatasmithCore/Private/DatasmithUtils.cpp b/Engine/Source/Runtime/Datasmith/DatasmithCore/Private/DatasmithUtils.cpp
--- a/Engine/Source/Runtime/Datasmith/DatasmithCore/Private/DatasmithUtils.cpp
+++ b/Engine/Source/Runtime/Datasmith/DatasmithCore/Private/DatasmithUtils.cpp
@@ -295,106 +295,18 @@ bool FDatasmithMeshUtils::ToRawMesh(const FDatasmithMesh& Mesh, FRawMesh& RawMes
 
 bool FDatasmithMeshUtils::ToMeshDescription(const FDatasmithMesh& DsMesh, FMeshDescription& MeshDescription)
 {
-	MeshDescription.Empty();
-
-	FStaticMeshAttributes Attributes(MeshDescription);
-	Attributes.Register();
-	TVertexAttributesRef<FVector3f> VertexPositions = Attributes.GetVertexPositions();
-	TVertexInstanceAttributesRef<FVector3f> VertexInstanceNormals = Attributes.GetVertexInstanceNormals();
-	TVertexInstanceAttributesRef<FVector2f> VertexInstanceUVs = Attributes.GetVertexInstanceUVs();
-	TPolygonGroupAttributesRef<FName> PolygonGroupImportedMaterialSlotNames = Attributes.GetPolygonGroupMaterialSlotNames();
-
-	// Reserve space for attributes.
-	int32 VertexCount = DsMesh.GetVerticesCount();
-	int32 TriangleCount = DsMesh.GetFacesCount();
-	int32 VertexInstanceCount = 3 * TriangleCount;
-	int32 MaterialCount = DsMesh.GetMaterialsCount();
-	MeshDescription.ReserveNewVertices(VertexCount);
-	MeshDescription.ReserveNewVertexInstances(VertexInstanceCount);
-	MeshDescription.ReserveNewEdges(VertexInstanceCount);
-	MeshDescription.ReserveNewPolygons(TriangleCount);
-	MeshDescription.ReserveNewPolygonGroups(MaterialCount);
-
-	// At least one UV set must exist.
-	int32 DsUVCount = DsMesh.GetUVChannelsCount();
-	int32 MeshDescUVCount = FMath::Max(1, DsUVCount);
-	VertexInstanceUVs.SetNumChannels(MeshDescUVCount);
-
-	//Fill the vertex array
-	for (int32 VertexIndex = 0; VertexIndex < VertexCount; ++VertexIndex)
-	{
-		FVertexID AddedVertexId = MeshDescription.CreateVertex();
-		VertexPositions[AddedVertexId] = (FVector3f)DsMesh.GetVertex(VertexIndex);
-	}
-
-	TMap<int32, FPolygonGroupID> PolygonGroupMapping;
-	auto GetOrCreatePolygonGroupId = [&](int32 MaterialIndex)
+	FRawMesh RawMesh;
+	ToRawMesh(DsMesh, RawMesh);
+	TMap<int32, FName> MaterialMap;
+	TSet<int32> MaterialIds{RawMesh.FaceMaterialIndices};
+	for (int32 MaterialId : MaterialIds)
 	{
-		FPolygonGroupID& PolyGroupId = PolygonGroupMapping.FindOrAdd(MaterialIndex);
-		if (PolyGroupId == INDEX_NONE)
-		{
-			PolyGroupId = MeshDescription.CreatePolygonGroup();
-			FName ImportedSlotName = *FString::FromInt(MaterialIndex); // No access to DatasmithMeshHelper::DefaultSlotName
-			PolygonGroupImportedMaterialSlotNames[PolyGroupId] = ImportedSlotName;
-		}
-		return PolyGroupId;
-	};
-
-	// corner informations
-	const int32 CornerCount = 3; // only triangles in DatasmithMesh
-	FVector CornerPositions[3];
-	TArray<FVertexInstanceID> CornerVertexInstanceIDs;
-	CornerVertexInstanceIDs.SetNum(3);
-	FVertexID CornerVertexIDs[3];
-	TArray<uint32> FaceSmoothingMasks;
-	for (int32 PolygonIndex = 0; PolygonIndex < TriangleCount; PolygonIndex++)
-	{
-		// face basics info
-		int32 MaterialIndex;
-		int32 VertexIndex[3];
-		DsMesh.GetFace(PolygonIndex, VertexIndex[0], VertexIndex[1], VertexIndex[2], MaterialIndex);
-		for (int32 CornerIndex = 0; CornerIndex < CornerCount; CornerIndex++)
-		{
-			CornerVertexIDs[CornerIndex] = FVertexID(VertexIndex[CornerIndex]);
-			CornerPositions[CornerIndex] = (FVector)VertexPositions[CornerVertexIDs[CornerIndex]];
-		}
-
-		// Create Vertex instances
-		for (int32 CornerIndex = 0; CornerIndex < CornerCount; CornerIndex++)
-		{
-			CornerVertexInstanceIDs[CornerIndex] = MeshDescription.CreateVertexInstance(CornerVertexIDs[CornerIndex]);
-		}
-
-		// UVs attributes
-		for (int32 UVChannelIndex = 0; UVChannelIndex < DsUVCount; ++UVChannelIndex)
-		{
-			int32 UV[3];
-			DsMesh.GetFaceUV(PolygonIndex, UVChannelIndex, UV[0], UV[1], UV[2]);
-			for (int32 CornerIndex = 0; CornerIndex < CornerCount; ++CornerIndex)
-			{
-				check(UV[CornerIndex] < DsMesh.GetUVCount(UVChannelIndex))
-				FVector2D UVVector = DsMesh.GetUV(UVChannelIndex, UV[CornerIndex]);
-				if (!UVVector.ContainsNaN())
-				{
-					VertexInstanceUVs.Set(CornerVertexInstanceIDs[CornerIndex], UVChannelIndex, FVector2f(UVVector));
-				}
-			}
-		}
-
-		for (int32 CornerIndex = 0; CornerIndex < CornerCount; CornerIndex++)
-		{
-			VertexInstanceNormals[CornerVertexInstanceIDs[CornerIndex]] = (FVector3f)DsMesh.GetNormal(3 * PolygonIndex + CornerIndex);
-		}
-
-		// smoothing information
-		FaceSmoothingMasks.Add(DsMesh.GetFaceSmoothingMask(PolygonIndex));
-
-		// Create in-mesh Polygon
-		const FPolygonGroupID PolygonGroupID = GetOrCreatePolygonGroupId(MaterialIndex);
-		const FPolygonID NewPolygonID = MeshDescription.CreatePolygon(PolygonGroupID, CornerVertexInstanceIDs);
+		MaterialMap.Add(MaterialId, *FString::FromInt(MaterialId));
 	}
 
-	FStaticMeshOperations::ConvertSmoothGroupToHardEdges(FaceSmoothingMasks, MeshDescription);
+	FStaticMeshAttributes Attributes(MeshDescription);
+	Attributes.Register();
+	FStaticMeshOperations::ConvertFromRawMesh(RawMesh, MeshDescription, MaterialMap);
 
 	return true;
 }
