Datasmith SDK: Rework handling of meshes exported without UVs data.

diff --git a/Engine/Source/Developer/Datasmith/DatasmithExporter/Private/DatasmithMeshExporter.cpp b/Engine/Source/Developer/Datasmith/DatasmithExporter/Private/DatasmithMeshExporter.cpp
--- a/Engine/Source/Developer/Datasmith/DatasmithExporter/Private/DatasmithMeshExporter.cpp
+++ b/Engine/Source/Developer/Datasmith/DatasmithExporter/Private/DatasmithMeshExporter.cpp
@@ -53,10 +53,6 @@ bool FDatasmithMeshExporterImpl::DoExport(TSharedPtr<IDatasmithMeshElement>& Mes
 	{
 		Mesh.SetName(*FPaths::GetBaseFilename(ExportOptions.MeshFullPath));
 	}
-
-	// make sure we have at least one UV channels on each source models
-	FDatasmithMeshUtils::CreateDefaultUVsWithLOD(Mesh);
-
 	FMD5Hash Hash;
 	if (WriteMeshFile(ExportOptions, Hash))
 	{
@@ -120,14 +116,14 @@ bool FDatasmithMeshExporterImpl::WriteMeshFile(const FDatasmithMeshExporterOptio
 
 		Models.SourceModels.Reserve(Mesh.GetLODsCount() + 1);
 		FMeshDescription& BaseMeshDescription = Models.SourceModels.AddDefaulted_GetRef();
-		FDatasmithMeshUtils::ToMeshDescription(Mesh, BaseMeshDescription);
+		FDatasmithMeshUtils::ToMeshDescription(Mesh, BaseMeshDescription, FDatasmithMeshUtils::GenerateBox);
 
 		for (int32 LodIndex = 0; LodIndex < Mesh.GetLODsCount(); ++LodIndex)
 		{
-			if (const FDatasmithMesh* LodMesh = Mesh.GetLOD(LodIndex))
+			if (FDatasmithMesh* LodMesh = Mesh.GetLOD(LodIndex))
 			{
 				FMeshDescription& LodMeshDescription = Models.SourceModels.AddDefaulted_GetRef();
-				FDatasmithMeshUtils::ToMeshDescription(*LodMesh, LodMeshDescription);
+				FDatasmithMeshUtils::ToMeshDescription(*LodMesh, LodMeshDescription, FDatasmithMeshUtils::GenerateBox);
 			}
 		}
 
diff --git a/Engine/Source/Runtime/Datasmith/DatasmithCore/Private/DatasmithUtils.cpp b/Engine/Source/Runtime/Datasmith/DatasmithCore/Private/DatasmithUtils.cpp
--- a/Engine/Source/Runtime/Datasmith/DatasmithCore/Private/DatasmithUtils.cpp
+++ b/Engine/Source/Runtime/Datasmith/DatasmithCore/Private/DatasmithUtils.cpp
@@ -293,10 +293,19 @@ bool FDatasmithMeshUtils::ToRawMesh(const FDatasmithMesh& Mesh, FRawMesh& RawMes
 	return true;
 }
 
-bool FDatasmithMeshUtils::ToMeshDescription(const FDatasmithMesh& DsMesh, FMeshDescription& MeshDescription)
+bool FDatasmithMeshUtils::ToMeshDescription(FDatasmithMesh& DsMesh, FMeshDescription& MeshDescription, EUvGenerationPolicy UvGen)
 {
 	FRawMesh RawMesh;
-	ToRawMesh(DsMesh, RawMesh);
+	bool bGenerateDefaultUvs = DsMesh.GetUVChannelsCount() == 0 && UvGen != EUvGenerationPolicy::Ignore;
+	if (bGenerateDefaultUvs)
+	{
+		DsMesh.AddUVChannel();
+	}
+
+	if (!ToRawMesh(DsMesh, RawMesh))
+	{
+		return false;
+	}
 	TMap<int32, FName> MaterialMap;
 	TSet<int32> MaterialIds{RawMesh.FaceMaterialIndices};
 	for (int32 MaterialId : MaterialIds)
@@ -308,6 +317,19 @@ bool FDatasmithMeshUtils::ToMeshDescription(const FDatasmithMesh& DsMesh, FMeshD
 	Attributes.Register();
 	FStaticMeshOperations::ConvertFromRawMesh(RawMesh, MeshDescription, MaterialMap);
 
+	if (bGenerateDefaultUvs)
+	{
+		FBox MeshBoundingBox = MeshDescription.ComputeBoundingBox();
+		FUVMapParameters UVParameters(MeshBoundingBox.GetCenter(), FQuat::Identity, MeshBoundingBox.GetSize(), FVector::OneVector, FVector2D::UnitVector);
+		TMap<FVertexInstanceID, FVector2D> TexCoords;
+		FStaticMeshOperations::GenerateBoxUV(MeshDescription, UVParameters, TexCoords);
+
+		TVertexInstanceAttributesRef<FVector2f> UVs = FStaticMeshAttributes(MeshDescription).GetVertexInstanceUVs();
+		for (const auto& Pair : TexCoords)
+		{
+			UVs.Set(Pair.Key, 0, (FVector2f)Pair.Value);
+		}
+	}
 	return true;
 }
 
@@ -346,7 +368,7 @@ void CreateDefaultUVs( FDatasmithMesh& Mesh )
 	// Get the mesh description to generate BoxUV.
 	FMeshDescription MeshDescription;
 	FStaticMeshAttributes(MeshDescription).Register();
-	FDatasmithMeshUtils::ToMeshDescription(Mesh, MeshDescription);
+	FDatasmithMeshUtils::ToMeshDescription(Mesh, MeshDescription, FDatasmithMeshUtils::Ignore);
 	FUVMapParameters UVParameters(
 		FVector(Mesh.GetExtents().GetCenter()),
 		FQuat::Identity,
diff --git a/Engine/Source/Runtime/Datasmith/DatasmithCore/Public/DatasmithUtils.h b/Engine/Source/Runtime/Datasmith/DatasmithCore/Public/DatasmithUtils.h
--- a/Engine/Source/Runtime/Datasmith/DatasmithCore/Public/DatasmithUtils.h
+++ b/Engine/Source/Runtime/Datasmith/DatasmithCore/Public/DatasmithUtils.h
@@ -148,7 +148,13 @@ public:
 	 * For example a collision mesh don't need to have a valid raw mesh
 	 */
 	static bool ToRawMesh(const FDatasmithMesh& Mesh, FRawMesh& RawMesh, bool bValidateRawMesh = true);
-	static bool ToMeshDescription(const FDatasmithMesh& DsMesh, FMeshDescription& MeshDescription);
+
+	enum EUvGenerationPolicy
+	{
+		Ignore, // do not generate a default UV. Conversion fails when no UV is available
+		GenerateBox, // generate a Box UV mapping on the resulting MeshDescription (see also FStaticMeshOperations::GenerateBoxUV)
+	};
+	static bool ToMeshDescription(FDatasmithMesh& DsMesh, FMeshDescription& MeshDescription, EUvGenerationPolicy UvGen=EUvGenerationPolicy::Ignore);
 
 	/**
 	 * Validates that the given UV Channel does not contain a degenerated triangle.
