Datasmith: Better handling of non-compact material ids

diff --git a/Engine/Source/Runtime/Datasmith/DatasmithCore/Private/DatasmithUtils.cpp b/Engine/Source/Runtime/Datasmith/DatasmithCore/Private/DatasmithUtils.cpp
--- a/Engine/Source/Runtime/Datasmith/DatasmithCore/Private/DatasmithUtils.cpp
+++ b/Engine/Source/Runtime/Datasmith/DatasmithCore/Private/DatasmithUtils.cpp
@@ -10,6 +10,8 @@
 #include "DatasmithVariantElements.h"
 #include "IDatasmithSceneElements.h"
 
+#include "Algo/MinElement.h"
+#include "Algo/MaxElement.h"
 #include "Algo/Sort.h"
 #include "HAL/FileManager.h"
 #include "HAL/PlatformFileManager.h"
@@ -300,17 +302,48 @@ bool FDatasmithMeshUtils::ToMeshDescription(FDatasmithMesh& DsMesh, FMeshDescrip
 	if (bGenerateDefaultUvs)
 	{
 		DsMesh.AddUVChannel();
+		DsMesh.SetUVCount(0, 1); // we need one valid uv value so that default UVIndex (which is 0) points to a valid location
 	}
 
 	if (!ToRawMesh(DsMesh, RawMesh))
 	{
 		return false;
 	}
+
+	// ToRawMesh() produces meshes whose material indices are in fact Ids, eg. can be any random int, instead of the compact set {0, 1, 2, ...}.
+	// we have to remap ids to compact indices in order to generate MeshDescriptions with the correct amount of polygon groups.
+	// The following steps ensure that
+	// - Slot names are FNames that are the literal representation of the original materialId (so that the importer can map materials)
+	// - MeshDescription creation receives a RawMesh with a compact set of materialIndices (in RawMesh.FaceMaterialIndices)
+
 	TMap<int32, FName> MaterialMap;
 	TSet<int32> MaterialIds{RawMesh.FaceMaterialIndices};
-	for (int32 MaterialId : MaterialIds)
+
+	int32* Min = Algo::MinElement(MaterialIds);
+	int32* Max = Algo::MaxElement(MaterialIds);
+	bool bCanUseIdsDirectly = Min && Max && *Min == 0 && *Max == MaterialIds.Num()-1; // We can skip the remapping when ids are already indices
+
+	if (bCanUseIdsDirectly)
 	{
-		MaterialMap.Add(MaterialId, *FString::FromInt(MaterialId));
+		for (int32 MaterialId : MaterialIds)
+		{
+			MaterialMap.Add(MaterialId, *FString::FromInt(MaterialId));
+		}
+	}
+	else
+	{
+		TMap<int32, int32> IDsToIndices;
+		for (int32 MaterialId : MaterialIds)
+		{
+			int32 MaterialIndex = MaterialMap.Num();
+			IDsToIndices.Add(MaterialId, MaterialIndex);
+
+			MaterialMap.Add(MaterialIndex, *FString::FromInt(MaterialId));
+		}
+		for (auto& OldId : RawMesh.FaceMaterialIndices)
+		{
+			OldId = IDsToIndices[OldId];
+		}
 	}
 
 	FStaticMeshAttributes Attributes(MeshDescription);
