Datasmith: Add proper version information on custom FDatasmithPackedMeshes archives

diff --git a/Engine/Source/Runtime/Datasmith/DatasmithCore/Private/DatasmithMeshSerialization.cpp b/Engine/Source/Runtime/Datasmith/DatasmithCore/Private/DatasmithMeshSerialization.cpp
--- a/Engine/Source/Runtime/Datasmith/DatasmithCore/Private/DatasmithMeshSerialization.cpp
+++ b/Engine/Source/Runtime/Datasmith/DatasmithCore/Private/DatasmithMeshSerialization.cpp
@@ -142,6 +142,15 @@ FMD5Hash FDatasmithPackedMeshes::Serialize(FArchive& Ar, bool bCompressed)
 {
 	using namespace DatasmithMeshSerializationImpl;
 
+	// structure          from
+	// | Guard           | v0 | constant string, prevents loading of invalid buffers
+	// | SerialVersion   | v0 | this structure versioning
+	// | BufferType      | v0 | describe the payload
+	// | UEVer           | v1 | global engine version
+	// | LicenseeUEVer   | v1 |
+	// | CustomVersions  | v0 | internal structures use them, we have to pass them along
+	// | Mesh buffer     | v0 | may be compressed, the actual payload
+
 	FString Guard = Ar.IsLoading() ? TEXT("") : TEXT("FDatasmithPackedMeshes");
 	Ar << Guard;
 	if (!ensure(Guard == TEXT("FDatasmithPackedMeshes")))
@@ -150,7 +159,7 @@ FMD5Hash FDatasmithPackedMeshes::Serialize(FArchive& Ar, bool bCompressed)
 		return {};
 	}
 
-	uint32 SerialVersion = 0;
+	uint32 SerialVersion = 1;
 	Ar << SerialVersion;
 
 	enum EBufferType{ RawMeshDescription, CompressedMeshDescription };
@@ -160,34 +169,57 @@ FMD5Hash FDatasmithPackedMeshes::Serialize(FArchive& Ar, bool bCompressed)
 	FMD5Hash OutHash;
 	if (Ar.IsLoading())
 	{
+		// versioning
+		if (SerialVersion >= 1)
+		{
+			FPackageFileVersion LoadingUEVer;
+			Ar << LoadingUEVer;
+			Ar.SetUEVer(LoadingUEVer);
+
+			int32 LoadingLicenseeUEVer = 0;
+			Ar << LoadingLicenseeUEVer;
+			Ar.SetLicenseeUEVer(LoadingLicenseeUEVer);
+		}
 		FCustomVersionContainer CustomVersions;
 		CustomVersions.Serialize(Ar);
 
+		// content
 		TArray<uint8> Bytes;
 		Ar << Bytes;
-
 		if (BufferType == CompressedMeshDescription)
 		{
 			DecompressInline(Bytes);
 		}
 
-		FMemoryReader Buffer(Bytes, true);
-		Buffer.SetCustomVersions(CustomVersions);
-		Buffer << Meshes;
+		FMemoryReader BytesReader(Bytes, true);
+		BytesReader.SetLicenseeUEVer(Ar.LicenseeUEVer());
+		BytesReader.SetUEVer(Ar.UEVer());
+		BytesReader.SetCustomVersions(CustomVersions);
+		BytesReader << Meshes;
 	}
 	else
 	{
-		TArray<uint8> Bytes;
-		FMemoryWriter Buffer(Bytes, true);
-		Buffer << Meshes;
+		if (ensure(SerialVersion >= 1))
+		{
+			FPackageFileVersion SavingUEVer = Ar.UEVer();
+			Ar << SavingUEVer;
 
-		// MeshDescriptions uses custom versionning,
-		const_cast<FCustomVersionContainer&>(Buffer.GetCustomVersions()).Serialize(Ar);
+			int32 SavingLicenseeUEVer = Ar.LicenseeUEVer();
+			Ar << SavingLicenseeUEVer;
+		}
 
+		TArray<uint8> Bytes;
+		FMemoryWriter BytesWriter(Bytes, true);
+		BytesWriter << Meshes;
 		if (BufferType == CompressedMeshDescription)
 		{
 			CompressInline(Bytes, ECompressionMethod::ECM_Default);
 		}
+
+		// MeshDescriptions uses custom versioning,
+		FCustomVersionContainer MeshesCustomVersions = BytesWriter.GetCustomVersions();
+		MeshesCustomVersions.Serialize(Ar);
+
 		Ar << Bytes;
 		FMD5 Md5;
 		Md5.Update(Bytes.GetData(), Bytes.Num());
