DirectLink: prevent startup when the platform doesn't support multithreading.

diff --git a/Engine/Source/Runtime/Datasmith/DirectLink/Private/DirectLinkEndpoint.cpp b/Engine/Source/Runtime/Datasmith/DirectLink/Private/DirectLinkEndpoint.cpp
--- a/Engine/Source/Runtime/Datasmith/DirectLink/Private/DirectLinkEndpoint.cpp
+++ b/Engine/Source/Runtime/Datasmith/DirectLink/Private/DirectLinkEndpoint.cpp
@@ -158,6 +158,12 @@ FEndpoint::FEndpoint(const FString& InName)
 	, InternalPtr(MakeUnique<FInternalThreadState>(*this, SharedState))
 	, Internal(*InternalPtr)
 {
+	if (!FPlatformProcess::SupportsMultithreading())
+	{
+		UE_LOG(LogDirectLinkNet, Error, TEXT("Endpoint '%s': Unable to start endpoint: support for threads is required for DirectLink."), *SharedState.NiceName);
+		return;
+	}
+
 	ECommunicationStatus ComStatus = ValidateCommunicationStatus();
 	if (ComStatus != ECommunicationStatus::NoIssue)
 	{
