Datasmith: Cloth Serialization support
- native translator can load expoorted cloths

diff --git a/Engine/Plugins/Enterprise/DatasmithImporter/Source/DatasmithNativeTranslator/Private/DatasmithNativeTranslator.cpp b/Engine/Plugins/Enterprise/DatasmithImporter/Source/DatasmithNativeTranslator/Private/DatasmithNativeTranslator.cpp
--- a/Engine/Plugins/Enterprise/DatasmithImporter/Source/DatasmithNativeTranslator/Private/DatasmithNativeTranslator.cpp
+++ b/Engine/Plugins/Enterprise/DatasmithImporter/Source/DatasmithNativeTranslator/Private/DatasmithNativeTranslator.cpp
@@ -142,9 +142,6 @@ bool FDatasmithNativeTranslator::LoadStaticMesh(const TSharedRef<IDatasmithMeshE
 				OutMeshPayload.LodMeshes.Add(MoveTemp(SourceModel));
 			}
 		}
-
-		// Do not wait until garbage collection to free memory of the models
-		DatasmithMesh.SourceModels.Reset();
 	}
 
 	return OutMeshPayload.LodMeshes.Num() != 0;
@@ -162,3 +159,19 @@ bool FDatasmithNativeTranslator::LoadLevelSequence(const TSharedRef<IDatasmithLe
 	return false;
 }
 
+bool FDatasmithNativeTranslator::LoadCloth(const TSharedRef<IDatasmithClothElement> ClothElement, FDatasmithClothElementPayload& OutClothPayload)
+{
+	TRACE_CPUPROFILER_EVENT_SCOPE(FDatasmithNativeTranslator::LoadCloth);
+
+	FString FilePath = ClothElement->GetFile();
+	FDatasmithPackedCloths Pack = GetDatasmithClothFromFile(FilePath);
+
+	if (Pack.ClothInfos.IsValidIndex(0))
+	{
+		OutClothPayload.Cloth = MoveTemp(Pack.ClothInfos[0].Cloth);
+		return true;
+	}
+
+	return false;
+}
+
diff --git a/Engine/Plugins/Enterprise/DatasmithImporter/Source/DatasmithNativeTranslator/Public/DatasmithNativeTranslator.h b/Engine/Plugins/Enterprise/DatasmithImporter/Source/DatasmithNativeTranslator/Public/DatasmithNativeTranslator.h
--- a/Engine/Plugins/Enterprise/DatasmithImporter/Source/DatasmithNativeTranslator/Public/DatasmithNativeTranslator.h
+++ b/Engine/Plugins/Enterprise/DatasmithImporter/Source/DatasmithNativeTranslator/Public/DatasmithNativeTranslator.h
@@ -19,6 +19,8 @@ public:
 
 	virtual bool LoadLevelSequence(const TSharedRef<IDatasmithLevelSequenceElement> LevelSequenceElement, FDatasmithLevelSequencePayload& OutLevelSequencePayload) override;
 
+	virtual bool LoadCloth(const TSharedRef<IDatasmithClothElement> ClothElement, FDatasmithClothElementPayload& OutClothPayload) override;
+
 protected:
 	static FString ResolveFilePath(const FString& FilePath, const TArray<FString>& ResourcePaths);
 	static void ResolveSceneFilePaths(TSharedRef<IDatasmithScene> Scene, const TArray<FString>& ResourcePaths);
diff --git a/Engine/Source/Developer/Datasmith/DatasmithExporter/Private/DatasmithMeshExporter.cpp b/Engine/Source/Developer/Datasmith/DatasmithExporter/Private/DatasmithMeshExporter.cpp
--- a/Engine/Source/Developer/Datasmith/DatasmithExporter/Private/DatasmithMeshExporter.cpp
+++ b/Engine/Source/Developer/Datasmith/DatasmithExporter/Private/DatasmithMeshExporter.cpp
@@ -3,6 +3,7 @@
 #include "DatasmithMeshExporter.h"
 
 #include "DatasmithCore.h"
+#include "DatasmithCloth.h"
 #include "DatasmithExporterManager.h"
 #include "DatasmithMesh.h"
 #include "DatasmithMeshUObject.h"
@@ -156,7 +157,7 @@ TSharedPtr<IDatasmithMeshElement> FDatasmithMeshExporter::ExportToUObject(const
 	FString FullPath(GetMeshFilePath(Filepath, Filename));
 
 	TSharedPtr<IDatasmithMeshElement> ExportedMeshElement;
-	FDatasmithMeshExporterOptions ExportOptions(MoveTemp(FullPath), Mesh, LightmapUV, CollisionMesh);
+	FDatasmithMeshExporterOptions ExportOptions(FullPath, Mesh, LightmapUV, CollisionMesh);
 	Impl->DoExport(ExportedMeshElement, ExportOptions);
 
 	return ExportedMeshElement;
@@ -172,6 +173,34 @@ bool FDatasmithMeshExporter::ExportToUObject(TSharedPtr<IDatasmithMeshElement>&
 }
 
 
+bool FDatasmithMeshExporter::ExportCloth(FDatasmithCloth& Cloth, TSharedPtr<IDatasmithClothElement>& ClothElement, const TCHAR* FilePath, const TCHAR* AssetsOutputPath) const
+{
+	FString Path = FPaths::SetExtension(FilePath, TEXT(".udscloth"));
+	FPaths::NormalizeFilename(Path);
+	FString RelativeFilePath = Path;
+	if (AssetsOutputPath)
+	{
+		FString Tmp(AssetsOutputPath);
+		Tmp += TEXT('/');
+		FPaths::MakePathRelativeTo(RelativeFilePath, *Tmp);
+	}
+	ClothElement->SetFile(*RelativeFilePath);
+
+	TUniquePtr<FArchive> Archive(IFileManager::Get().CreateFileWriter(*Path));
+	if (!Archive.IsValid())
+	{
+		Impl->LastError = FString::Printf(TEXT("Failed writing to file %s"), *Path);
+		return false;
+	}
+
+	FDatasmithPackedCloths Pack;
+	FDatasmithClothInfo& Info = Pack.ClothInfos.AddDefaulted_GetRef();
+	Info.Cloth = Cloth;
+	FMD5Hash OutHash = Pack.Serialize(*Archive);
+
+	return true;
+}
+
 FString FDatasmithMeshExporter::GetLastError() const
 {
 	return Impl->LastError;
diff --git a/Engine/Source/Developer/Datasmith/DatasmithExporter/Public/DatasmithMeshExporter.h b/Engine/Source/Developer/Datasmith/DatasmithExporter/Public/DatasmithMeshExporter.h
--- a/Engine/Source/Developer/Datasmith/DatasmithExporter/Public/DatasmithMeshExporter.h
+++ b/Engine/Source/Developer/Datasmith/DatasmithExporter/Public/DatasmithMeshExporter.h
@@ -42,6 +42,8 @@ public:
 	 */
 	bool ExportToUObject(TSharedPtr<IDatasmithMeshElement>& MeshElement, const TCHAR* Filepath, FDatasmithMesh& Mesh, FDatasmithMesh* CollisionMesh, EDSExportLightmapUV LightmapUV);
 
+	bool ExportCloth(class FDatasmithCloth& Cloth, TSharedPtr<class IDatasmithClothElement>& ClothElement, const TCHAR* FilePath, const TCHAR* AssetsOutputPath) const;
+
 	/**
 	 * @return The error that happened during the last export, if any
 	 */
@@ -49,5 +51,4 @@ public:
 
 private:
 	TUniquePtr<class FDatasmithMeshExporterImpl> Impl;
-	bool bAllowOldSerialization = false;
 };
diff --git a/Engine/Source/Runtime/Datasmith/DatasmithCore/Private/DatasmithCloth.cpp b/Engine/Source/Runtime/Datasmith/DatasmithCore/Private/DatasmithCloth.cpp
--- a/Engine/Source/Runtime/Datasmith/DatasmithCore/Private/DatasmithCloth.cpp
+++ b/Engine/Source/Runtime/Datasmith/DatasmithCore/Private/DatasmithCloth.cpp
@@ -2,26 +2,72 @@
 
 #include "DatasmithCloth.h"
 
-// #include "DatasmithMesh.h"
-
-// class FDatasmithClothImpl
-// {
-// public:
-// 	FDatasmithMesh SimMesh;
-// 	FDatasmithMesh RenderMesh;
-// };
-//
-//
-// FDatasmithCloth::FDatasmithCloth()
-// 	: Impl(MakePimpl<FDatasmithClothImpl>())
-// {}
-//
-// FDatasmithMesh& FDatasmithCloth::GetSimMesh()
-// {
-// 	return Impl->SimMesh;
-// }
-//
-// FDatasmithMesh& FDatasmithCloth::GetRenderMesh()
-// {
-// 	return Impl->RenderMesh;
-// }
+#include "Serialization/CustomVersion.h"
+
+
+
+enum EDatasmithClothSerializationVersion
+{
+    EDCSV_Base = 0,
+
+	// -----<new versions can be added before this line>-------------------------------------------------
+    _EDCSV_LastPlusOne,
+    _EDCSV_Last = _EDCSV_LastPlusOne -1
+};
+
+struct FDatasmithClothSerializationVersion
+{
+    const static FGuid GUID;
+
+private:
+    FDatasmithClothSerializationVersion() = default;
+};
+
+
+const FGuid FDatasmithClothSerializationVersion::GUID(0x28B01036, 0x66B4498F, 0x99425ACA, 0xDB78A9B5);
+
+// Register the custom version with core
+FCustomVersionRegistration GRegisterDatasmithClothCustomVersion(FDatasmithClothSerializationVersion::GUID, _EDCSV_Last, TEXT("DatasmithCloth"));
+
+
+FArchive& operator<<(FArchive& Ar, FDatasmithClothPresetProperty& Property)
+{
+	Ar.UsingCustomVersion(FDatasmithClothSerializationVersion::GUID);
+
+	Ar << Property.Name;
+	Ar << Property.Value;
+
+	return Ar;
+}
+
+FArchive& operator<<(FArchive& Ar, FDatasmithClothPresetPropertySet& PropertySet)
+{
+	Ar.UsingCustomVersion(FDatasmithClothSerializationVersion::GUID);
+
+	Ar << PropertySet.SetName;
+	Ar << PropertySet.Properties;
+
+	return Ar;
+}
+
+FArchive& operator<<(FArchive& Ar, FDatasmithCloth& Cloth)
+{
+	Ar.UsingCustomVersion(FDatasmithClothSerializationVersion::GUID);
+
+	Ar << Cloth.Patterns;
+	Ar << Cloth.PropertySets;
+
+	return Ar;
+}
+
+FArchive& operator<<(FArchive& Ar, FDatasmithClothPattern& Pattern)
+{
+	Ar.UsingCustomVersion(FDatasmithClothSerializationVersion::GUID);
+
+	Ar << Pattern.SimPosition;
+	Ar << Pattern.SimRestPosition;
+	Ar << Pattern.SimTriangleIndices;
+
+	return Ar;
+}
+
diff --git a/Engine/Source/Runtime/Datasmith/DatasmithCore/Private/DatasmithMeshSerialization.cpp b/Engine/Source/Runtime/Datasmith/DatasmithCore/Private/DatasmithMeshSerialization.cpp
--- a/Engine/Source/Runtime/Datasmith/DatasmithCore/Private/DatasmithMeshSerialization.cpp
+++ b/Engine/Source/Runtime/Datasmith/DatasmithCore/Private/DatasmithMeshSerialization.cpp
@@ -10,11 +10,24 @@
 #include "StaticMeshAttributes.h"
 #include "StaticMeshOperations.h"
 
-void operator<<(FArchive& Ar, FDatasmithMeshModels& Models)
+FArchive& operator<<(FArchive& Ar, FDatasmithMeshModels& Models)
 {
 	Ar << Models.MeshName;
 	Ar << Models.bIsCollisionMesh;
 	Ar << Models.SourceModels;
+
+	return Ar;
+}
+
+
+FArchive& operator<<(FArchive& Ar, FDatasmithClothInfo& Info)
+{
+	uint8 Version = 0;
+	Ar << Version;
+
+	Ar << Info.Cloth;
+
+	return Ar;
 }
 
 
@@ -151,9 +164,10 @@ FMD5Hash FDatasmithPackedMeshes::Serialize(FArchive& Ar, bool bCompressed)
 	// | CustomVersions  | v0 | internal structures use them, we have to pass them along
 	// | Mesh buffer     | v0 | may be compressed, the actual payload
 
-	FString Guard = Ar.IsLoading() ? TEXT("") : TEXT("FDatasmithPackedMeshes");
+	const TCHAR* StaticGuard = TEXT("FDatasmithPackedMeshes");
+	FString Guard = Ar.IsLoading() ? TEXT("") : StaticGuard;
 	Ar << Guard;
-	if (!ensure(Guard == TEXT("FDatasmithPackedMeshes")))
+	if (!ensure(Guard == StaticGuard))
 	{
 		Ar.SetError();
 		return {};
@@ -350,24 +364,110 @@ FDatasmithPackedMeshes GetDatasmithMeshFromFile(const FString& MeshPath)
 	return Result;
 }
 
-FDatasmithPackedMeshes GetDatasmithClothFromFile(const FString& Path)
+
+FMD5Hash FDatasmithPackedCloths::Serialize(FArchive& Ar, bool bSaveCompressed)
 {
-	// #ue_ds_cloth_note: asset serialization
-	FDatasmithPackedMeshes Result;
+	using namespace DatasmithMeshSerializationImpl;
+	// structure          from
+	// | Guard           | v0 | constant string, prevents loading of invalid buffers
+	// | SerialVersion   | v0 | this structure versioning
+	// | BufferType      | v0 | describe the payload
+	// | UEVer           | v0 | global engine version
+	// | LicenseeUEVer   | v0 |
+	// | Payload buffer  | v0 | may be compressed
 
-	TUniquePtr<FArchive> Archive( IFileManager::Get().CreateFileReader(*Path) );
+	const TCHAR* StaticGuard = TEXT("FDatasmithPackedCloths");
+	FString Guard = Ar.IsLoading() ? TEXT("") : StaticGuard;
+	Ar << Guard;
+	if (!ensure(Guard == StaticGuard))
+	{
+		Ar.SetError();
+		return {};
+	}
+
+	uint32 SerialVersion = 0;
+	Ar << SerialVersion;
+
+	enum EBufferType{ Cloths, CompressedCloth };
+	uint8 BufferType = bSaveCompressed ? CompressedCloth : Cloths;
+	Ar << BufferType;
+
+	FMD5Hash OutHash;
+	if (Ar.IsLoading())
+	{
+		FPackageFileVersion LoadingUEVer;
+		Ar << LoadingUEVer;
+		Ar.SetUEVer(LoadingUEVer);
+
+		int32 LoadingLicenseeUEVer = 0;
+		Ar << LoadingLicenseeUEVer;
+		Ar.SetLicenseeUEVer(LoadingLicenseeUEVer);
+
+		FCustomVersionContainer CustomVersions;
+		CustomVersions.Serialize(Ar);
+
+		// content
+		TArray<uint8> Bytes;
+		Ar << Bytes;
+		if (BufferType == CompressedCloth)
+		{
+			DecompressInline(Bytes);
+		}
+
+		FMemoryReader BytesReader(Bytes, true);
+		BytesReader.SetLicenseeUEVer(Ar.LicenseeUEVer());
+		BytesReader.SetUEVer(Ar.UEVer());
+		BytesReader.SetCustomVersions(CustomVersions);
+		BytesReader << ClothInfos;
+	}
+	else
+	{
+		FPackageFileVersion SavingUEVer = Ar.UEVer();
+		Ar << SavingUEVer;
+
+		int32 SavingLicenseeUEVer = Ar.LicenseeUEVer();
+		Ar << SavingLicenseeUEVer;
+
+		{
+			TArray<uint8> Bytes;
+			FMemoryWriter BytesWriter(Bytes, true);
+			BytesWriter << ClothInfos;
+			if (BufferType == CompressedCloth)
+			{
+				CompressInline(Bytes, ECompressionMethod::ECM_Default);
+			}
+
+			FCustomVersionContainer MeshesCustomVersions = BytesWriter.GetCustomVersions();
+			MeshesCustomVersions.Serialize(Ar);
+
+			Ar << Bytes;
+
+			FMD5 Md5;
+			Md5.Update(Bytes.GetData(), Bytes.Num());
+			OutHash.Set(Md5);
+		}
+	}
+
+	return OutHash;
+}
+
+
+FDatasmithPackedCloths GetDatasmithClothFromFile(const FString& Path)
+{
+	TUniquePtr<FArchive> Archive(IFileManager::Get().CreateFileReader(*Path));
 	if ( !Archive.IsValid() )
 	{
 		UE_LOG(LogDatasmith, Warning, TEXT("Cannot read file %s"), *Path);
-		return Result;
+		return {};
 	}
 
+	FDatasmithPackedCloths Result;
 	Result.Serialize(*Archive);
 
 	if (Archive->IsError())
 	{
-		Result = FDatasmithPackedMeshes{};
 		UE_LOG(LogDatasmith, Warning, TEXT("Failed to read cloth from %s"), *Path);
+		return {};
 	}
 
 	return Result;
diff --git a/Engine/Source/Runtime/Datasmith/DatasmithCore/Public/DatasmithCloth.h b/Engine/Source/Runtime/Datasmith/DatasmithCore/Public/DatasmithCloth.h
--- a/Engine/Source/Runtime/Datasmith/DatasmithCore/Public/DatasmithCloth.h
+++ b/Engine/Source/Runtime/Datasmith/DatasmithCore/Public/DatasmithCloth.h
@@ -11,6 +11,8 @@
 
 class FDatasmithMesh;
 
+
+
 class DATASMITHCORE_API FDatasmithClothPattern
 {
 public:
@@ -20,6 +22,7 @@ public:
 
 public:
 	bool IsValid() { return SimRestPosition.Num() == SimPosition.Num() && SimTriangleIndices.Num() % 3 == 0 && SimTriangleIndices.Num(); }
+	friend FArchive& operator<<(FArchive& Ar, FDatasmithClothPattern& Pattern);
 };
 
 
@@ -28,27 +31,30 @@ class FDatasmithClothPresetProperty
 public:
 	FName Name;
 	double Value;
+
+public:
+	friend FArchive& operator<<(FArchive& Ar, FDatasmithClothPresetProperty& Property);
 };
 
+
 class DATASMITHCORE_API FDatasmithClothPresetPropertySet
 {
 public:
 	FString SetName;
 	TArray<FDatasmithClothPresetProperty> Properties;
+
+public:
+	friend FArchive& operator<<(FArchive& Ar, FDatasmithClothPresetPropertySet& PropertySet);
 };
 
-/*
- * Structure of a Cloth
- * - One Thin Mesh:
- * - base for physics simulation
- *
- * - One Thick Mesh
- * - aka. render mesh
- */
+
 class DATASMITHCORE_API FDatasmithCloth
 {
 public:
 	TArray<FDatasmithClothPattern> Patterns;
 	TArray<FDatasmithClothPresetPropertySet> PropertySets;
+
+public:
+	friend FArchive& operator<<(FArchive& Ar, FDatasmithCloth& Cloth);
 };
 
diff --git a/Engine/Source/Runtime/Datasmith/DatasmithCore/Public/DatasmithMeshSerialization.h b/Engine/Source/Runtime/Datasmith/DatasmithCore/Public/DatasmithMeshSerialization.h
--- a/Engine/Source/Runtime/Datasmith/DatasmithCore/Public/DatasmithMeshSerialization.h
+++ b/Engine/Source/Runtime/Datasmith/DatasmithCore/Public/DatasmithMeshSerialization.h
@@ -2,22 +2,24 @@
 
 #pragma once
 
+#include "DatasmithCore.h"
+
 #include "Containers/Array.h"
 #include "Containers/UnrealString.h"
-#include "DatasmithCore.h"
 #include "MeshDescription.h"
 #include "Misc/SecureHash.h"
+#include "DatasmithCloth.h"
 
 class FArchive;
 
 
-struct FDatasmithMeshModels
+struct DATASMITHCORE_API FDatasmithMeshModels
 {
 	FString MeshName;
 	bool bIsCollisionMesh = false;
 	TArray<FMeshDescription> SourceModels;
 
-	DATASMITHCORE_API friend void operator << (FArchive& Ar, FDatasmithMeshModels& Models);
+	friend FArchive& operator<<(FArchive& Ar, FDatasmithMeshModels& Models);
 };
 
 struct DATASMITHCORE_API FDatasmithPackedMeshes
@@ -28,5 +30,22 @@ struct DATASMITHCORE_API FDatasmithPackedMeshes
 };
 
 DATASMITHCORE_API FDatasmithPackedMeshes GetDatasmithMeshFromFile(const FString& MeshPath);
-DATASMITHCORE_API FDatasmithPackedMeshes GetDatasmithClothFromFile(const FString& Path);
+
+
+
+struct DATASMITHCORE_API FDatasmithClothInfo
+{
+	FDatasmithCloth Cloth;
+	friend FArchive& operator<<(FArchive& Ar, FDatasmithClothInfo& Info);
+};
+
+struct DATASMITHCORE_API FDatasmithPackedCloths
+{
+	TArray<FDatasmithClothInfo> ClothInfos;
+
+	FMD5Hash Serialize(FArchive& Ar, bool bSaveCompressed=true);
+};
+
+
+DATASMITHCORE_API FDatasmithPackedCloths GetDatasmithClothFromFile(const FString& Path);
 
