CADKernel Sources Update: some cleanup before next improvements

diff --git a/Engine/Plugins/Enterprise/DatasmithCADImporter/Source/CADInterfaces/Private/TechSoftFileParser.cpp b/Engine/Plugins/Enterprise/DatasmithCADImporter/Source/CADInterfaces/Private/TechSoftFileParser.cpp
--- a/Engine/Plugins/Enterprise/DatasmithCADImporter/Source/CADInterfaces/Private/TechSoftFileParser.cpp
+++ b/Engine/Plugins/Enterprise/DatasmithCADImporter/Source/CADInterfaces/Private/TechSoftFileParser.cpp
@@ -9,8 +9,10 @@
 #include "TUniqueTechSoftObj.h"
 
 #include "HAL/FileManager.h"
+#ifndef CADKERNEL_DEV
 #include "Serialization/JsonSerializer.h"
 #include "Serialization/JsonWriter.h"
+#endif
 #include "Templates/UnrealTemplate.h"
 
 namespace CADLibrary
@@ -524,6 +526,7 @@ void FTechSoftFileParser::GenerateBodyMesh(A3DRiRepresentationItem* Representati
 
 	if (Type == kA3DTypeRiBrepModel)
 	{
+#ifndef CADKERNEL_DEV
 		FString FilePath = CADFileData.GetBodyCachePath(Body.MeshActorUId);
 		if (!FilePath.IsEmpty())
 		{
@@ -546,9 +549,9 @@ void FTechSoftFileParser::GenerateBodyMesh(A3DRiRepresentationItem* Representati
 			TSharedRef<TJsonWriter<TCHAR, TPrettyJsonPrintPolicy<TCHAR>>> JsonWriter = TJsonWriterFactory<TCHAR, TPrettyJsonPrintPolicy<TCHAR>>::Create(&JsonString);
 
 			FJsonSerializer::Serialize(JsonObject.ToSharedRef(), JsonWriter);
-
 			TechSoftUtils::SaveBodiesToPrcFile(&Representation, 1, FilePath, JsonString);
 		}
+#endif
 	}
 }
 
diff --git a/Engine/Plugins/Enterprise/DatasmithCADImporter/Source/CADInterfaces/Private/TechSoftUtils.cpp b/Engine/Plugins/Enterprise/DatasmithCADImporter/Source/CADInterfaces/Private/TechSoftUtils.cpp
--- a/Engine/Plugins/Enterprise/DatasmithCADImporter/Source/CADInterfaces/Private/TechSoftUtils.cpp
+++ b/Engine/Plugins/Enterprise/DatasmithCADImporter/Source/CADInterfaces/Private/TechSoftUtils.cpp
@@ -9,10 +9,10 @@
 
 #include "GenericPlatform/GenericPlatformMisc.h"
 #include "Math/Color.h"
-#include "Serialization/JsonSerializer.h"
-#include "Serialization/JsonWriter.h"
 
 #ifndef CADKERNEL_DEV
+#include "Serialization/JsonSerializer.h"
+#include "Serialization/JsonWriter.h"
 #include "TechSoftInterfaceUtils.inl"
 #endif
 
@@ -298,8 +298,10 @@ void RestoreMaterials(const TSharedPtr<FJsonObject>& DefaultValues, FBodyMesh& B
 	FMaterialUId DefaultColorName = 0;
 	FMaterialUId DefaultMaterialName = 0;
 
+#ifndef CADKERNEL_DEV
 	DefaultValues->TryGetNumberField(JSON_ENTRY_COLOR_NAME, DefaultColorName);
 	DefaultValues->TryGetNumberField(JSON_ENTRY_MATERIAL_NAME, DefaultMaterialName);
+#endif
 
 	BodyMesh.MaterialSet.Empty();
 	BodyMesh.ColorSet.Empty();
diff --git a/Engine/Source/Runtime/Datasmith/CADKernel/Private/CADKernel/Geo/Sampling/SurfacicSampling.h b/Engine/Source/Runtime/Datasmith/CADKernel/Private/CADKernel/Geo/Sampling/SurfacicSampling.h
--- a/Engine/Source/Runtime/Datasmith/CADKernel/Private/CADKernel/Geo/Sampling/SurfacicSampling.h
+++ b/Engine/Source/Runtime/Datasmith/CADKernel/Private/CADKernel/Geo/Sampling/SurfacicSampling.h
@@ -13,7 +13,7 @@ struct FSurfacicSampling
 	bool bWithNormals = false;
 	TArray<FPoint2D> Points2D;
 	TArray<FPoint> Points3D;
-	TArray<FVector> Normals;
+	TArray<FVector3f> Normals;
 
 	int32 Count()
 	{
@@ -42,7 +42,7 @@ struct FSurfacicSampling
 
 	void NormalizeNormals()
 	{
-		for (FVector& Normal : Normals)
+		for (FVector3f& Normal : Normals)
 		{
 			Normal.Normalize();
 		}
diff --git a/Engine/Source/Runtime/Datasmith/CADKernel/Private/CADKernel/Geo/Surfaces/ConeSurface.cpp b/Engine/Source/Runtime/Datasmith/CADKernel/Private/CADKernel/Geo/Surfaces/ConeSurface.cpp
--- a/Engine/Source/Runtime/Datasmith/CADKernel/Private/CADKernel/Geo/Surfaces/ConeSurface.cpp
+++ b/Engine/Source/Runtime/Datasmith/CADKernel/Private/CADKernel/Geo/Surfaces/ConeSurface.cpp
@@ -94,7 +94,8 @@ void FConeSurface::EvaluatePointGrid(const FCoordinateGrid& Coordinates, FSurfac
 				FPoint GradientV(DeltaVR * CosU[Undex], DeltaVR * SinU[Undex], 1.0);
 				GradientU = Matrix.MultiplyVector(GradientU);
 				GradientV = Matrix.MultiplyVector(GradientV);
-				OutPoints.Normals.Emplace(GradientU ^ GradientV);
+				const FPoint Normal = GradientU ^ GradientV;
+				OutPoints.Normals.Emplace(Normal);
 			}
 		}
 
diff --git a/Engine/Source/Runtime/Datasmith/CADKernel/Private/CADKernel/Geo/Surfaces/Surface.cpp b/Engine/Source/Runtime/Datasmith/CADKernel/Private/CADKernel/Geo/Surfaces/Surface.cpp
--- a/Engine/Source/Runtime/Datasmith/CADKernel/Private/CADKernel/Geo/Surfaces/Surface.cpp
+++ b/Engine/Source/Runtime/Datasmith/CADKernel/Private/CADKernel/Geo/Surfaces/Surface.cpp
@@ -104,7 +104,7 @@ void FSurface::EvaluatePoints(FSurfacicPolyline& Polyline) const
 		{
 			Polyline.Normals.Emplace(Point.GradientU ^ Point.GradientV);
 		}
-		for (FVector Normal : Polyline.Normals)
+		for (FVector3f Normal : Polyline.Normals)
 		{
 			Normal.Normalize();
 		}
@@ -143,13 +143,13 @@ void FSurface::EvaluatePoints(const TArray<FCurvePoint2D>& InPoints2D, FSurfacic
 		Polyline.Normals.Emplace(Point.GradientU ^ Point.GradientV);
 	}
 
-	for (FVector& Normal : Polyline.Normals)
+	for (FVector3f& Normal : Polyline.Normals)
 	{
 		Normal.Normalize();
 	}
 }
 
-void FSurface::EvaluateNormals(const TArray<FPoint2D>& InPoints2D, TArray<FVector>& Normals) const
+void FSurface::EvaluateNormals(const TArray<FPoint2D>& InPoints2D, TArray<FVector3f>& Normals) const
 {
 	int32 DerivativeOrder = 1;
 	TArray<FSurfacicPoint> Points3D;
@@ -162,7 +162,7 @@ void FSurface::EvaluateNormals(const TArray<FPoint2D>& InPoints2D, TArray<FVecto
 		Normals.Emplace(Point.GradientU ^ Point.GradientV);
 	}
 
-	for (FVector& Normal : Normals)
+	for (FVector3f& Normal : Normals)
 	{
 		Normal.Normalize();
 	}
@@ -189,7 +189,8 @@ void FSurface::EvaluatePointGrid(const FCoordinateGrid& Coordinates, FSurfacicSa
 	{
 		for (FSurfacicPoint Point : Point3D)
 		{
-			OutPoints.Normals.Emplace(Point.GradientU ^ Point.GradientV);
+			const FPoint Normal = Point.GradientU ^ Point.GradientV;
+			OutPoints.Normals.Emplace(Normal);
 		}
 		OutPoints.NormalizeNormals();
 	}
diff --git a/Engine/Source/Runtime/Datasmith/CADKernel/Private/CADKernel/Geo/Surfaces/TorusSurface.cpp b/Engine/Source/Runtime/Datasmith/CADKernel/Private/CADKernel/Geo/Surfaces/TorusSurface.cpp
--- a/Engine/Source/Runtime/Datasmith/CADKernel/Private/CADKernel/Geo/Surfaces/TorusSurface.cpp
+++ b/Engine/Source/Runtime/Datasmith/CADKernel/Private/CADKernel/Geo/Surfaces/TorusSurface.cpp
@@ -122,8 +122,8 @@ void FTorusSurface::EvaluatePointGrid(const FCoordinateGrid& Coordinates, FSurfa
 
 	if (bComputeNormals)
 	{
-		FVector Center = Matrix.Column(3);
-		for (FVector& Normal : OutPoints.Normals)
+		FVector3f Center = Matrix.Column(3);
+		for (FVector3f& Normal : OutPoints.Normals)
 		{
 			 Normal = Matrix.MultiplyVector(Normal);
 		}
diff --git a/Engine/Source/Runtime/Datasmith/CADKernel/Private/CADKernel/Mesh/Criteria/CriteriaGrid.cpp b/Engine/Source/Runtime/Datasmith/CADKernel/Private/CADKernel/Mesh/Criteria/CriteriaGrid.cpp
--- a/Engine/Source/Runtime/Datasmith/CADKernel/Private/CADKernel/Mesh/Criteria/CriteriaGrid.cpp
+++ b/Engine/Source/Runtime/Datasmith/CADKernel/Private/CADKernel/Mesh/Criteria/CriteriaGrid.cpp
@@ -14,7 +14,7 @@
 namespace CADKernel
 {
 
-const FPoint& FCriteriaGrid::GetPoint(int32 UIndex, int32 VIndex, bool bIsInternalU, bool bIsInternalV, FVector* OutNormal) const
+const FPoint& FCriteriaGrid::GetPoint(int32 UIndex, int32 VIndex, bool bIsInternalU, bool bIsInternalV, FVector3f* OutNormal) const
 {
 	int32 Index = GetIndex(UIndex, VIndex, bIsInternalU, bIsInternalV);
 	ensureCADKernel(Grid.Points3D.IsValidIndex(Index));
diff --git a/Engine/Source/Runtime/Datasmith/CADKernel/Private/CADKernel/Mesh/Criteria/CriteriaGrid.h b/Engine/Source/Runtime/Datasmith/CADKernel/Private/CADKernel/Mesh/Criteria/CriteriaGrid.h
--- a/Engine/Source/Runtime/Datasmith/CADKernel/Private/CADKernel/Mesh/Criteria/CriteriaGrid.h
+++ b/Engine/Source/Runtime/Datasmith/CADKernel/Private/CADKernel/Mesh/Criteria/CriteriaGrid.h
@@ -35,7 +35,7 @@ namespace CADKernel
 			return TrueVIndex * TrueUcoorindateCount + TrueUIndex;
 		}
 
-		const FPoint& GetPoint(int32 UIndex, int32 VIndex, bool bIsInterU, bool bIsInterV, FVector* normal = nullptr) const;
+		const FPoint& GetPoint(int32 UIndex, int32 VIndex, bool bIsInterU, bool bIsInterV, FVector3f* normal = nullptr) const;
 		void Init();
 
 	public:
@@ -68,27 +68,27 @@ namespace CADKernel
 			return Grid.Points3D;
 		}
 
-		const TArray<FVector>& GetGridNormals() const
+		const TArray<FVector3f>& GetGridNormals() const
 		{
 			return Grid.Normals;
 		}
 
-		const FPoint& GetPoint(int32 iU, int32 iV, FVector* normal = nullptr) const
+		const FPoint& GetPoint(int32 iU, int32 iV, FVector3f* normal = nullptr) const
 		{
 			return GetPoint(iU, iV, false, false, normal);
 		}
 
-		const FPoint& GetIntermediateU(int32 iU, int32 iV, FVector* normal = nullptr) const
+		const FPoint& GetIntermediateU(int32 iU, int32 iV, FVector3f* normal = nullptr) const
 		{
 			return GetPoint(iU, iV, true, false, normal);
 		}
 
-		const FPoint& GetIntermediateV(int32 iU, int32 iV, FVector* normal = nullptr) const
+		const FPoint& GetIntermediateV(int32 iU, int32 iV, FVector3f* normal = nullptr) const
 		{
 			return GetPoint(iU, iV, false, true, normal);
 		}
 
-		const FPoint& GetIntermediateUV(int32 iU, int32 iV, FVector* normal = nullptr) const
+		const FPoint& GetIntermediateUV(int32 iU, int32 iV, FVector3f* normal = nullptr) const
 		{
 			return GetPoint(iU, iV, true, true, normal);
 		}
diff --git a/Engine/Source/Runtime/Datasmith/CADKernel/Private/CADKernel/Mesh/Meshers/IsoTriangulator.cpp b/Engine/Source/Runtime/Datasmith/CADKernel/Private/CADKernel/Mesh/Meshers/IsoTriangulator.cpp
--- a/Engine/Source/Runtime/Datasmith/CADKernel/Private/CADKernel/Mesh/Meshers/IsoTriangulator.cpp
+++ b/Engine/Source/Runtime/Datasmith/CADKernel/Private/CADKernel/Mesh/Meshers/IsoTriangulator.cpp
@@ -2396,7 +2396,10 @@ void FIsoTriangulator::TriangulateOverCycle(const EGridSpace Space)
 			Cycle.Empty();
 			CycleOrientation.Empty();
 			bool bLeftSide = true;
-			FindCycle(Segment, bLeftSide, Cycle, CycleOrientation);
+			if (!FindCycle(Segment, bLeftSide, Cycle, CycleOrientation))
+			{
+				continue;
+			}
 #ifdef FIND_CYCLE
 			if (Grid.GetFace()->GetId() == FaceToDebug)
 			{
@@ -2414,7 +2417,10 @@ void FIsoTriangulator::TriangulateOverCycle(const EGridSpace Space)
 			Cycle.Empty();
 			CycleOrientation.Empty();
 			bool bLeftSide = false;
-			FindCycle(Segment, bLeftSide, Cycle, CycleOrientation);
+			if(!FindCycle(Segment, bLeftSide, Cycle, CycleOrientation))
+			{
+				continue;
+			}
 #ifdef FIND_CYCLE
 			if (Grid.GetFace()->GetId() == FaceToDebug)
 			{
@@ -2442,7 +2448,7 @@ static int32 CycleId = -1;
 static int32 CycleIndex = 0;
 #endif
 
-void FIsoTriangulator::FindCycle(FIsoSegment* StartSegment, bool LeftSide, TArray<FIsoSegment*>& Cycle, TArray<bool>& CycleOrientation)
+bool FIsoTriangulator::FindCycle(FIsoSegment* StartSegment, bool LeftSide, TArray<FIsoSegment*>& Cycle, TArray<bool>& CycleOrientation)
 {
 
 #ifdef DEBUG_FIND_CYCLE
@@ -2529,7 +2535,10 @@ void FIsoTriangulator::FindCycle(FIsoSegment* StartSegment, bool LeftSide, TArra
 				Wait();
 			}
 #endif
-			ensureCADKernel(!Segment->HasCycleOnLeft());
+			if (Segment->HasCycleOnLeft())
+			{
+				return false;
+			}
 			Segment->SetHaveCycleOnLeft();
 			Node = &Segment->GetSecondNode();
 			CycleOrientation.Add(true);
@@ -2553,7 +2562,10 @@ void FIsoTriangulator::FindCycle(FIsoSegment* StartSegment, bool LeftSide, TArra
 				Wait();
 			}
 #endif
-			ensureCADKernel(!Segment->HasCycleOnRight());
+			if (Segment->HasCycleOnRight())
+			{
+				return false;
+			}
 			Segment->SetHaveCycleOnRight();
 			Node = &Segment->GetFirstNode();
 			CycleOrientation.Add(false);
@@ -2578,6 +2590,7 @@ void FIsoTriangulator::FindCycle(FIsoSegment* StartSegment, bool LeftSide, TArra
 		Close3DDebugSession();
 	}
 #endif
+	return true;
 }
 
 //#define DEBUG_FIND_NEXTSEGMENT
diff --git a/Engine/Source/Runtime/Datasmith/CADKernel/Private/CADKernel/Mesh/Meshers/IsoTriangulator.h b/Engine/Source/Runtime/Datasmith/CADKernel/Private/CADKernel/Mesh/Meshers/IsoTriangulator.h
--- a/Engine/Source/Runtime/Datasmith/CADKernel/Private/CADKernel/Mesh/Meshers/IsoTriangulator.h
+++ b/Engine/Source/Runtime/Datasmith/CADKernel/Private/CADKernel/Mesh/Meshers/IsoTriangulator.h
@@ -291,8 +291,9 @@ public:
 
 	/**
 	 * Find in the network a minimal cycle stating from a segment
+	 * @return false if the new cycle crosses a segement already used
 	 */
-	void FindCycle(FIsoSegment* StartSegment, bool bLeftSide, TArray<FIsoSegment*>& Cycle, TArray<bool>& CycleOrientation);
+	bool FindCycle(FIsoSegment* StartSegment, bool bLeftSide, TArray<FIsoSegment*>& Cycle, TArray<bool>& CycleOrientation);
 
 	/**
 	 * Generate the "Delaunay" tessellation of the cycle.
diff --git a/Engine/Source/Runtime/Datasmith/CADKernel/Private/CADKernel/Mesh/Meshers/IsoTriangulator/IsoNode.h b/Engine/Source/Runtime/Datasmith/CADKernel/Private/CADKernel/Mesh/Meshers/IsoTriangulator/IsoNode.h
--- a/Engine/Source/Runtime/Datasmith/CADKernel/Private/CADKernel/Mesh/Meshers/IsoTriangulator/IsoNode.h
+++ b/Engine/Source/Runtime/Datasmith/CADKernel/Private/CADKernel/Mesh/Meshers/IsoTriangulator/IsoNode.h
@@ -211,7 +211,7 @@ public:
 	virtual void Set2DPoint(EGridSpace Space, FGrid& Grid, const FPoint2D& NewCoordinate) = 0;
 
 	virtual const FPoint& Get3DPoint(const FGrid& Grid) const = 0;
-	virtual const FVector& GetNormal(const FGrid& Grid) const = 0;
+	virtual const FVector3f& GetNormal(const FGrid& Grid) const = 0;
 
 	/**
 	 * Only for display purpose as it return a copy of the point
@@ -324,7 +324,7 @@ public:
 		return Grid.GetLoops3D()[LoopIndex][Index];
 	}
 
-	virtual const FVector& GetNormal(const FGrid& Grid) const override
+	virtual const FVector3f& GetNormal(const FGrid& Grid) const override
 	{
 		return Grid.GetLoopNormals()[LoopIndex][Index];
 	}
@@ -497,7 +497,7 @@ public:
 		return Grid.GetInner3DPoint(Index);
 	}
 
-	virtual const FVector& GetNormal(const FGrid& Grid) const override
+	virtual const FVector3f& GetNormal(const FGrid& Grid) const override
 	{
 		return Grid.GetPointNormal(Index);
 	}
diff --git a/Engine/Source/Runtime/Datasmith/CADKernel/Private/CADKernel/Mesh/Meshers/ParametricMesher.cpp b/Engine/Source/Runtime/Datasmith/CADKernel/Private/CADKernel/Mesh/Meshers/ParametricMesher.cpp
--- a/Engine/Source/Runtime/Datasmith/CADKernel/Private/CADKernel/Mesh/Meshers/ParametricMesher.cpp
+++ b/Engine/Source/Runtime/Datasmith/CADKernel/Private/CADKernel/Mesh/Meshers/ParametricMesher.cpp
@@ -246,9 +246,6 @@ void FParametricMesher::Mesh(FTopologicalFace& Face)
 	TSharedRef<FFaceMesh> SurfaceMesh = StaticCastSharedRef<FFaceMesh>(Face.GetOrCreateMesh(MeshModel));
 
 	FIsoTriangulator IsoTrianguler(Grid, SurfaceMesh);
-#ifdef CADKERNEL_DEV
-	MesherReport
-#endif
 
 	if (IsoTrianguler.Triangulate())
 	{
diff --git a/Engine/Source/Runtime/Datasmith/CADKernel/Private/CADKernel/Mesh/Structure/Grid.cpp b/Engine/Source/Runtime/Datasmith/CADKernel/Private/CADKernel/Mesh/Structure/Grid.cpp
--- a/Engine/Source/Runtime/Datasmith/CADKernel/Private/CADKernel/Mesh/Structure/Grid.cpp
+++ b/Engine/Source/Runtime/Datasmith/CADKernel/Private/CADKernel/Mesh/Structure/Grid.cpp
@@ -947,7 +947,7 @@ void FGrid::GetMeshOfLoop(const FTopologicalLoop& Loop)
 	TArray<FPoint>& Loop3D = FaceLoops3D.Emplace_GetRef();
 	Loop3D.Reserve(LoopNodeCount);
 
-	TArray<FVector>& LoopNormals = NormalsOfFaceLoops.Emplace_GetRef();
+	TArray<FVector3f>& LoopNormals = NormalsOfFaceLoops.Emplace_GetRef();
 	LoopNormals.Reserve(LoopNodeCount);
 
 	TArray<int32>& LoopIds = NodeIdsOfFaceLoops.Emplace_GetRef();
diff --git a/Engine/Source/Runtime/Datasmith/CADKernel/Private/CADKernel/Mesh/Structure/Grid.h b/Engine/Source/Runtime/Datasmith/CADKernel/Private/CADKernel/Mesh/Structure/Grid.h
--- a/Engine/Source/Runtime/Datasmith/CADKernel/Private/CADKernel/Mesh/Structure/Grid.h
+++ b/Engine/Source/Runtime/Datasmith/CADKernel/Private/CADKernel/Mesh/Structure/Grid.h
@@ -97,7 +97,7 @@ protected:
 	/**
 	 * Surface Normal at each boundary nodes
 	 */
-	TArray<TArray<FVector>> NormalsOfFaceLoops;
+	TArray<TArray<FVector3f>> NormalsOfFaceLoops;
 
 	TArray<TArray<int32>> NodeIdsOfFaceLoops;
 
@@ -155,7 +155,7 @@ protected:
 	/**
 	 * Surface Normal at each inner nodes
 	 */
-	TArray<FVector> Normals;
+	TArray<FVector3f> Normals;
 
 public:
 	FGridChronos Chronos;
@@ -399,7 +399,7 @@ public:
 	/**
 	 * @return the normal of the surface at the point at the Index of the grid
 	 */
-	const FVector& GetPointNormal(int32 IndexU, int32 IndexV) const
+	const FVector3f& GetPointNormal(int32 IndexU, int32 IndexV) const
 	{
 		return Normals[GobalIndex(IndexU, IndexV)];
 	}
@@ -407,7 +407,7 @@ public:
 	/**
 	 * @return the normal of the surface at the point at the Index of the grid
 	 */
-	const FVector& GetPointNormal(int32 Index) const
+	const FVector3f& GetPointNormal(int32 Index) const
 	{
 		return Normals[Index];
 	}
@@ -459,7 +459,7 @@ public:
 	/**
 	 * @return the array of normal of the points of the grid
 	 */
-	TArray<FVector>& GetNormals()
+	TArray<FVector3f>& GetNormals()
 	{
 		return Normals;
 	}
@@ -526,7 +526,7 @@ public:
 	/**
 	 * @return the array of array of normal of the points of the loops
 	 */
-	const TArray<TArray<FVector>>& GetLoopNormals() const
+	const TArray<TArray<FVector3f>>& GetLoopNormals() const
 	{
 		return NormalsOfFaceLoops;
 	}
diff --git a/Engine/Source/Runtime/Datasmith/CADKernel/Private/CADKernel/Mesh/Structure/GridDisplay.cpp b/Engine/Source/Runtime/Datasmith/CADKernel/Private/CADKernel/Mesh/Structure/GridDisplay.cpp
--- a/Engine/Source/Runtime/Datasmith/CADKernel/Private/CADKernel/Mesh/Structure/GridDisplay.cpp
+++ b/Engine/Source/Runtime/Datasmith/CADKernel/Private/CADKernel/Mesh/Structure/GridDisplay.cpp
@@ -344,7 +344,7 @@ void FGrid::DisplayGridNormal() const
 		for (int32 Index = 0; Index < CuttingSize; ++Index)
 		{
 			F3DDebugSegment GraphicSegment(Index);
-			FVector Normal = Normals[Index];
+			FVector3f Normal = Normals[Index];
 			Normal.Normalize();
 			Normal *= NormalLength;
 			DrawSegment(Points3D[Index], Points3D[Index] + Normal, EVisuProperty::GreenCurve);
@@ -356,12 +356,12 @@ void FGrid::DisplayGridNormal() const
 		for (int32 LoopIndex = 0; LoopIndex < FaceLoops3D.Num(); ++LoopIndex)
 		{
 			const TArray<FPoint>& LoopPoints = FaceLoops3D[LoopIndex];
-			const TArray<FVector>& LoopNormals = NormalsOfFaceLoops[LoopIndex];
+			const TArray<FVector3f>& LoopNormals = NormalsOfFaceLoops[LoopIndex];
 
 			for (int32 Index = 0; Index < LoopPoints.Num(); ++Index)
 			{
 				F3DDebugSegment GraphicSegment(Index);
-				FVector Normal = LoopNormals[Index];
+				FVector3f Normal = LoopNormals[Index];
 				Normal.Normalize();
 				Normal *= NormalLength;
 				DrawSegment(LoopPoints[Index], LoopPoints[Index] + Normal, EVisuProperty::YellowCurve);
diff --git a/Engine/Source/Runtime/Datasmith/CADKernel/Private/CADKernel/Mesh/Structure/LoopCleaner.h b/Engine/Source/Runtime/Datasmith/CADKernel/Private/CADKernel/Mesh/Structure/LoopCleaner.h
--- a/Engine/Source/Runtime/Datasmith/CADKernel/Private/CADKernel/Mesh/Structure/LoopCleaner.h
+++ b/Engine/Source/Runtime/Datasmith/CADKernel/Private/CADKernel/Mesh/Structure/LoopCleaner.h
@@ -181,7 +181,7 @@ public:
 private:
 
 #ifdef CADKERNEL_DEV
-	void SetMesherReport(FMesherReport& MesherReport)
+	void SetMesherReport(FMesherReport& InMesherReport)
 	{
 		MesherReport = &InMesherReport;
 	}
diff --git a/Engine/Source/Runtime/Datasmith/CADKernel/Private/CADKernel/Topo/Topomaker.cpp b/Engine/Source/Runtime/Datasmith/CADKernel/Private/CADKernel/Topo/Topomaker.cpp
--- a/Engine/Source/Runtime/Datasmith/CADKernel/Private/CADKernel/Topo/Topomaker.cpp
+++ b/Engine/Source/Runtime/Datasmith/CADKernel/Private/CADKernel/Topo/Topomaker.cpp
@@ -18,8 +18,6 @@
 #include "CADKernel/UI/Message.h"
 #include "CADKernel/Utils/Util.h"
 
-#include "Math/UnrealMathUtility.h"
-
 namespace CADKernel
 {
 
@@ -155,6 +153,8 @@ void FTopomaker::Sew()
 
 	MergeUnconnectedAdjacentEdges();
 
+	StitchParallelEdges(BorderVertices);
+
 #ifdef CADKERNEL_DEV
 	Report.SewDuration = FChrono::Elapse(StartJoinTime);
 	FChrono::PrintClockElapse(EVerboseLevel::Log, TEXT(""), TEXT("Sew"), Report.SewDuration);
@@ -1130,7 +1130,7 @@ void FTopomaker::OrientShells()
 	}
 
 #ifdef CADKERNEL_DEV
- 	Report.OrientationDuration = FChrono::Elapse(StartTime);
+	Report.OrientationDuration = FChrono::Elapse(StartTime);
 	FChrono::PrintClockElapse(EVerboseLevel::Log, TEXT(""), TEXT("Orient"), Report.OrientationDuration);
 #endif
 }
diff --git a/Engine/Source/Runtime/Datasmith/CADKernel/Private/CADKernel/UI/Display.cpp b/Engine/Source/Runtime/Datasmith/CADKernel/Private/CADKernel/UI/Display.cpp
--- a/Engine/Source/Runtime/Datasmith/CADKernel/Private/CADKernel/UI/Display.cpp
+++ b/Engine/Source/Runtime/Datasmith/CADKernel/Private/CADKernel/UI/Display.cpp
@@ -1084,7 +1084,7 @@ void DisplayMesh(const FFaceMesh& Mesh)
 	if (FSystem::Get().GetVisu()->GetParameters()->bDisplayNormals)
 	{
 		double NormalLength = FSystem::Get().GetVisu()->GetParameters()->NormalLength;
-		const TArray<FVector>& Normals = Mesh.Normals;
+		const TArray<FVector3f>& Normals = Mesh.Normals;
 		for (int32 Index = 0; Index < VertexIndices.Num(); ++Index)
 		{
 			const FPoint** PointPtr = NodeIdToCoordinates.Find(VertexIndices[Index]);
@@ -1094,7 +1094,7 @@ void DisplayMesh(const FFaceMesh& Mesh)
 			}
 
 			F3DDebugSegment GraphicSegment(Index);
-			FVector Normal = Normals[Index];
+			FVector3f Normal = Normals[Index];
 			Normal.Normalize();
 			Normal *= NormalLength;
 			DrawSegment(**PointPtr, **PointPtr + Normal, EVisuProperty::EdgeMesh);
diff --git a/Engine/Source/Runtime/Datasmith/CADKernel/Public/CADKernel/Core/Database.h b/Engine/Source/Runtime/Datasmith/CADKernel/Public/CADKernel/Core/Database.h
--- a/Engine/Source/Runtime/Datasmith/CADKernel/Public/CADKernel/Core/Database.h
+++ b/Engine/Source/Runtime/Datasmith/CADKernel/Public/CADKernel/Core/Database.h
@@ -3,14 +3,7 @@
 
 #include "CADKernel/Core/CADKernelArchive.h"
 #include "CADKernel/Core/Entity.h"
-#include "CADKernel/Core/Types.h"
-#include "Containers/Array.h"
 #include "Containers/Queue.h"
-#include "Containers/Set.h"
-#include "HAL/Platform.h"
-#include "Misc/AssertionMacros.h"
-#include "Templates/SharedPointer.h"
-#include "Templates/TypeHash.h"
 
 namespace CADKernel
 {
diff --git a/Engine/Source/Runtime/Datasmith/CADKernel/Public/CADKernel/Core/Entity.h b/Engine/Source/Runtime/Datasmith/CADKernel/Public/CADKernel/Core/Entity.h
--- a/Engine/Source/Runtime/Datasmith/CADKernel/Public/CADKernel/Core/Entity.h
+++ b/Engine/Source/Runtime/Datasmith/CADKernel/Public/CADKernel/Core/Entity.h
@@ -6,11 +6,7 @@
 #include "CADKernel/Core/OrientedEntity.h"
 #include "CADKernel/Core/Types.h"
 #include "CADKernel/Geo/GeoEnum.h"
-#include "Containers/Array.h"
-#include "HAL/Platform.h"
 #include "Serialization/Archive.h"
-#include "Templates/SharedPointer.h"
-#include "Templates/UnrealTemplate.h"
 
 #ifdef CADKERNEL_DEV
 #include "Toolkit/Core/InfoEntity.h"
diff --git a/Engine/Source/Runtime/Datasmith/CADKernel/Public/CADKernel/Core/Parameters.h b/Engine/Source/Runtime/Datasmith/CADKernel/Public/CADKernel/Core/Parameters.h
--- a/Engine/Source/Runtime/Datasmith/CADKernel/Public/CADKernel/Core/Parameters.h
+++ b/Engine/Source/Runtime/Datasmith/CADKernel/Public/CADKernel/Core/Parameters.h
@@ -3,9 +3,9 @@
 
 #include "CADKernel/Core/Parameter.h"
 #include "CADKernel/Core/Types.h"
-#include "Containers/Map.h"
-#include "Containers/UnrealString.h"
-#include "HAL/Platform.h"
+//#include "Containers/Map.h"
+//#include "Containers/UnrealString.h"
+//#include "HAL/Platform.h"
 #include "Serialization/Archive.h"
 
 class FArchive;
diff --git a/Engine/Source/Runtime/Datasmith/CADKernel/Public/CADKernel/Core/Session.h b/Engine/Source/Runtime/Datasmith/CADKernel/Public/CADKernel/Core/Session.h
--- a/Engine/Source/Runtime/Datasmith/CADKernel/Public/CADKernel/Core/Session.h
+++ b/Engine/Source/Runtime/Datasmith/CADKernel/Public/CADKernel/Core/Session.h
@@ -5,10 +5,10 @@
 #include "CADKernel/Core/Database.h"
 #include "CADKernel/Core/Entity.h"
 #include "CADKernel/Core/Types.h"
-#include "Containers/Array.h"
-#include "HAL/Platform.h"
-#include "HAL/PlatformCrt.h"
-#include "Templates/SharedPointer.h"
+//#include "Containers/Array.h"
+//#include "HAL/Platform.h"
+//#include "HAL/PlatformCrt.h"
+//#include "Templates/SharedPointer.h"
 
 class FArchive;
 
diff --git a/Engine/Source/Runtime/Datasmith/CADKernel/Public/CADKernel/Core/System.h b/Engine/Source/Runtime/Datasmith/CADKernel/Public/CADKernel/Core/System.h
--- a/Engine/Source/Runtime/Datasmith/CADKernel/Public/CADKernel/Core/System.h
+++ b/Engine/Source/Runtime/Datasmith/CADKernel/Public/CADKernel/Core/System.h
@@ -5,11 +5,11 @@
 #include "CADKernel/UI/Console.h"
 #include "CADKernel/UI/Progress.h"
 #include "CADKernel/UI/Visu.h"
-#include "Containers/UnrealString.h"
-#include "HAL/Platform.h"
-#include "Misc/AssertionMacros.h"
-#include "Templates/SharedPointer.h"
-#include "Templates/UniquePtr.h"
+//#include "Containers/UnrealString.h"
+//#include "HAL/Platform.h"
+//#include "Misc/AssertionMacros.h"
+//#include "Templates/SharedPointer.h"
+//#include "Templates/UniquePtr.h"
 
 class FArchive;
 
diff --git a/Engine/Source/Runtime/Datasmith/CADKernel/Public/CADKernel/Geo/Curves/Curve.h b/Engine/Source/Runtime/Datasmith/CADKernel/Public/CADKernel/Geo/Curves/Curve.h
--- a/Engine/Source/Runtime/Datasmith/CADKernel/Public/CADKernel/Geo/Curves/Curve.h
+++ b/Engine/Source/Runtime/Datasmith/CADKernel/Public/CADKernel/Geo/Curves/Curve.h
@@ -9,14 +9,14 @@
 #include "CADKernel/Geo/GeoEnum.h"
 #include "CADKernel/Geo/GeoPoint.h"
 #include "CADKernel/Math/Boundary.h"
-#include "CADKernel/Math/MatrixH.h"
-#include "CADKernel/Math/Point.h"
+//#include "CADKernel/Math/MatrixH.h"
+//#include "CADKernel/Math/Point.h"
 #include "CADKernel/UI/Message.h"
 #include "CADKernel/Utils/Cache.h"
-#include "Containers/Array.h"
-#include "HAL/Platform.h"
-#include "Misc/AssertionMacros.h"
-#include "Templates/SharedPointer.h"
+//#include "Containers/Array.h"
+//#include "HAL/Platform.h"
+//#include "Misc/AssertionMacros.h"
+//#include "Templates/SharedPointer.h"
 
 namespace CADKernel
 {
diff --git a/Engine/Source/Runtime/Datasmith/CADKernel/Public/CADKernel/Geo/Curves/RestrictionCurve.h b/Engine/Source/Runtime/Datasmith/CADKernel/Public/CADKernel/Geo/Curves/RestrictionCurve.h
--- a/Engine/Source/Runtime/Datasmith/CADKernel/Public/CADKernel/Geo/Curves/RestrictionCurve.h
+++ b/Engine/Source/Runtime/Datasmith/CADKernel/Public/CADKernel/Geo/Curves/RestrictionCurve.h
@@ -11,11 +11,6 @@
 #include "CADKernel/Math/Boundary.h"
 #include "CADKernel/Math/MatrixH.h"
 #include "CADKernel/Math/Point.h"
-#include "CADKernel/Utils/IndexOfCoordinateFinder.h"
-#include "Containers/Array.h"
-#include "HAL/Platform.h"
-#include "Misc/AssertionMacros.h"
-#include "Templates/SharedPointer.h"
 
 namespace CADKernel
 {
diff --git a/Engine/Source/Runtime/Datasmith/CADKernel/Public/CADKernel/Geo/Sampling/SurfacicPolyline.h b/Engine/Source/Runtime/Datasmith/CADKernel/Public/CADKernel/Geo/Sampling/SurfacicPolyline.h
--- a/Engine/Source/Runtime/Datasmith/CADKernel/Public/CADKernel/Geo/Sampling/SurfacicPolyline.h
+++ b/Engine/Source/Runtime/Datasmith/CADKernel/Public/CADKernel/Geo/Sampling/SurfacicPolyline.h
@@ -1,9 +1,9 @@
 // Copyright Epic Games, Inc. All Rights Reserved.
 #pragma once
 
-#include "Algo/Reverse.h"
-#include "CADKernel/Core/CADKernelArchive.h"
 #include "CADKernel/Core/Types.h"
+
+#include "CADKernel/Core/CADKernelArchive.h"
 #include "CADKernel/Geo/Curves/Curve.h"
 #include "CADKernel/Geo/GeoEnum.h"
 #include "CADKernel/Geo/GeoPoint.h"
@@ -14,13 +14,6 @@
 #include "CADKernel/Math/Point.h"
 #include "CADKernel/UI/Display.h"
 #include "CADKernel/Utils/IndexOfCoordinateFinder.h"
-#include "Containers/Array.h"
-#include "Math/UnrealMathSSE.h"
-#include "Math/Vector.h"
-#include "Serialization/Archive.h"
-#include "Templates/Function.h"
-#include "Templates/SharedPointer.h"
-#include "Templates/UnrealTemplate.h"
 
 namespace CADKernel
 {
@@ -38,7 +31,7 @@ public:
 	TArray<double> Coordinates;
 	TArray<FPoint2D> Points2D;
 	TArray<FPoint> Points3D;
-	TArray<FVector> Normals;
+	TArray<FVector3f> Normals;
 	TArray<FPoint> Tangents;
 
 	FSurfacicBoundary BoundingBox;
@@ -293,7 +286,7 @@ public:
 		return Points3D;
 	}
 
-	const TArray<FVector>& GetNormals() const
+	const TArray<FVector3f>& GetNormals() const
 	{
 		return Normals;
 	}
@@ -439,4 +432,3 @@ public:
 };
 
 } // ns CADKernel
-
diff --git a/Engine/Source/Runtime/Datasmith/CADKernel/Public/CADKernel/Geo/Surfaces/Surface.h b/Engine/Source/Runtime/Datasmith/CADKernel/Public/CADKernel/Geo/Surfaces/Surface.h
--- a/Engine/Source/Runtime/Datasmith/CADKernel/Public/CADKernel/Geo/Surfaces/Surface.h
+++ b/Engine/Source/Runtime/Datasmith/CADKernel/Public/CADKernel/Geo/Surfaces/Surface.h
@@ -203,7 +203,7 @@ namespace CADKernel
 		virtual void EvaluatePointGrid(const FCoordinateGrid& Coordinates, FSurfacicSampling& OutPoints, bool bComputeNormals = false) const;
 		void EvaluateGrid(FGrid& Grid) const;
 
-		virtual void EvaluateNormals(const TArray<FPoint2D>& Points2D, TArray<FVector>& Normals) const;
+		virtual void EvaluateNormals(const TArray<FPoint2D>& Points2D, TArray<FVector3f>& Normals) const;
 
 		virtual FVector EvaluateNormal(const FPoint2D& InPoint2D) const
 		{
diff --git a/Engine/Source/Runtime/Datasmith/CADKernel/Public/CADKernel/Math/Geometry.h b/Engine/Source/Runtime/Datasmith/CADKernel/Public/CADKernel/Math/Geometry.h
--- a/Engine/Source/Runtime/Datasmith/CADKernel/Public/CADKernel/Math/Geometry.h
+++ b/Engine/Source/Runtime/Datasmith/CADKernel/Public/CADKernel/Math/Geometry.h
@@ -7,10 +7,10 @@
 #include "CADKernel/Math/MathConst.h"
 #include "CADKernel/Math/MatrixH.h"
 #include "CADKernel/Math/Point.h"
-#include "Containers/Array.h"
-#include "Math/UnrealMathSSE.h"
-#include "Math/Vector.h"
-#include "Misc/AssertionMacros.h"
+//#include "Containers/Array.h"
+//#include "Math/UnrealMathSSE.h"
+//#include "Math/Vector.h"
+//#include "Misc/AssertionMacros.h"
 
 namespace CADKernel
 {
diff --git a/Engine/Source/Runtime/Datasmith/CADKernel/Public/CADKernel/Math/MatrixH.h b/Engine/Source/Runtime/Datasmith/CADKernel/Public/CADKernel/Math/MatrixH.h
--- a/Engine/Source/Runtime/Datasmith/CADKernel/Public/CADKernel/Math/MatrixH.h
+++ b/Engine/Source/Runtime/Datasmith/CADKernel/Public/CADKernel/Math/MatrixH.h
@@ -4,11 +4,11 @@
 #include "CADKernel/Core/Types.h"
 #include "CADKernel/Math/MathConst.h"
 #include "CADKernel/Math/Point.h"
-#include "HAL/PlatformCrt.h"
-#include "Math/UnrealMathSSE.h"
-#include "Math/Vector.h"
-#include "Serialization/Archive.h"
-#include "Templates/UnrealTemplate.h"
+//#include "HAL/PlatformCrt.h"
+//#include "Math/UnrealMathSSE.h"
+//#include "Math/Vector.h"
+//#include "Serialization/Archive.h"
+//#include "Templates/UnrealTemplate.h"
 
 namespace CADKernel
 {
@@ -108,6 +108,15 @@ namespace CADKernel
 			);
 		}
 
+		FVector3f MultiplyVector(const FVector3f& InVector) const
+		{
+			return FVector3f(
+				(float) (InVector.X * Get(0, 0) + InVector.Y * Get(0, 1) + InVector.Z * Get(0, 2)),
+				(float) (InVector.X * Get(1, 0) + InVector.Y * Get(1, 1) + InVector.Z * Get(1, 2)),
+				(float) (InVector.X * Get(2, 0) + InVector.Y * Get(2, 1) + InVector.Z * Get(2, 2))
+			);
+		}
+
 		static FMatrixH MakeRotationMatrix(double InAngle, const FPoint InAxe);
 
 		static FMatrixH MakeTranslationMatrix(const FPoint& InPoint);
diff --git a/Engine/Source/Runtime/Datasmith/CADKernel/Public/CADKernel/Math/Point.h b/Engine/Source/Runtime/Datasmith/CADKernel/Public/CADKernel/Math/Point.h
--- a/Engine/Source/Runtime/Datasmith/CADKernel/Public/CADKernel/Math/Point.h
+++ b/Engine/Source/Runtime/Datasmith/CADKernel/Public/CADKernel/Math/Point.h
@@ -3,11 +3,11 @@
 
 #include "CADKernel/Core/Types.h"
 #include "CADKernel/Math/MathConst.h"
-#include "HAL/PlatformMath.h"
-#include "Math/UnrealMathSSE.h"
-#include "Math/Vector.h"
-#include "Serialization/Archive.h"
-#include "Templates/TypeHash.h"
+//#include "HAL/PlatformMath.h"
+//#include "Math/UnrealMathSSE.h"
+//#include "Math/Vector.h"
+//#include "Serialization/Archive.h"
+//#include "Templates/TypeHash.h"
 
 // this will eventually be reconciled with what the Geometry team is using.
 namespace CADKernel
@@ -94,6 +94,11 @@ PRAGMA_ENABLE_DEPRECATION_WARNINGS
 		return FPoint(X + Point.X, Y + Point.Y, Z + Point.Z);
 	}
 
+	FPoint operator+(const FVector3f& Point) const
+	{
+		return FPoint(X + Point.X, Y + Point.Y, Z + Point.Z);
+	}
+
 	FPoint operator-(const FPoint& Point) const
 	{
 		return FPoint(X - Point.X, Y - Point.Y, Z - Point.Z);
@@ -148,6 +153,14 @@ PRAGMA_ENABLE_DEPRECATION_WARNINGS
 		return *this;
 	}
 
+	FPoint& operator+=(const FVector3f& Point)
+	{
+		X += Point.X;
+		Y += Point.Y;
+		Z += Point.Z;
+		return *this;
+	}
+
 	FPoint& operator-=(const FPoint& Point)
 	{
 		X -= Point.X;
@@ -205,11 +218,16 @@ PRAGMA_ENABLE_DEPRECATION_WARNINGS
 		return (Compare < 0);
 	}
 
-	operator FVector()
+	operator FVector() const
 	{
 		return FVector(X, Y, Z);
 	}
 
+	operator FVector3f() const 
+	{
+		return FVector3f((float) X, (float)Y, (float)Z);
+	}
+
 	void SetMin(const FPoint& Point)
 	{
 		X = FMath::Min(X, Point.X);
diff --git a/Engine/Source/Runtime/Datasmith/CADKernel/Public/CADKernel/Mesh/Criteria/Criterion.h b/Engine/Source/Runtime/Datasmith/CADKernel/Public/CADKernel/Mesh/Criteria/Criterion.h
--- a/Engine/Source/Runtime/Datasmith/CADKernel/Public/CADKernel/Mesh/Criteria/Criterion.h
+++ b/Engine/Source/Runtime/Datasmith/CADKernel/Public/CADKernel/Mesh/Criteria/Criterion.h
@@ -7,8 +7,8 @@
 #include "CADKernel/Geo/GeoPoint.h"
 #include "CADKernel/Math/Point.h"
 #include "CADKernel/Mesh/Criteria/CriterionType.h"
-#include "Containers/Array.h"
-#include "Containers/UnrealString.h"
+//#include "Containers/Array.h"
+//#include "Containers/UnrealString.h"
 #include "HAL/PlatformMath.h"
 #include "Math/UnrealMathUtility.h"
 #include "Misc/AssertionMacros.h"
diff --git a/Engine/Source/Runtime/Datasmith/CADKernel/Public/CADKernel/Mesh/Structure/FaceMesh.h b/Engine/Source/Runtime/Datasmith/CADKernel/Public/CADKernel/Mesh/Structure/FaceMesh.h
--- a/Engine/Source/Runtime/Datasmith/CADKernel/Public/CADKernel/Mesh/Structure/FaceMesh.h
+++ b/Engine/Source/Runtime/Datasmith/CADKernel/Public/CADKernel/Mesh/Structure/FaceMesh.h
@@ -2,8 +2,8 @@
 #pragma once
 
 #include "CADKernel/Mesh/Structure/Mesh.h"
-#include "Containers/Array.h"
-#include "Containers/Map.h"
+//#include "Containers/Array.h"
+//#include "Containers/Map.h"
 #include "HAL/PlatformCrt.h"
 #include "Math/UnrealMathSSE.h"
 #include "Math/Vector.h"
diff --git a/Engine/Source/Runtime/Datasmith/CADKernel/Public/CADKernel/Mesh/Structure/Mesh.h b/Engine/Source/Runtime/Datasmith/CADKernel/Public/CADKernel/Mesh/Structure/Mesh.h
--- a/Engine/Source/Runtime/Datasmith/CADKernel/Public/CADKernel/Mesh/Structure/Mesh.h
+++ b/Engine/Source/Runtime/Datasmith/CADKernel/Public/CADKernel/Mesh/Structure/Mesh.h
@@ -4,9 +4,9 @@
 #include "CADKernel/Core/Entity.h"
 #include "CADKernel/Core/EntityGeom.h"
 #include "CADKernel/Core/Types.h"
-#include "CADKernel/Math/Point.h"
-#include "Containers/Array.h"
-#include "HAL/Platform.h"
+//#include "CADKernel/Math/Point.h"
+//#include "Containers/Array.h"
+//#include "HAL/Platform.h"
 
 namespace CADKernel
 {
diff --git a/Engine/Source/Runtime/Datasmith/CADKernel/Public/CADKernel/Mesh/Structure/ModelMesh.h b/Engine/Source/Runtime/Datasmith/CADKernel/Public/CADKernel/Mesh/Structure/ModelMesh.h
--- a/Engine/Source/Runtime/Datasmith/CADKernel/Public/CADKernel/Mesh/Structure/ModelMesh.h
+++ b/Engine/Source/Runtime/Datasmith/CADKernel/Public/CADKernel/Mesh/Structure/ModelMesh.h
@@ -7,9 +7,9 @@
 #include "CADKernel/Core/Types.h"
 #include "CADKernel/Mesh/Structure/Mesh.h"
 #include "Containers/Array.h"
-#include "HAL/PlatformCrt.h"
-#include "Math/UnrealMathSSE.h"
-#include "Templates/SharedPointer.h"
+//#include "HAL/PlatformCrt.h"
+//#include "Math/UnrealMathSSE.h"
+//#include "Templates/SharedPointer.h"
 
 namespace CADKernel
 {
diff --git a/Engine/Source/Runtime/Datasmith/CADKernel/Public/CADKernel/Topo/Body.h b/Engine/Source/Runtime/Datasmith/CADKernel/Public/CADKernel/Topo/Body.h
--- a/Engine/Source/Runtime/Datasmith/CADKernel/Public/CADKernel/Topo/Body.h
+++ b/Engine/Source/Runtime/Datasmith/CADKernel/Public/CADKernel/Topo/Body.h
@@ -8,9 +8,9 @@
 #include "CADKernel/Topo/TopologicalEntity.h"
 #include "CADKernel/Topo/TopologicalFace.h"
 #include "CADKernel/Topo/TopologicalShapeEntity.h"
-#include "Containers/Array.h"
-#include "HAL/Platform.h"
-#include "Templates/SharedPointer.h"
+//#include "Containers/Array.h"
+//#include "HAL/Platform.h"
+//#include "Templates/SharedPointer.h"
 
 namespace CADKernel
 {
diff --git a/Engine/Source/Runtime/Datasmith/CADKernel/Public/CADKernel/Topo/TopologicalFace.h b/Engine/Source/Runtime/Datasmith/CADKernel/Public/CADKernel/Topo/TopologicalFace.h
--- a/Engine/Source/Runtime/Datasmith/CADKernel/Public/CADKernel/Topo/TopologicalFace.h
+++ b/Engine/Source/Runtime/Datasmith/CADKernel/Public/CADKernel/Topo/TopologicalFace.h
@@ -15,15 +15,15 @@
 #include "CADKernel/Topo/TopologicalEntity.h"
 #include "CADKernel/Topo/TopologicalLoop.h"
 #include "CADKernel/Topo/TopologicalShapeEntity.h"
-#include "CADKernel/Utils/Cache.h"
-#include "Containers/Array.h"
-#include "Containers/Map.h"
-#include "Containers/UnrealString.h"
-#include "HAL/PlatformCrt.h"
-#include "Math/UnrealMathSSE.h"
-#include "Math/Vector.h"
-#include "Misc/AssertionMacros.h"
-#include "Templates/SharedPointer.h"
+//#include "CADKernel/Utils/Cache.h"
+//#include "Containers/Array.h"
+//#include "Containers/Map.h"
+//#include "Containers/UnrealString.h"
+//#include "HAL/PlatformCrt.h"
+//#include "Math/UnrealMathSSE.h"
+//#include "Math/Vector.h"
+//#include "Misc/AssertionMacros.h"
+//#include "Templates/SharedPointer.h"
 
 namespace CADKernel
 {
diff --git a/Engine/Source/Runtime/Datasmith/CADKernel/Public/CADKernel/Topo/TopologicalLoop.h b/Engine/Source/Runtime/Datasmith/CADKernel/Public/CADKernel/Topo/TopologicalLoop.h
--- a/Engine/Source/Runtime/Datasmith/CADKernel/Public/CADKernel/Topo/TopologicalLoop.h
+++ b/Engine/Source/Runtime/Datasmith/CADKernel/Public/CADKernel/Topo/TopologicalLoop.h
@@ -9,10 +9,10 @@
 #include "CADKernel/Topo/TopologicalEdge.h"
 #include "CADKernel/Topo/TopologicalEntity.h"
 #include "CADKernel/Topo/TopologicalVertex.h"
-#include "Containers/Array.h"
-#include "HAL/Platform.h"
-#include "HAL/PlatformCrt.h"
-#include "Templates/SharedPointer.h"
+//#include "Containers/Array.h"
+//#include "HAL/Platform.h"
+//#include "HAL/PlatformCrt.h"
+//#include "Templates/SharedPointer.h"
 
 namespace CADKernel
 {
