CADKernel Sources Update: some cleanup before next improvements

diff --git a/Engine/Source/Runtime/Datasmith/CADKernel/Private/CADKernel/Mesh/Meshers/IsoTriangulator.cpp b/Engine/Source/Runtime/Datasmith/CADKernel/Private/CADKernel/Mesh/Meshers/IsoTriangulator.cpp
--- a/Engine/Source/Runtime/Datasmith/CADKernel/Private/CADKernel/Mesh/Meshers/IsoTriangulator.cpp
+++ b/Engine/Source/Runtime/Datasmith/CADKernel/Private/CADKernel/Mesh/Meshers/IsoTriangulator.cpp
@@ -1314,7 +1314,7 @@ void FIsoTriangulator::FindCellContainingBoundaryNodes(TArray<FCell>& Cells)
 				Cells.Emplace(CellIndex, CellNodes, Grid);
 
 				CellIndex = NodeToCellIndices[Index];
-				CellNodes.Empty(LoopNodeCount);
+				CellNodes.Reset(LoopNodeCount);
 
 #ifdef DEBUG_BUILD_CELLS
 				DisplayCell(Cells.Last());
diff --git a/Engine/Source/Runtime/Datasmith/CADKernel/Private/CADKernel/Mesh/Meshers/ParametricMesher.cpp b/Engine/Source/Runtime/Datasmith/CADKernel/Private/CADKernel/Mesh/Meshers/ParametricMesher.cpp
--- a/Engine/Source/Runtime/Datasmith/CADKernel/Private/CADKernel/Mesh/Meshers/ParametricMesher.cpp
+++ b/Engine/Source/Runtime/Datasmith/CADKernel/Private/CADKernel/Mesh/Meshers/ParametricMesher.cpp
@@ -246,6 +246,9 @@ void FParametricMesher::Mesh(FTopologicalFace& Face)
 	TSharedRef<FFaceMesh> SurfaceMesh = StaticCastSharedRef<FFaceMesh>(Face.GetOrCreateMesh(MeshModel));
 
 	FIsoTriangulator IsoTrianguler(Grid, SurfaceMesh);
+#ifdef CADKERNEL_DEV
+	IsoTrianguler.SetMesherReport(MesherReport);
+#endif
 
 	if (IsoTrianguler.Triangulate())
 	{
diff --git a/Engine/Source/Runtime/Datasmith/CADKernel/Public/CADKernel/Core/Database.h b/Engine/Source/Runtime/Datasmith/CADKernel/Public/CADKernel/Core/Database.h
--- a/Engine/Source/Runtime/Datasmith/CADKernel/Public/CADKernel/Core/Database.h
+++ b/Engine/Source/Runtime/Datasmith/CADKernel/Public/CADKernel/Core/Database.h
@@ -3,6 +3,7 @@
 
 #include "CADKernel/Core/CADKernelArchive.h"
 #include "CADKernel/Core/Entity.h"
+#include "CADKernel/Core/Types.h"
 #include "Containers/Queue.h"
 
 namespace CADKernel
diff --git a/Engine/Source/Runtime/Datasmith/CADKernel/Public/CADKernel/Core/Session.h b/Engine/Source/Runtime/Datasmith/CADKernel/Public/CADKernel/Core/Session.h
--- a/Engine/Source/Runtime/Datasmith/CADKernel/Public/CADKernel/Core/Session.h
+++ b/Engine/Source/Runtime/Datasmith/CADKernel/Public/CADKernel/Core/Session.h
@@ -5,10 +5,6 @@
 #include "CADKernel/Core/Database.h"
 #include "CADKernel/Core/Entity.h"
 #include "CADKernel/Core/Types.h"
-//#include "Containers/Array.h"
-//#include "HAL/Platform.h"
-//#include "HAL/PlatformCrt.h"
-//#include "Templates/SharedPointer.h"
 
 class FArchive;
 
diff --git a/Engine/Source/Runtime/Datasmith/CADKernel/Public/CADKernel/Core/System.h b/Engine/Source/Runtime/Datasmith/CADKernel/Public/CADKernel/Core/System.h
--- a/Engine/Source/Runtime/Datasmith/CADKernel/Public/CADKernel/Core/System.h
+++ b/Engine/Source/Runtime/Datasmith/CADKernel/Public/CADKernel/Core/System.h
@@ -5,11 +5,6 @@
 #include "CADKernel/UI/Console.h"
 #include "CADKernel/UI/Progress.h"
 #include "CADKernel/UI/Visu.h"
-//#include "Containers/UnrealString.h"
-//#include "HAL/Platform.h"
-//#include "Misc/AssertionMacros.h"
-//#include "Templates/SharedPointer.h"
-//#include "Templates/UniquePtr.h"
 
 class FArchive;
 
diff --git a/Engine/Source/Runtime/Datasmith/CADKernel/Public/CADKernel/Geo/Curves/Curve.h b/Engine/Source/Runtime/Datasmith/CADKernel/Public/CADKernel/Geo/Curves/Curve.h
--- a/Engine/Source/Runtime/Datasmith/CADKernel/Public/CADKernel/Geo/Curves/Curve.h
+++ b/Engine/Source/Runtime/Datasmith/CADKernel/Public/CADKernel/Geo/Curves/Curve.h
@@ -9,19 +9,11 @@
 #include "CADKernel/Geo/GeoEnum.h"
 #include "CADKernel/Geo/GeoPoint.h"
 #include "CADKernel/Math/Boundary.h"
-//#include "CADKernel/Math/MatrixH.h"
-//#include "CADKernel/Math/Point.h"
 #include "CADKernel/UI/Message.h"
 #include "CADKernel/Utils/Cache.h"
-//#include "Containers/Array.h"
-//#include "HAL/Platform.h"
-//#include "Misc/AssertionMacros.h"
-//#include "Templates/SharedPointer.h"
 
 namespace CADKernel
 {
-	//class FSurface;
-
 	struct FCurvePoint2D;
 	struct FCurvePoint;
 	struct FNurbsCurveData;
diff --git a/Engine/Source/Runtime/Datasmith/CADKernel/Public/CADKernel/Geo/Curves/SurfacicCurve.h b/Engine/Source/Runtime/Datasmith/CADKernel/Public/CADKernel/Geo/Curves/SurfacicCurve.h
--- a/Engine/Source/Runtime/Datasmith/CADKernel/Public/CADKernel/Geo/Curves/SurfacicCurve.h
+++ b/Engine/Source/Runtime/Datasmith/CADKernel/Public/CADKernel/Geo/Curves/SurfacicCurve.h
@@ -7,10 +7,6 @@
 #include "CADKernel/Geo/GeoEnum.h"
 #include "CADKernel/Geo/Surfaces/Surface.h"
 #include "CADKernel/Math/MatrixH.h"
-#include "Containers/Array.h"
-#include "HAL/Platform.h"
-#include "Misc/AssertionMacros.h"
-#include "Templates/SharedPointer.h"
 
 namespace CADKernel
 {
diff --git a/Engine/Source/Runtime/Datasmith/CADKernel/Public/CADKernel/Geo/Sampling/SurfacicPolyline.h b/Engine/Source/Runtime/Datasmith/CADKernel/Public/CADKernel/Geo/Sampling/SurfacicPolyline.h
--- a/Engine/Source/Runtime/Datasmith/CADKernel/Public/CADKernel/Geo/Sampling/SurfacicPolyline.h
+++ b/Engine/Source/Runtime/Datasmith/CADKernel/Public/CADKernel/Geo/Sampling/SurfacicPolyline.h
@@ -15,6 +15,8 @@
 #include "CADKernel/UI/Display.h"
 #include "CADKernel/Utils/IndexOfCoordinateFinder.h"
 
+#include "Serialization/Archive.h"
+
 namespace CADKernel
 {
 
diff --git a/Engine/Source/Runtime/Datasmith/CADKernel/Public/CADKernel/Geo/Surfaces/Surface.h b/Engine/Source/Runtime/Datasmith/CADKernel/Public/CADKernel/Geo/Surfaces/Surface.h
--- a/Engine/Source/Runtime/Datasmith/CADKernel/Public/CADKernel/Geo/Surfaces/Surface.h
+++ b/Engine/Source/Runtime/Datasmith/CADKernel/Public/CADKernel/Geo/Surfaces/Surface.h
@@ -11,13 +11,6 @@
 #include "CADKernel/Math/MatrixH.h"
 #include "CADKernel/Math/Point.h"
 #include "CADKernel/Utils/Cache.h"
-#include "Containers/Array.h"
-#include "HAL/PlatformCrt.h"
-#include "Math/UnrealMathSSE.h"
-#include "Math/Vector.h"
-#include "Misc/AssertionMacros.h"
-#include "Serialization/Archive.h"
-#include "Templates/SharedPointer.h"
 
 namespace CADKernel
 {
diff --git a/Engine/Source/Runtime/Datasmith/CADKernel/Public/CADKernel/Math/BSpline.h b/Engine/Source/Runtime/Datasmith/CADKernel/Public/CADKernel/Math/BSpline.h
--- a/Engine/Source/Runtime/Datasmith/CADKernel/Public/CADKernel/Math/BSpline.h
+++ b/Engine/Source/Runtime/Datasmith/CADKernel/Public/CADKernel/Math/BSpline.h
@@ -4,9 +4,6 @@
 #include "CADKernel/Core/Types.h"
 #include "CADKernel/Math/Boundary.h"
 #include "CADKernel/Math/Point.h"
-#include "Containers/Array.h"
-#include "HAL/Platform.h"
-#include "Templates/SharedPointer.h"
 
 namespace CADKernel
 {
diff --git a/Engine/Source/Runtime/Datasmith/CADKernel/Public/CADKernel/Math/Geometry.h b/Engine/Source/Runtime/Datasmith/CADKernel/Public/CADKernel/Math/Geometry.h
--- a/Engine/Source/Runtime/Datasmith/CADKernel/Public/CADKernel/Math/Geometry.h
+++ b/Engine/Source/Runtime/Datasmith/CADKernel/Public/CADKernel/Math/Geometry.h
@@ -7,10 +7,6 @@
 #include "CADKernel/Math/MathConst.h"
 #include "CADKernel/Math/MatrixH.h"
 #include "CADKernel/Math/Point.h"
-//#include "Containers/Array.h"
-//#include "Math/UnrealMathSSE.h"
-//#include "Math/Vector.h"
-//#include "Misc/AssertionMacros.h"
 
 namespace CADKernel
 {
diff --git a/Engine/Source/Runtime/Datasmith/CADKernel/Public/CADKernel/Math/MatrixH.h b/Engine/Source/Runtime/Datasmith/CADKernel/Public/CADKernel/Math/MatrixH.h
--- a/Engine/Source/Runtime/Datasmith/CADKernel/Public/CADKernel/Math/MatrixH.h
+++ b/Engine/Source/Runtime/Datasmith/CADKernel/Public/CADKernel/Math/MatrixH.h
@@ -4,11 +4,6 @@
 #include "CADKernel/Core/Types.h"
 #include "CADKernel/Math/MathConst.h"
 #include "CADKernel/Math/Point.h"
-//#include "HAL/PlatformCrt.h"
-//#include "Math/UnrealMathSSE.h"
-//#include "Math/Vector.h"
-//#include "Serialization/Archive.h"
-//#include "Templates/UnrealTemplate.h"
 
 namespace CADKernel
 {
diff --git a/Engine/Source/Runtime/Datasmith/CADKernel/Public/CADKernel/Math/Point.h b/Engine/Source/Runtime/Datasmith/CADKernel/Public/CADKernel/Math/Point.h
--- a/Engine/Source/Runtime/Datasmith/CADKernel/Public/CADKernel/Math/Point.h
+++ b/Engine/Source/Runtime/Datasmith/CADKernel/Public/CADKernel/Math/Point.h
@@ -3,11 +3,6 @@
 
 #include "CADKernel/Core/Types.h"
 #include "CADKernel/Math/MathConst.h"
-//#include "HAL/PlatformMath.h"
-//#include "Math/UnrealMathSSE.h"
-//#include "Math/Vector.h"
-//#include "Serialization/Archive.h"
-//#include "Templates/TypeHash.h"
 
 // this will eventually be reconciled with what the Geometry team is using.
 namespace CADKernel
diff --git a/Engine/Source/Runtime/Datasmith/CADKernel/Public/CADKernel/Mesh/Criteria/Criterion.h b/Engine/Source/Runtime/Datasmith/CADKernel/Public/CADKernel/Mesh/Criteria/Criterion.h
--- a/Engine/Source/Runtime/Datasmith/CADKernel/Public/CADKernel/Mesh/Criteria/Criterion.h
+++ b/Engine/Source/Runtime/Datasmith/CADKernel/Public/CADKernel/Mesh/Criteria/Criterion.h
@@ -7,13 +7,6 @@
 #include "CADKernel/Geo/GeoPoint.h"
 #include "CADKernel/Math/Point.h"
 #include "CADKernel/Mesh/Criteria/CriterionType.h"
-//#include "Containers/Array.h"
-//#include "Containers/UnrealString.h"
-#include "HAL/PlatformMath.h"
-#include "Math/UnrealMathUtility.h"
-#include "Misc/AssertionMacros.h"
-#include "Templates/SharedPointer.h"
-
 
 namespace CADKernel
 {
diff --git a/Engine/Source/Runtime/Datasmith/CADKernel/Public/CADKernel/Mesh/Structure/FaceMesh.h b/Engine/Source/Runtime/Datasmith/CADKernel/Public/CADKernel/Mesh/Structure/FaceMesh.h
--- a/Engine/Source/Runtime/Datasmith/CADKernel/Public/CADKernel/Mesh/Structure/FaceMesh.h
+++ b/Engine/Source/Runtime/Datasmith/CADKernel/Public/CADKernel/Mesh/Structure/FaceMesh.h
@@ -2,12 +2,6 @@
 #pragma once
 
 #include "CADKernel/Mesh/Structure/Mesh.h"
-//#include "Containers/Array.h"
-//#include "Containers/Map.h"
-#include "HAL/PlatformCrt.h"
-#include "Math/UnrealMathSSE.h"
-#include "Math/Vector.h"
-#include "Math/Vector2D.h"
 
 namespace CADKernel
 {
diff --git a/Engine/Source/Runtime/Datasmith/CADKernel/Public/CADKernel/Mesh/Structure/Mesh.h b/Engine/Source/Runtime/Datasmith/CADKernel/Public/CADKernel/Mesh/Structure/Mesh.h
--- a/Engine/Source/Runtime/Datasmith/CADKernel/Public/CADKernel/Mesh/Structure/Mesh.h
+++ b/Engine/Source/Runtime/Datasmith/CADKernel/Public/CADKernel/Mesh/Structure/Mesh.h
@@ -4,9 +4,7 @@
 #include "CADKernel/Core/Entity.h"
 #include "CADKernel/Core/EntityGeom.h"
 #include "CADKernel/Core/Types.h"
-//#include "CADKernel/Math/Point.h"
-//#include "Containers/Array.h"
-//#include "HAL/Platform.h"
+#include "CADKernel/Math/Point.h"
 
 namespace CADKernel
 {
diff --git a/Engine/Source/Runtime/Datasmith/CADKernel/Public/CADKernel/Mesh/Structure/ModelMesh.h b/Engine/Source/Runtime/Datasmith/CADKernel/Public/CADKernel/Mesh/Structure/ModelMesh.h
--- a/Engine/Source/Runtime/Datasmith/CADKernel/Public/CADKernel/Mesh/Structure/ModelMesh.h
+++ b/Engine/Source/Runtime/Datasmith/CADKernel/Public/CADKernel/Mesh/Structure/ModelMesh.h
@@ -6,10 +6,6 @@
 #include "CADKernel/Core/Session.h"
 #include "CADKernel/Core/Types.h"
 #include "CADKernel/Mesh/Structure/Mesh.h"
-#include "Containers/Array.h"
-//#include "HAL/PlatformCrt.h"
-//#include "Math/UnrealMathSSE.h"
-//#include "Templates/SharedPointer.h"
 
 namespace CADKernel
 {
diff --git a/Engine/Source/Runtime/Datasmith/CADKernel/Public/CADKernel/Topo/Body.h b/Engine/Source/Runtime/Datasmith/CADKernel/Public/CADKernel/Topo/Body.h
--- a/Engine/Source/Runtime/Datasmith/CADKernel/Public/CADKernel/Topo/Body.h
+++ b/Engine/Source/Runtime/Datasmith/CADKernel/Public/CADKernel/Topo/Body.h
@@ -8,9 +8,6 @@
 #include "CADKernel/Topo/TopologicalEntity.h"
 #include "CADKernel/Topo/TopologicalFace.h"
 #include "CADKernel/Topo/TopologicalShapeEntity.h"
-//#include "Containers/Array.h"
-//#include "HAL/Platform.h"
-//#include "Templates/SharedPointer.h"
 
 namespace CADKernel
 {
diff --git a/Engine/Source/Runtime/Datasmith/CADKernel/Public/CADKernel/Topo/Model.h b/Engine/Source/Runtime/Datasmith/CADKernel/Public/CADKernel/Topo/Model.h
--- a/Engine/Source/Runtime/Datasmith/CADKernel/Public/CADKernel/Topo/Model.h
+++ b/Engine/Source/Runtime/Datasmith/CADKernel/Public/CADKernel/Topo/Model.h
@@ -5,10 +5,7 @@
 #include "CADKernel/Core/Entity.h"
 #include "CADKernel/Core/EntityGeom.h"
 #include "CADKernel/Topo/TopologicalShapeEntity.h"
-#include "Containers/Array.h"
 #include "CoreTypes.h"
-#include "HAL/PlatformCrt.h"
-#include "Templates/SharedPointer.h"
 
 namespace CADKernel
 {
diff --git a/Engine/Source/Runtime/Datasmith/CADKernel/Public/CADKernel/Topo/Shell.h b/Engine/Source/Runtime/Datasmith/CADKernel/Public/CADKernel/Topo/Shell.h
--- a/Engine/Source/Runtime/Datasmith/CADKernel/Public/CADKernel/Topo/Shell.h
+++ b/Engine/Source/Runtime/Datasmith/CADKernel/Public/CADKernel/Topo/Shell.h
@@ -10,10 +10,6 @@
 #include "CADKernel/Geo/GeoEnum.h"
 #include "CADKernel/Topo/TopologicalFace.h"
 #include "CADKernel/Topo/TopologicalShapeEntity.h"
-#include "Containers/Array.h"
-#include "HAL/Platform.h"
-#include "Templates/SharedPointer.h"
-#include "Templates/UnrealTemplate.h"
 
 namespace CADKernel
 {
diff --git a/Engine/Source/Runtime/Datasmith/CADKernel/Public/CADKernel/Topo/TopologicalFace.h b/Engine/Source/Runtime/Datasmith/CADKernel/Public/CADKernel/Topo/TopologicalFace.h
--- a/Engine/Source/Runtime/Datasmith/CADKernel/Public/CADKernel/Topo/TopologicalFace.h
+++ b/Engine/Source/Runtime/Datasmith/CADKernel/Public/CADKernel/Topo/TopologicalFace.h
@@ -15,15 +15,6 @@
 #include "CADKernel/Topo/TopologicalEntity.h"
 #include "CADKernel/Topo/TopologicalLoop.h"
 #include "CADKernel/Topo/TopologicalShapeEntity.h"
-//#include "CADKernel/Utils/Cache.h"
-//#include "Containers/Array.h"
-//#include "Containers/Map.h"
-//#include "Containers/UnrealString.h"
-//#include "HAL/PlatformCrt.h"
-//#include "Math/UnrealMathSSE.h"
-//#include "Math/Vector.h"
-//#include "Misc/AssertionMacros.h"
-//#include "Templates/SharedPointer.h"
 
 namespace CADKernel
 {
diff --git a/Engine/Source/Runtime/Datasmith/CADKernel/Public/CADKernel/Topo/TopologicalLoop.h b/Engine/Source/Runtime/Datasmith/CADKernel/Public/CADKernel/Topo/TopologicalLoop.h
--- a/Engine/Source/Runtime/Datasmith/CADKernel/Public/CADKernel/Topo/TopologicalLoop.h
+++ b/Engine/Source/Runtime/Datasmith/CADKernel/Public/CADKernel/Topo/TopologicalLoop.h
@@ -9,10 +9,6 @@
 #include "CADKernel/Topo/TopologicalEdge.h"
 #include "CADKernel/Topo/TopologicalEntity.h"
 #include "CADKernel/Topo/TopologicalVertex.h"
-//#include "Containers/Array.h"
-//#include "HAL/Platform.h"
-//#include "HAL/PlatformCrt.h"
-//#include "Templates/SharedPointer.h"
 
 namespace CADKernel
 {
diff --git a/Engine/Source/Runtime/Datasmith/CADKernel/Public/CADKernel/Topo/TopologicalVertex.h b/Engine/Source/Runtime/Datasmith/CADKernel/Public/CADKernel/Topo/TopologicalVertex.h
--- a/Engine/Source/Runtime/Datasmith/CADKernel/Public/CADKernel/Topo/TopologicalVertex.h
+++ b/Engine/Source/Runtime/Datasmith/CADKernel/Public/CADKernel/Topo/TopologicalVertex.h
@@ -8,10 +8,6 @@
 #include "CADKernel/Math/Point.h"
 #include "CADKernel/Topo/Linkable.h"
 #include "CADKernel/Topo/TopologicalLink.h"
-#include "Containers/Array.h"
-#include "HAL/Platform.h"
-#include "Misc/AssertionMacros.h"
-#include "Templates/SharedPointer.h"
 
 namespace CADKernel
 {
diff --git a/Engine/Source/Runtime/Datasmith/CADKernel/Public/CADKernel/Topo/Topomaker.h b/Engine/Source/Runtime/Datasmith/CADKernel/Public/CADKernel/Topo/Topomaker.h
--- a/Engine/Source/Runtime/Datasmith/CADKernel/Public/CADKernel/Topo/Topomaker.h
+++ b/Engine/Source/Runtime/Datasmith/CADKernel/Public/CADKernel/Topo/Topomaker.h
@@ -2,8 +2,6 @@
 #pragma once
 
 #include "CADKernel/Core/Types.h"
-#include "Containers/Array.h"
-#include "Templates/SharedPointer.h"
 
 #ifdef CADKERNEL_DEV
 #include "CADKernel/Topo/TopomakerReport.h"
diff --git a/Engine/Source/Runtime/Datasmith/CADKernel/Public/CADKernel/UI/Display.h b/Engine/Source/Runtime/Datasmith/CADKernel/Public/CADKernel/UI/Display.h
--- a/Engine/Source/Runtime/Datasmith/CADKernel/Public/CADKernel/UI/Display.h
+++ b/Engine/Source/Runtime/Datasmith/CADKernel/Public/CADKernel/UI/Display.h
@@ -6,11 +6,6 @@
 #include "CADKernel/Geo/GeoEnum.h"
 #include "CADKernel/Math/Point.h"
 #include "CADKernel/UI/Visu.h"
-#include "Containers/Array.h"
-#include "Containers/UnrealString.h"
-#include "HAL/Platform.h"
-#include "Math/Plane.h"
-#include "Templates/SharedPointer.h"
 
 class FString;
 
diff --git a/Engine/Source/Runtime/Datasmith/CADKernel/Public/CADKernel/UI/Message.h b/Engine/Source/Runtime/Datasmith/CADKernel/Public/CADKernel/UI/Message.h
--- a/Engine/Source/Runtime/Datasmith/CADKernel/Public/CADKernel/UI/Message.h
+++ b/Engine/Source/Runtime/Datasmith/CADKernel/Public/CADKernel/UI/Message.h
@@ -4,10 +4,6 @@
 #include "CADKernel/Core/Chrono.h"
 #include "CADKernel/Core/Types.h"
 #include "CADKernel/UI/Progress.h"
-#include "Containers/UnrealString.h"
-#include "HAL/Platform.h"
-#include "Templates/AndOrNot.h"
-#include "Templates/IsValidVariadicFunctionArg.h"
 
 namespace CADKernel
 {
diff --git a/Engine/Source/Runtime/Datasmith/CADKernel/Public/CADKernel/UI/Progress.h b/Engine/Source/Runtime/Datasmith/CADKernel/Public/CADKernel/UI/Progress.h
--- a/Engine/Source/Runtime/Datasmith/CADKernel/Public/CADKernel/UI/Progress.h
+++ b/Engine/Source/Runtime/Datasmith/CADKernel/Public/CADKernel/UI/Progress.h
@@ -2,8 +2,6 @@
 #pragma once
 
 #include "CADKernel/Core/Types.h"
-#include "Containers/UnrealString.h"
-#include "HAL/Platform.h"
 
 namespace CADKernel
 {
diff --git a/Engine/Source/Runtime/Datasmith/CADKernel/Public/CADKernel/UI/Visu.h b/Engine/Source/Runtime/Datasmith/CADKernel/Public/CADKernel/UI/Visu.h
--- a/Engine/Source/Runtime/Datasmith/CADKernel/Public/CADKernel/UI/Visu.h
+++ b/Engine/Source/Runtime/Datasmith/CADKernel/Public/CADKernel/UI/Visu.h
@@ -5,9 +5,6 @@
 #include "CADKernel/Core/Parameters.h"
 #include "CADKernel/Core/Types.h"
 #include "CADKernel/Math/Point.h"
-#include "Containers/Array.h"
-#include "HAL/Platform.h"
-#include "Misc/AssertionMacros.h"
 
 namespace CADKernel
 {
