[CADKernel] Fix crashs in meshing step: improve BoywerWatson algorithm

diff --git a/Engine/Source/Runtime/Datasmith/CADKernel/Private/CADKernel/Math/SlopeUtils.h b/Engine/Source/Runtime/Datasmith/CADKernel/Private/CADKernel/Math/SlopeUtils.h
--- a/Engine/Source/Runtime/Datasmith/CADKernel/Private/CADKernel/Math/SlopeUtils.h
+++ b/Engine/Source/Runtime/Datasmith/CADKernel/Private/CADKernel/Math/SlopeUtils.h
@@ -9,348 +9,375 @@
 namespace CADKernel
 {
 
-	typedef TFunction<double(const FPoint2D&, const FPoint2D&, double)> SlopMethod;
+namespace Slope
+{
 
+constexpr double NullSlope = 0.;
+constexpr double FullSlope = 8.;
 
-	/**
-	 * Transform a positive slope into an oriented slope [-4, 4] i.e. an equivalent angle between [-Pi, Pi]
-	 * @return a slope between [-4, 4]
-	 */
-	inline double TransformIntoOrientedSlope(double Slope)
-	{
-		if (Slope > 4.) Slope -= 8;
-		if (Slope < -4.) Slope += 8;
-		return Slope;
-	}
+/**
+ * RightSlope i.e. Right angle i.e Pi / 2
+ */
+constexpr double RightSlope = 2.;
+
+/**
+ * ThreeRightSlope i.e. 3Pi / 2
+ */
+constexpr double ThreeRightSlope = 6.;
+
+/**
+ * MinusRightSlope i.e. -Pi / 2
+ */
+constexpr double MinusRightSlope = -2.;
+
+/**
+ * PiSlope i.e. Pi angle
+ */
+constexpr double PiSlope = 4.;
+}
+
+typedef TFunction<double(const FPoint2D&, const FPoint2D&, double)> SlopMethod;
+
+
+/**
+ * Transform a positive slope into an oriented slope [-4, 4] i.e. an equivalent angle between [-Pi, Pi]
+ * @return a slope between [-4, 4]
+ */
+inline double TransformIntoOrientedSlope(double Slope)
+{
+	if (Slope > 4.) Slope -= 8;
+	if (Slope < -4.) Slope += 8;
+	return Slope;
+}
+
+/**
+ * Transform a slope into a positive slope [0, *] i.e. an equivalent angle between [0, 2.Pi]
+ * @return a slope between [0, 8]
+ */
+inline double TransformIntoPositiveSlope(double Slope)
+{
+	if (Slope < 0) Slope += 8;
+	return Slope;
+}
+
+/**
+ * Fast angle approximation by a "slope" :
+ * This method is used to compute an approximation of the angle between the input segment defined by two points and [0, u) axis.
+ * The return value is a real in the interval [0, 8] for an angle in the interval [0, 2Pi]
+ * Warning, it's only an approximation... The conversion is not linear but the error is small near the integer value of slope (0, 1, 2, 3, ...8)
+ *
+ * This approximation is very good when only comparison of angles is needed.
+ * This method is not adapted to compute angle
+ *
+ * To compute an angle value between two segments, the call of acos (and asin for an oriented angle) is necessary while with this approximation, only a division is useful.
+ *
+ * [0 - 2Pi] is divide into 8 angular sector i.e. [0, Pi/4] = [0,1], [Pi/4, Pi/2] = [1,2], ...
+ *
+ * @return a slope between [0, 8] i.e. an equivalent angle between [0, 2Pi]
+ *
+ * Angle (Degree) to Slop
+ *		  0   = 0
+ *		  7.  = 0.125
+ *		 14.  = 0.25
+ *		 30   = 0.5
+ *		 36.8 = 0.75
+ *		 45   = 1
+ *		 53.2 = 1.25
+ *		 60   = 1.5
+ *		 76.  = 1.75
+ *		 90   = 2
+ *		135   = 3
+ *		180   = 4
+ *		360   = 8
+ */
+inline double ComputeSlope(const FPoint2D& StartPoint, const FPoint2D& EndPoint)
+{
+	double DeltaU = EndPoint.U - StartPoint.U;
+	double DeltaV = EndPoint.V - StartPoint.V;
+	double Delta;
 
-	/**
-	 * Transform a slope into a positive slope [0, *] i.e. an equivalent angle between [0, 2.Pi]
-	 * @return a slope between [0, 8]
-	 */
-	inline double TransformIntoPositiveSlope(double Slope)
+	if (FMath::Abs(DeltaU) < DOUBLE_SMALL_NUMBER && FMath::Abs(DeltaV) < DOUBLE_SMALL_NUMBER)
 	{
-		if (Slope < 0) Slope += 8;
-		return Slope;
+		return 0;
 	}
 
-	/**
-	 * Fast angle approximation by a "slope" :
-	 * This method is used to compute an approximation of the angle between the input segment defined by two points and [0, u) axis.
-	 * The return value is a real in the interval [0, 8] for an angle in the interval [0, 2Pi]
-	 * Warning, it's only an approximation... The conversion is not linear but the error is small near the integer value of slope (0, 1, 2, 3, ...8)
-	 * 
-	 * This approximation is very good when only comparison of angles is needed.
-	 * This method is not adapted to compute angle
-	 *
-	 * To compute an angle value between two segments, the call of acos (and asin for an oriented angle) is necessary while with this approximation, only a division is useful.
-	 * 
-	 * [0 - 2Pi] is divide into 8 angular sector i.e. [0, Pi/4] = [0,1], [Pi/4, Pi/2] = [1,2], ...
-	 * 
-	 * @return a slope between [0, 8] i.e. an equivalent angle between [0, 2Pi]
-	 *
-	 * Angle (Degree) to Slop
-	 *		  0   = 0
-	 *		  7.  = 0.125
-	 *		 14.  = 0.25
-	 *		 30   = 0.5
-	 *		 36.8 = 0.75 
-	 *		 45   = 1
-	 *		 53.2 = 1.25
-	 *		 60   = 1.5
-	 *		 76.  = 1.75
-	 *		 90   = 2
-	 *		135   = 3
-	 *		180   = 4
-	 *		360   = 8
-	 */
-	inline double ComputeSlope(const FPoint2D& StartPoint, const FPoint2D& EndPoint)
+	if (DeltaU > DOUBLE_SMALL_NUMBER)
 	{
-		double DeltaU = EndPoint.U - StartPoint.U;
-		double DeltaV = EndPoint.V - StartPoint.V;
-		double Delta;
-
-		if (FMath::Abs(DeltaU) < DOUBLE_SMALL_NUMBER && FMath::Abs(DeltaV) < DOUBLE_SMALL_NUMBER)
+		if (DeltaV > DOUBLE_SMALL_NUMBER)
 		{
-			return 0;
+			if (DeltaU > DeltaV)
+			{
+				Delta = DeltaV / DeltaU;
+			}
+			else
+			{
+				Delta = 2 - DeltaU / DeltaV;
+			}
 		}
-
-		if (DeltaU > DOUBLE_SMALL_NUMBER)
+		else
+		{
+			if (DeltaU > -DeltaV)
+			{
+				Delta = 8 + DeltaV / DeltaU;
+			}
+			else
+			{
+				Delta = 6 - DeltaU / DeltaV; // deltaU/deltaV <0
+			}
+		}
+	}
+	else
+	{
+		if (DeltaV > DOUBLE_SMALL_NUMBER)
 		{
-			if (DeltaV > DOUBLE_SMALL_NUMBER)
+			if (-DeltaU > DeltaV)
 			{
-				if (DeltaU > DeltaV)
-				{
-					Delta = DeltaV / DeltaU;
-				}
-				else
-				{
-					Delta = 2 - DeltaU / DeltaV;
-				}
+				Delta = 4 + DeltaV / DeltaU;
 			}
 			else
 			{
-				if (DeltaU > -DeltaV)
-				{
-					Delta = 8 + DeltaV / DeltaU;
-				}
-				else
-				{
-					Delta = 6 - DeltaU / DeltaV; // deltaU/deltaV <0
-				}
+				Delta = 2 - DeltaU / DeltaV;
 			}
 		}
 		else
 		{
-			if (DeltaV > DOUBLE_SMALL_NUMBER)
+			if (-DeltaU > -DeltaV)
 			{
-				if (-DeltaU > DeltaV)
-				{
-					Delta = 4 + DeltaV / DeltaU;
-				}
-				else
-				{
-					Delta = 2 - DeltaU / DeltaV;
-				}
+				Delta = 4 + DeltaV / DeltaU;
 			}
 			else
 			{
-				if (-DeltaU > -DeltaV)
-				{
-					Delta = 4 + DeltaV / DeltaU;
-				}
-				else
-				{
-					Delta = 6 - DeltaU / DeltaV;
-				}
+				Delta = 6 - DeltaU / DeltaV;
 			}
 		}
-
-		return Delta;
 	}
 
-	/**
-	 * Compute the oriented slope of a segment according to a reference slope
-	 * This method is used to compute an approximation of the angle between two segments in 2D.
-	 * return a slope between [0, 8] i.e. an equivalent angle between [0, 2Pi]
- 	 */
-	inline double ComputePositiveSlope(const FPoint2D& StartPoint, const FPoint2D& EndPoint, double ReferenceSlope)
-	{
-		double Slope = ComputeSlope(StartPoint, EndPoint);
-		Slope -= ReferenceSlope;
-		if (Slope < 0) Slope += 8;
-		return Slope;
-	}
+	return Delta;
+}
 
-	/**
-	 * Compute the positive slope between the segments [StartPoint, EndPoint1] and [StartPoint, EndPoint2]
-	 * @return a slope between [0, 8] i.e. an equivalent angle between [0, 2Pi]
-	 */
-	inline double ComputePositiveSlope(const FPoint2D& StartPoint, const FPoint2D& EndPoint1, const FPoint2D& EndPoint2)
-	{
-		double ReferenceSlope = ComputeSlope(StartPoint, EndPoint1);
-		double Slope = ComputeSlope(StartPoint, EndPoint2);
-		Slope -= ReferenceSlope;
-		return TransformIntoPositiveSlope(Slope);
-	}
+/**
+ * Compute the oriented slope of a segment according to a reference slope
+ * This method is used to compute an approximation of the angle between two segments in 2D.
+ * return a slope between [0, 8] i.e. an equivalent angle between [0, 2Pi]
+ */
+inline double ComputePositiveSlope(const FPoint2D& StartPoint, const FPoint2D& EndPoint, double ReferenceSlope)
+{
+	double Slope = ComputeSlope(StartPoint, EndPoint);
+	Slope -= ReferenceSlope;
+	if (Slope < 0) Slope += 8;
+	return Slope;
+}
+
+/**
+ * Compute the positive slope between the segments [StartPoint, EndPoint1] and [StartPoint, EndPoint2]
+ * @return a slope between [0, 8] i.e. an equivalent angle between [0, 2Pi]
+ */
+inline double ComputePositiveSlope(const FPoint2D& StartPoint, const FPoint2D& EndPoint1, const FPoint2D& EndPoint2)
+{
+	double ReferenceSlope = ComputeSlope(StartPoint, EndPoint1);
+	double Slope = ComputeSlope(StartPoint, EndPoint2);
+	Slope -= ReferenceSlope;
+	return TransformIntoPositiveSlope(Slope);
+}
 
-	inline double ClockwiseSlop(const FPoint2D& StartPoint, const FPoint2D& EndPoint, double ReferenceSlope)
-	{
-		return 8. - ComputePositiveSlope(StartPoint, EndPoint, ReferenceSlope);
-	}
+inline double ClockwiseSlop(const FPoint2D& StartPoint, const FPoint2D& EndPoint, double ReferenceSlope)
+{
+	return 8. - ComputePositiveSlope(StartPoint, EndPoint, ReferenceSlope);
+}
 
-	inline double CounterClockwiseSlop(const FPoint2D& StartPoint, const FPoint2D& EndPoint, double ReferenceSlope)
-	{
-		return ComputePositiveSlope(StartPoint, EndPoint, ReferenceSlope);
-	}
+inline double CounterClockwiseSlop(const FPoint2D& StartPoint, const FPoint2D& EndPoint, double ReferenceSlope)
+{
+	return ComputePositiveSlope(StartPoint, EndPoint, ReferenceSlope);
+}
+
+/**
+ * Compute the oriented slope of a segment according to a reference slope
+ * @return a slope between [-4, 4] i.e. an equivalent angle between [-Pi, Pi]
+ */
+inline double ComputeOrientedSlope(const FPoint2D& StartPoint, const FPoint2D& EndPoint, double ReferenceSlope)
+{
+	return TransformIntoOrientedSlope(ComputePositiveSlope(StartPoint, EndPoint, ReferenceSlope));
+}
+
+/**
+ * Compute the positive slope between the segments [StartPoint, EndPoint1] and [StartPoint, EndPoint2]
+ * @return a slope between [-4, 4] i.e. an equivalent angle between [-Pi, Pi]
+ */
+inline double ComputeOrientedSlope(const FPoint2D& StartPoint, const FPoint2D& EndPoint1, const FPoint2D& EndPoint2)
+{
+	return TransformIntoOrientedSlope(ComputePositiveSlope(StartPoint, EndPoint1, EndPoint2));
+}
 
-	/**
-	 * Compute the oriented slope of a segment according to a reference slope
-	 * @return a slope between [-4, 4] i.e. an equivalent angle between [-Pi, Pi]
-	 */
-	inline double ComputeOrientedSlope(const FPoint2D& StartPoint, const FPoint2D& EndPoint, double ReferenceSlope)
+/**
+ * return a slope between [0, 4] i.e. an angle between [0, Pi]
+ */
+inline double ComputeUnorientedSlope(const FPoint2D& StartPoint, const FPoint2D& EndPoint, double ReferenceSlope)
+{
+	return FMath::Abs(ComputeOrientedSlope(StartPoint, EndPoint, ReferenceSlope));
+}
+
+/**
+ * return a slope between [0, 1] relative to the nearest axis between horizontal or vertical axis i.e.
+ * ComputeUnorientedSlope => 0.5 return 0.5
+ * ComputeUnorientedSlope => 2.3 return 0.3
+ * ComputeUnorientedSlope => 3.6 return 0.4
+ */
+inline double ComputeSlopeRelativeToNearestAxis(const FPoint2D& StartPoint, const FPoint2D& EndPoint)
+{
+	double Slope = FMath::Abs(TransformIntoOrientedSlope(ComputeSlope(StartPoint, EndPoint)));
+	if (Slope > 2)
 	{
-		return TransformIntoOrientedSlope(ComputePositiveSlope(StartPoint, EndPoint, ReferenceSlope));
+		Slope = 4 - Slope;
 	}
 
-	/**
-	 * Compute the positive slope between the segments [StartPoint, EndPoint1] and [StartPoint, EndPoint2]
-	 * @return a slope between [-4, 4] i.e. an equivalent angle between [-Pi, Pi]
-	 */
-	inline double ComputeOrientedSlope(const FPoint2D& StartPoint, const FPoint2D& EndPoint1, const FPoint2D& EndPoint2)
+	// if slope close to 2 means segment close to IsoU, otherwise segment close to IsoV
+	// Wants a slope between 0 and 1 to manage either IsoU and IsoV
+	// Close to 0 means close to IsoU or IsoV
+	if (Slope > 1)
 	{
-		return TransformIntoOrientedSlope(ComputePositiveSlope(StartPoint, EndPoint1, EndPoint2));
+		Slope = 2 - Slope;
 	}
 
-	/**
-	 * return a slope between [0, 4] i.e. an angle between [0, Pi]
-	 */
-	inline double ComputeUnorientedSlope(const FPoint2D& StartPoint, const FPoint2D& EndPoint, double ReferenceSlope)
-	{
-		return FMath::Abs(ComputeOrientedSlope(StartPoint, EndPoint, ReferenceSlope));
-	}
+	return Slope;
+}
 
-	/**
-	 * return a slope between [0, 1] relative to the nearest axis between horizontal or vertical axis i.e. 
-	 * ComputeUnorientedSlope => 0.5 return 0.5
-	 * ComputeUnorientedSlope => 2.3 return 0.3
-	 * ComputeUnorientedSlope => 3.6 return 0.4
-	 */
-	inline double ComputeSlopeRelativeToNearestAxis(const FPoint2D& StartPoint, const FPoint2D& EndPoint)
+/**
+ * return a slope between [0, 2] relative to reference Axis i.e.
+ * ComputeUnorientedSlope => 0.5 return 0.5
+ * ComputeUnorientedSlope => 2.3 return 1.7
+ */
+inline double ComputeSlopeRelativeToReferenceAxis(const FPoint2D& StartPoint, const FPoint2D& EndPoint, double ReferenceAxisSlope)
+{
+	double Slope = ComputeUnorientedSlope(StartPoint, EndPoint, ReferenceAxisSlope);
+	if (Slope > 2)
 	{
-		double Slope = FMath::Abs(TransformIntoOrientedSlope(ComputeSlope(StartPoint, EndPoint)));
-		if (Slope > 2)
-		{
-			Slope = 4 - Slope;
-		}
-
-		// if slope close to 2 means segment close to IsoU, otherwise segment close to IsoV
-		// Wants a slope between 0 and 1 to manage either IsoU and IsoV
-		// Close to 0 means close to IsoU or IsoV
-		if (Slope > 1)
-		{
-			Slope = 2 - Slope;
-		}
-
-		return Slope;
+		Slope = 4 - Slope;
 	}
 
-	/**
-	 * return a slope between [0, 2] relative to reference Axis i.e.
-	 * ComputeUnorientedSlope => 0.5 return 0.5
-	 * ComputeUnorientedSlope => 2.3 return 1.7
-	 */
-	inline double ComputeSlopeRelativeToReferenceAxis(const FPoint2D& StartPoint, const FPoint2D& EndPoint, double ReferenceAxisSlope)
-	{
-		double Slope = ComputeUnorientedSlope(StartPoint, EndPoint, ReferenceAxisSlope);
-		if (Slope > 2)
-		{
-			Slope = 4 - Slope;
-		}
+	return Slope;
+}
 
-		return Slope;
-	}
-
-	/**
-	 * return a slope between [0, 4] i.e. an angle between [0, Pi]
-	 */
-	inline double ComputeUnorientedSlope(const FPoint2D& StartPoint, const FPoint2D& EndPoint1, const FPoint2D& EndPoint2)
+/**
+ * return a slope between [0, 4] i.e. an angle between [0, Pi]
+ */
+inline double ComputeUnorientedSlope(const FPoint2D& StartPoint, const FPoint2D& EndPoint1, const FPoint2D& EndPoint2)
+{
+	return FMath::Abs(ComputeOrientedSlope(StartPoint, EndPoint1, EndPoint2));
+}
+
+/**
+ *                         P1
+ *          inside        /
+ *                       /   inside
+ *                      /
+ *    A -------------- B --------------- C
+ *                      \
+ *           Outside     \  Outside
+ *                        \
+ *                         P2
+ *
+ * Return true if the segment BP is inside the sector defined the half-lines [BA) and [BC) in the counterclockwise.
+ * Return false if ABP angle or PBC angle is too flat (smaller than FlatAngle)
+ */
+inline bool IsPointPInsideSectorABC(const FPoint2D& PointA, const FPoint2D& PointB, const FPoint2D& PointC, const FPoint2D& PointP, const double FlatAngle)
+{
+	double SlopWithNextBoundary = ComputeSlope(PointB, PointC);
+	double BoundaryDeltaSlope = ComputePositiveSlope(PointB, PointA, SlopWithNextBoundary);
+	double SegmentSlope = ComputePositiveSlope(PointB, PointP, SlopWithNextBoundary);
+	if (SegmentSlope < FlatAngle || SegmentSlope + FlatAngle > BoundaryDeltaSlope)
 	{
-		return FMath::Abs(ComputeOrientedSlope(StartPoint, EndPoint1, EndPoint2));
+		return false;
 	}
+	return true;
+}
+
+/**
+ *                         P1
+ *          inside        /
+ *                       /   inside
+ *                      /
+ *    A -------------- B --------------- C
+ *                      \
+ *           Outside     \  Outside
+ *                        \
+ *                         P2
+ *
+ * Return true if all of the segment BPi is inside the sector defined the half-lines [BA) and [BC) in the counterclockwise.
+ */
+inline bool ArePointsInsideSectorABC(const FPoint2D& PointA, const FPoint2D& PointB, const FPoint2D& PointC, const TArray<const FPoint2D*>& Points, const double FlatAngle = -DOUBLE_SMALL_NUMBER)
+{
+	double SlopWithNextBoundary = ComputeSlope(PointB, PointC);
+	double BoundaryDeltaSlope = ComputePositiveSlope(PointB, PointA, SlopWithNextBoundary);
 
-	/**
-	 *                         P1
-	 *          inside        /
-	 *                       /   inside
-	 *                      /
-	 *    A -------------- B --------------- C 
-	 *                      \
-	 *           Outside     \  Outside
-	 *                        \
-	 *                         P2
-	 *
-	 * Return true if the segment BP is inside the sector defined the half-lines [BA) and [BC) in the counterclockwise.
-	 * Return false if ABP angle or PBC angle is too flat (smaller than FlatAngle)
-	 */
-	inline bool IsPointPInsideSectorABC(const FPoint2D& PointA, const FPoint2D& PointB, const FPoint2D& PointC, const FPoint2D& PointP, const double FlatAngle)
-	 {
-		 double SlopWithNextBoundary = ComputeSlope(PointB, PointC);
-		 double BoundaryDeltaSlope = ComputePositiveSlope(PointB, PointA, SlopWithNextBoundary);
-		 double SegmentSlope = ComputePositiveSlope(PointB, PointP, SlopWithNextBoundary);
-		 if (SegmentSlope < FlatAngle || SegmentSlope + FlatAngle > BoundaryDeltaSlope)
-		 {
-			 return false;
-		 }
-		 return true;
-	 }
-
-	/**
-	 *                         P1
-	 *          inside        /
-	 *                       /   inside
-	 *                      /
-	 *    A -------------- B --------------- C
-	 *                      \
-	 *           Outside     \  Outside
-	 *                        \
-	 *                         P2
-	 *
-	 * Return true if all of the segment BPi is inside the sector defined the half-lines [BA) and [BC) in the counterclockwise.
-	 */
-	inline bool ArePointsInsideSectorABC(const FPoint2D& PointA, const FPoint2D& PointB, const FPoint2D& PointC, const TArray<const FPoint2D*>& Points, const double FlatAngle = -DOUBLE_SMALL_NUMBER)
-	{
-		double SlopWithNextBoundary = ComputeSlope(PointB, PointC);
-		double BoundaryDeltaSlope = ComputePositiveSlope(PointB, PointA, SlopWithNextBoundary);
-
-		return Algo::AllOf(Points, [&](const FPoint2D* PointP) {
-			double DeltaU = PointB.U - PointP->U;
-			double DeltaV = PointB.V - PointP->V;
+	return Algo::AllOf(Points, [&](const FPoint2D* PointP) {
+		double DeltaU = PointB.U - PointP->U;
+		double DeltaV = PointB.V - PointP->V;
 
-			if (FMath::Abs(DeltaU) < DOUBLE_SMALL_NUMBER && FMath::Abs(DeltaV) < DOUBLE_SMALL_NUMBER)
-			{
-				return true;
-			}
-
-			double SegmentSlope = ComputePositiveSlope(PointB, *PointP, SlopWithNextBoundary);
-			if (SegmentSlope < FlatAngle || SegmentSlope + FlatAngle > BoundaryDeltaSlope)
-			{
-				return false;
-			}
+		if (FMath::Abs(DeltaU) < DOUBLE_SMALL_NUMBER && FMath::Abs(DeltaV) < DOUBLE_SMALL_NUMBER)
+		{
 			return true;
-			});
-	}
+		}
 
-	inline FPoint2D SlopeToVector(const double Slope)
-	{
-		int32 SlopeStep = (int32) (Slope);
-		FPoint2D Vector;
-		switch (SlopeStep)
+		double SegmentSlope = ComputePositiveSlope(PointB, *PointP, SlopWithNextBoundary);
+		if (SegmentSlope < FlatAngle || SegmentSlope + FlatAngle > BoundaryDeltaSlope)
 		{
-		case 0:
-			// Delta = DeltaV / DeltaU;
-			Vector[0] = 1.;
-			Vector[1] = Slope;
-			break;
-		case 1:
-			// 2 - DeltaU / DeltaV;
-			Vector[0] = 2. - Slope;
-			Vector[1] = 1.;
-			break;
-		case 2:
-			// 2 - DeltaU / DeltaV;
-			Vector[0] = 2. - Slope;
-			Vector[1] = 1.;
-			break;
-		case 3:
-			// 4 + DeltaV / DeltaU;
-			Vector[0] = -1.;
-			Vector[1] = 4. - Slope;
-			break;
-		case 4:
-			// 4 + DeltaV / DeltaU;
-			Vector[0] = -1.;
-			Vector[1] = 4. - Slope;
-			break;
-		case 5:
-			// 6 - DeltaU / DeltaV;
-			Vector[0] = Slope - 6.;
-			Vector[1] = -1.;
-			break;
-		case 6:
-			// 6 - DeltaU / DeltaV // deltaU/deltaV <0
-			Vector[0] = Slope - 6.;
-			Vector[1] = -1;
-			break;
-		case 7:
-			// 8 + DeltaV / DeltaU; // deltaU/deltaV <0
-			Vector[0] = 1.;
-			Vector[1] = Slope - 8.;
-			break;
-		default:
-			break;
+			return false;
 		}
+		return true;
+		});
+}
 
-		return Vector;
+inline FPoint2D SlopeToVector(const double Slope)
+{
+	int32 SlopeStep = (int32)(Slope);
+	FPoint2D Vector;
+	switch (SlopeStep)
+	{
+	case 0:
+		// Delta = DeltaV / DeltaU;
+		Vector[0] = 1.;
+		Vector[1] = Slope;
+		break;
+	case 1:
+		// 2 - DeltaU / DeltaV;
+		Vector[0] = 2. - Slope;
+		Vector[1] = 1.;
+		break;
+	case 2:
+		// 2 - DeltaU / DeltaV;
+		Vector[0] = 2. - Slope;
+		Vector[1] = 1.;
+		break;
+	case 3:
+		// 4 + DeltaV / DeltaU;
+		Vector[0] = -1.;
+		Vector[1] = 4. - Slope;
+		break;
+	case 4:
+		// 4 + DeltaV / DeltaU;
+		Vector[0] = -1.;
+		Vector[1] = 4. - Slope;
+		break;
+	case 5:
+		// 6 - DeltaU / DeltaV;
+		Vector[0] = Slope - 6.;
+		Vector[1] = -1.;
+		break;
+	case 6:
+		// 6 - DeltaU / DeltaV // deltaU/deltaV <0
+		Vector[0] = Slope - 6.;
+		Vector[1] = -1;
+		break;
+	case 7:
+		// 8 + DeltaV / DeltaU; // deltaU/deltaV <0
+		Vector[0] = 1.;
+		Vector[1] = Slope - 8.;
+		break;
+	default:
+		break;
 	}
 
+	return Vector;
+}
+
 } // namespace CADKernel	
diff --git a/Engine/Source/Runtime/Datasmith/CADKernel/Private/CADKernel/Mesh/Meshers/BowyerWatsonTriangulator.h b/Engine/Source/Runtime/Datasmith/CADKernel/Private/CADKernel/Mesh/Meshers/BowyerWatsonTriangulator.h
--- a/Engine/Source/Runtime/Datasmith/CADKernel/Private/CADKernel/Mesh/Meshers/BowyerWatsonTriangulator.h
+++ b/Engine/Source/Runtime/Datasmith/CADKernel/Private/CADKernel/Mesh/Meshers/BowyerWatsonTriangulator.h
@@ -5,17 +5,13 @@
 #include "CADKernel/Math/SlopeUtils.h"
 #include "CADKernel/Mesh/Meshers/IsoTriangulator/IsoCell.h"
 #include "CADKernel/UI/Display.h"
-
 #ifdef CADKERNEL_DEV
 #include "CADKernel/Mesh/Meshers/IsoTriangulator/DefineForDebug.h"
 #endif
 
-//#define DEBUG_BOWYERWATSON
 namespace CADKernel
 {
 
-
-
 class FBowyerWatsonTriangulator
 {
 protected:
@@ -59,13 +55,37 @@ public:
 	void Triangulate()
 	{
 		FTimePoint StartTime = FChrono::Now();
+#ifdef DEBUG_BOWYERWATSON
+		F3DDebugSession DelaunayDebugSession(bDisplay, TEXT("Delaunay Algo"));
+#endif // DEBUG_DELAUNAY
 
 		Vertices.Sort([](const TPair<int32, FPoint2D>& Vertex1, const TPair<int32, FPoint2D>& Vertex2) {return (Vertex1.Value.U + Vertex1.Value.V) > (Vertex2.Value.U + Vertex2.Value.V); });
 
+#ifdef DEBUG_BOWYERWATSON
+		if (bDisplay)
+		{
+			F3DDebugSession _(TEXT("Vertices"));
+			for (const TPair<int32, FPoint2D>& Vertex : Vertices)
+			{
+				//F3DDebugSession _(TEXT("Vertex"));
+				DisplayPoint(Vertex.Value * DisplayScale, EVisuProperty::YellowPoint, Vertex.Key);
+			}
+			//Wait();
+		}
+#endif
 		// initialization of Bowyer & Watson algorithm with a bounding mesh of the vertex cloud 
 		// i.e. 2 triangles defined by the corners of the offset vertices bounding box
 		MakeBoundingMesh();
 
+#ifdef DEBUG_BOWYERWATSON_STEP
+		DisplayTriangles();
+		Wait(bDisplay);
+
+		static int32 StopIndex = 206;
+#endif
+
+		EdgeVertexIndices.Reserve(TriangleSet.Num() * 6);
+
 		// insert each point in the mesh
 		// The points are sorted on the diagonal of the bbox and are inserted from one end to the other  
 		int32 VertexIndex = 0;
@@ -80,12 +100,23 @@ public:
 				VertexIndex = VerticesCount - 1 - VIndex / 2;
 			}
 
-			TriangleIndices.Empty(VerticesCount);
-			AdditionalTriangleIndices.Empty(VerticesCount);
+			SelectedTriangleIndices.Reset(VerticesCount);
+			AdditionalTriangleIndices.Reset(VerticesCount);
 
 			const FPoint2D& NewVertex = Vertices[VertexIndex].Value;
 
-			bool bHasAdditionalTriangles = false;
+#ifdef DEBUG_BOWYERWATSON_STEP
+			F3DDebugSession _(bDisplay, FString::Printf(TEXT("Step %d"), VIndex));
+			if (bDisplay)
+			{
+				F3DDebugSession _(TEXT("New Point"));
+				DisplayPoint(NewVertex * DisplayScale, EVisuProperty::YellowPoint, Vertices[VertexIndex].Key);
+				Wait(VIndex >= StopIndex);
+			}
+#endif
+
+			bool bHaveDoubt = false;
+			
 			// find all triangles whose circumcircles contain the new vertex
 			{
 				// The problem is for triangle with huge circumscribed circle radius (flat triangle)
@@ -95,7 +126,8 @@ public:
 				//    if the slop is nearly null this mean that the new vertex is outside the triangle and will generate a flat triangle
 
 				constexpr double IncreaseFactor = 1.001; // to process the case of new vertex on the circumscribed circle
-				constexpr double ReducingFactor = 0.9999 / IncreaseFactor; // to remove the case of new vertex on the circumscribed circle 
+				constexpr double ReducingFactor = 0.999 / IncreaseFactor; // to remove the case of new vertex on the circumscribed circle 
+				constexpr double LargeReducingFactor = 0.99; // to remove the case of new vertex on the circumscribed circle 
 				// the case of 4 points nearly on the same circle is manage in a second time i.e. we check that the 
 
 				for (int32 TriangleIndex = 0; TriangleIndex < TriangleSet.Num(); TriangleIndex++)
@@ -106,64 +138,112 @@ public:
 
 					const double SquareRadiusMax = TriangleSet[TriangleIndex].SquareRadius * IncreaseFactor;
 					const double SquareRadiusMin = SquareRadiusMax * ReducingFactor;
+					const double SquareRadiusMinMin = SquareRadiusMin * LargeReducingFactor;
 
 					if (SquareDistanceToCenter < SquareRadiusMin)
 					{
-						TriangleIndices.Add(TriangleIndex);
+						SelectedTriangleIndices.Add(TriangleIndex);
+						if (SquareDistanceToCenter > SquareRadiusMinMin)
+						{
+							bHaveDoubt = true;
+						}
 					}
 					else
 					{
 						if (SquareDistanceToCenter < SquareRadiusMax)
 						{
-							bHasAdditionalTriangles = true;
-							AdditionalTriangleIndices.Add(TriangleIndex);
+							if (DoesTriangleContainingNewVertex(NewVertex, TriangleIndex))
+							{
+								SelectedTriangleIndices.Add(TriangleIndex);
+							}
+							else
+							{
+								AdditionalTriangleIndices.Add(TriangleIndex);
+							}
 						}
 					}
 				}
 			}
 
-			if (TriangleIndices.Num() == 0)
+#ifdef DEBUG_BOWYERWATSON_STEP
+			if (bHaveDoubt && (VIndex >= StopIndex))
+			{
+				F3DDebugSession A(bDisplay, TEXT("HaveDoubt"));
+				DisplaySelectedTrianglesAndComplements();
+			}
+#endif
+			if (!bHaveDoubt)
+			{
+				bHaveDoubt = !DoSelectedTrianglesFormSinglePartition();
+			}
+
+			if (SelectedTriangleIndices.Num() == 0)
 			{
 				for (int32& TriangleIndex : AdditionalTriangleIndices)
 				{
-					const FPoint2D& Point0 = Vertices[TriangleSet[TriangleIndex].VertexIndices[0]].Value;
-					const FPoint2D& Point1 = Vertices[TriangleSet[TriangleIndex].VertexIndices[1]].Value;
-					const FPoint2D& Point2 = Vertices[TriangleSet[TriangleIndex].VertexIndices[2]].Value;
-					const double Slop0 = ComputePositiveSlope(Point0, Point1, NewVertex);
-					const double Slop1 = ComputePositiveSlope(Point1, Point2, NewVertex);
-					const double Slop2 = ComputePositiveSlope(Point2, Point0, NewVertex);
-					if (Slop0 < 4 && Slop1 < 4 && Slop2 < 4)
+					if (DoesTriangleContainingNewVertex(NewVertex, TriangleIndex))
 					{
-						TriangleIndices.Add(TriangleIndex);
+						SelectedTriangleIndices.Add(TriangleIndex);
 						TriangleIndex = -1;
 					}
 				}
 			}
+			else if(bHaveDoubt)
+			{
+				TArray<int32> TmpTriangleIndices;
+				TmpTriangleIndices = MoveTemp(SelectedTriangleIndices);
+				SelectedTriangleIndices.Reset(TmpTriangleIndices.Num());
+
+				AddTrianglesContainingNewVertexToSelectedTriangles(NewVertex, TmpTriangleIndices);
+
+				AddConnectedTrianglesToSelection(TmpTriangleIndices);
+
+				for (int32& TriangleIndex : TmpTriangleIndices)
+				{
+					if (TriangleIndex >= 0)
+					{
+ 						AdditionalTriangleIndices.Add(TriangleIndex);
+					}
+				}
+			}
 
-			if (bHasAdditionalTriangles)
+#ifdef DEBUG_BOWYERWATSON_STEP
+			if (VIndex == StopIndex)
 			{
+				DisplaySelectedTrianglesAndComplements(); 
+				Wait();
+			}
+#endif
 
-				TArray<int32> TmpVertexIndices;
-				for (int32 TriangleIndex : TriangleIndices)
+			if (AdditionalTriangleIndices.Num() > 0)
+			{
+				// Additional triangles are triangles that the new vertex is nearly coincident to the circle.
+				// in this case, to select the triangle, the analyze is made with the dual space of the Delaunay triangulation: the Voronoi diagram.
+				// https://docs.google.com/presentation/d/1Hr5jvLH8tm4KqXgmT-Di-2kFZYOT7MB246hhPKUk6Cw/edit?usp=sharing
+
+				// Boundary vertex of selected triangles
+				TArray<int32> BoundaryVertex;
+				BoundaryVertex.Reserve(SelectedTriangleIndices.Num() + AdditionalTriangleIndices.Num() + 3);
+				for (int32 TriangleIndex : SelectedTriangleIndices)
 				{
 					for (int32 Index = 0; Index < 3; Index++)
 					{
 						const int32 StartVertex = TriangleSet[TriangleIndex].VertexIndices[Index];
-						TmpVertexIndices.AddUnique(StartVertex);
+						BoundaryVertex.AddUnique(StartVertex);
 					}
 				}
 
-				TArray<TPair<int32, double>> VertexToSlop;
-				for (int32 Index : TmpVertexIndices)
+				TArray<TPair<int32, double>> BoundaryVertexToSlop;
+				for (int32 Index : BoundaryVertex)
 				{
 					FPoint2D Vertex = Vertices[Index].Value;
 					double Slop = ComputeSlope(NewVertex, Vertex);
-					VertexToSlop.Emplace(Index, Slop);
+					BoundaryVertexToSlop.Emplace(Index, Slop);
 				}
 
-				VertexToSlop.Sort([](const TPair<int32, double>& Vertex1, const TPair<int32, double>& Vertex2) {return Vertex1.Value < Vertex2.Value; });
+				// Sort the vertex to make the boundary polygon
+				BoundaryVertexToSlop.Sort([](const TPair<int32, double>& Vertex1, const TPair<int32, double>& Vertex2) {return Vertex1.Value < Vertex2.Value; });
 
-				//Wait();
 				bool bTriangleHasBeenAdded = true;
 				while (bTriangleHasBeenAdded)
 				{
@@ -175,144 +255,156 @@ public:
 							continue;
 						}
 
+						FTriangle& CandidateTriangle = TriangleSet[TriangleIndex];
+
+						// if the CandidateTriangle is connected by an edge to the boundary polygon
 						int32 CandidateVertexIndex = -1;
+						int32 ConnectedVertex = 0;
 						for (int32 Index = 0; Index < 3; Index++)
 						{
-							int32 Candidate = TriangleSet[TriangleIndex].VertexIndices[Index];
-							if (TmpVertexIndices.Find(Candidate) == INDEX_NONE)
+							int32 Candidate = CandidateTriangle.VertexIndices[Index];
+							if (BoundaryVertex.Find(Candidate) == INDEX_NONE)
+							{
+								CandidateVertexIndex = Candidate;
+							}
+							else
 							{
-								if (CandidateVertexIndex == -1)
-								{
-									CandidateVertexIndex = Candidate;
-								}
-								else
-								{
-									CandidateVertexIndex = -1;
-									break;
-								}
+								ConnectedVertex++;
 							}
 						}
-
-						if (CandidateVertexIndex == -1)
+						if (ConnectedVertex != 2)
 						{
 							continue;
 						}
 
-						FPoint2D CandidateVertexPoint = Vertices[CandidateVertexIndex].Value;
+						FPoint2D CandidateVertex = Vertices[CandidateVertexIndex].Value;
 
-						int32 StartIndex = -1;
-						int32 EndIndex = 0;
+#ifdef DEBUG_BOWYERWATSON_STEP_2
+						if (VIndex == StopIndex)
 						{
-							double Slope = ComputeSlope(NewVertex, CandidateVertexPoint);
-							for (int32 AIndex = 0; AIndex < VertexToSlop.Num(); ++AIndex)
+							DisplaySelectedTrianglesBoundary(BoundaryVertexToSlop);
 							{
-								constexpr double SmallSlop = 0.01; // 0.5 deg
-								if (FMath::IsNearlyEqual(Slope, VertexToSlop[AIndex].Value, SmallSlop))
-								{
-									EndIndex = -1;
-									break;
-								}
-
-								if (Slope < VertexToSlop[AIndex].Value)
-								{
-									EndIndex = AIndex;
-									break;
-								}
+								F3DDebugSession A(TEXT("New Sel triangle"));
+								DisplayTriangle(TriangleIndex, EVisuProperty::YellowCurve);
 							}
-							if (EndIndex < 0)
 							{
-								continue;
+								F3DDebugSession A(TEXT("Candidate Point"));
+								DisplayPoint(Vertices[CandidateVertexIndex].Value * DisplayScale, EVisuProperty::BluePoint);
 							}
-							StartIndex = EndIndex == 0 ? VertexToSlop.Num() - 1 : EndIndex - 1;
+							Wait();
 						}
+#endif
 
-						FPoint2D StartPoint = Vertices[VertexToSlop[StartIndex].Key].Value;
-						FPoint2D Vect = StartPoint - NewVertex;
-						FPoint2D Perp(Vect.V, -Vect.U);
-						FPoint2D StartPoint2 = StartPoint - Perp;
-						TSegment<FPoint2D> Segment1(StartPoint, StartPoint2);
-
-						FPoint2D EndPoint = Vertices[VertexToSlop[EndIndex].Key].Value;
-						FPoint2D Vect2 = EndPoint - NewVertex;
-						FPoint2D Perp2(Vect2.V, -Vect2.U);
-						FPoint2D EndPoint2 = EndPoint + Perp2;
-						TSegment<FPoint2D> Segment2(EndPoint, EndPoint2);
-
-						double SlopCandidate = 4;
-						FPoint2D Intersection; 
-						if(FindIntersectionOfLines2D(Segment1, Segment2, Intersection))
+						double NewVertexSlop = ComputeSlope(NewVertex, CandidateVertex);
+
+						int32 StartIndex = -1;
+						int32 EndIndex = 0;
+						for (int32 ApexIndex = 0; ApexIndex < BoundaryVertexToSlop.Num(); ++ApexIndex)
 						{
-							SlopCandidate = ComputePositiveSlope(CandidateVertexPoint, NewVertex, Intersection);
+							if (NewVertexSlop < BoundaryVertexToSlop[ApexIndex].Value)
+							{
+								EndIndex = ApexIndex;
+								break;
+							}
 						}
 
-						if (SlopCandidate > 2.00 && SlopCandidate < 6.00)
+						StartIndex = EndIndex == 0 ? BoundaryVertexToSlop.Num() - 1 : EndIndex - 1;
+
+						FPoint2D StartPoint = Vertices[BoundaryVertexToSlop[StartIndex].Key].Value;
+						double SlopeNewStartPlusPi = ComputeSlope(NewVertex, StartPoint) + Slope::RightSlope;
+						double SlopeStartCandidate = ComputePositiveSlope(StartPoint, CandidateVertex, SlopeNewStartPlusPi);
+
+						FPoint2D EndPoint = Vertices[BoundaryVertexToSlop[EndIndex].Key].Value;
+						double SlopeNewEndMinusPi = ComputeSlope(NewVertex, EndPoint) - Slope::RightSlope;
+						double SlopeEndCandidate = ComputePositiveSlope(EndPoint, CandidateVertex, SlopeNewEndMinusPi);
+
+#ifdef DEBUG_BOWYERWATSON_STEP_2
+						if (VIndex == StopIndex)
 						{
-							bTriangleHasBeenAdded = true;
+							DisplayLocalVoronoiDiagram(CandidateVertex, NewVertex, StartPoint, EndPoint);
+							Wait();
+						}
+#endif
 
-							double SlopeAtNewVertex = ComputeSlope(NewVertex, CandidateVertexPoint);
-							if (EndIndex == 0 && SlopeAtNewVertex > VertexToSlop.Last().Value)
+						if (SlopeStartCandidate > 0 && SlopeStartCandidate < Slope::RightSlope && SlopeEndCandidate > Slope::ThreeRightSlope && SlopeEndCandidate < Slope::FullSlope)
+						{
+							bTriangleHasBeenAdded = true;
+							if (EndIndex == 0 && NewVertexSlop > BoundaryVertexToSlop.Last().Value)
 							{
-								VertexToSlop.Emplace(CandidateVertexIndex, SlopeAtNewVertex);
+								BoundaryVertexToSlop.Emplace(CandidateVertexIndex, NewVertexSlop);
 							}
 							else
 							{
-								VertexToSlop.EmplaceAt(EndIndex, CandidateVertexIndex, SlopeAtNewVertex);
+								BoundaryVertexToSlop.EmplaceAt(EndIndex, CandidateVertexIndex, NewVertexSlop);
+							}
+							BoundaryVertex.Add(CandidateVertexIndex);
+							SelectedTriangleIndices.Add(TriangleIndex);
+
+#ifdef DEBUG_BOWYERWATSON_STEP_2
+							if (VIndex == StopIndex)
+							{
+								DisplaySelectedTriangles();
+								Wait();
 							}
-							TmpVertexIndices.Add(CandidateVertexIndex);
-							TriangleIndices.Add(TriangleIndex);
+#endif
 						}
 						TriangleIndex = -1;
 					}
 				}
 			}
 
-			// Find the boundary edges of the selected triangles:
-			// For all selected triangles, 
-			//    For each triangle edges
-			//       if the edge is not in EdgeVertexIndices: Add the edge i.e. add its vertex indices
-			//       else (the edge is in EdgeVertexIndices), remove the edge of EdgeVertexIndices
-			// As the triangles are oriented, the edge AB of a triangle is the edge BA of the adjacent triangle
-			EdgeVertexIndices.Empty(VerticesCount);
-			for (int32 TriangleIndex : TriangleIndices)
+#ifdef DEBUG_BOWYERWATSON_STEP
+			if (VIndex >= StopIndex)
 			{
-				int32 EndVertex = TriangleSet[TriangleIndex].VertexIndices[2];
-				for (int32 Index = 0; Index < 3; Index++)
+				DisplaySelectedTriangles();
+				Wait();
+			}
+#endif
+
+			EdgeVertexIndices.Reset(VerticesCount);
+ 			FindBoundaryEdgesOfSelectedTriangles(EdgeVertexIndices);
+
+			// make the new triangles : Each new triangle is defined by an edge of the boundary and the new vertex 
+#ifdef DEBUG_BOWYERWATSON_STEP
+			if (bDisplay && VIndex >= StopIndex)
+			{
+				F3DDebugSession _(TEXT("To remesh"));
+				for (int32 EdgeIndex = 0; EdgeIndex < EdgeVertexIndices.Num(); EdgeIndex += 2)
 				{
-					int32 StartVertex = TriangleSet[TriangleIndex].VertexIndices[Index];
-					int32 Endex = 0;
-					// Does the edge exist
-					for (; Endex < EdgeVertexIndices.Num(); Endex += 2)
+					if (EdgeVertexIndices[EdgeIndex] < 0)
 					{
-						if (EdgeVertexIndices[Endex] == EndVertex && EdgeVertexIndices[Endex + 1] == StartVertex)
-						{
-							EdgeVertexIndices[Endex] = -1;
-							EdgeVertexIndices[Endex + 1] = -1;
-							break;
-						}
-					}
-					if (Endex == EdgeVertexIndices.Num())
-					{   // No
-						EdgeVertexIndices.Add(StartVertex);
-						EdgeVertexIndices.Add(EndVertex);
+						continue;
 					}
-					EndVertex = StartVertex;
+					DisplaySegment(Vertices[EdgeVertexIndices[EdgeIndex]].Value * DisplayScale, Vertices[EdgeVertexIndices[EdgeIndex + 1]].Value * DisplayScale, 0, EVisuProperty::YellowCurve);
 				}
+				Wait();
 			}
-
-			// make the new triangles : Each new triangle is defined by an edge of the boundary and the new vertex 
+#endif
 			{
+#ifdef DEBUG_BOWYERWATSON_STEP
+				F3DDebugSession _(bDisplay && VIndex >= StopIndex, TEXT("New Triangles"));
+#endif
+
 				// The deleted triangles are replaced by the new ones
 				int32 EdgeIndex = 0;
-				for (int32 TriangleIndex : TriangleIndices)
+				for (int32 TriangleIndex : SelectedTriangleIndices)
 				{
 					while (EdgeVertexIndices[EdgeIndex] < 0)
 					{
 						EdgeIndex += 2;
 					}
 					TriangleSet[TriangleIndex].Set(EdgeVertexIndices[EdgeIndex + 1], EdgeVertexIndices[EdgeIndex], VertexIndex, Vertices);
+#ifdef DEBUG_BOWYERWATSON_STEP
+					if (bDisplay && VIndex >= StopIndex)
+					{
+						//F3DDebugSession A(TEXT("Triangle"));
+						DisplayTriangle(TriangleIndex, EVisuProperty::BlueCurve);
+						//Wait();
+					}
+#endif
+
 					EdgeIndex += 2;
 				}
-
 				// When all deleted triangles are replaced, the new triangles are added in the array
 				for (; EdgeIndex < EdgeVertexIndices.Num(); EdgeIndex += 2)
 				{
@@ -321,12 +413,37 @@ public:
 						continue;
 					}
 					TriangleSet.Emplace(EdgeVertexIndices[EdgeIndex + 1], EdgeVertexIndices[EdgeIndex], VertexIndex, Vertices);
+#ifdef DEBUG_BOWYERWATSON_STEP
+					if (bDisplay && VIndex >= StopIndex)
+					{
+						//F3DDebugSession A(TEXT("Triangle"));
+						DisplayTriangle(TriangleSet.Num() - 1, EVisuProperty::BlueCurve);
+						//Wait();
+					}
+#endif
 				}
 			}
+
+#ifdef DEBUG_BOWYERWATSON_STEP
+			if (bDisplay)
+			{
+				DisplayTriangles();
+				Wait(VIndex >= StopIndex);
+			}
+#endif
 		}
 
+#ifdef DEBUG_BOWYERWATSON
+		if (bDisplay)
+		{
+			// The final mesh
+			DisplayTriangles();
+			Wait(bDisplay);
+		}
+#endif
+
 		// Find all Edges and their type (inner edge or boundary edge)
-		EdgeVertexIndices.Empty(TriangleSet.Num() * 6);
+		EdgeVertexIndices.Reset(TriangleSet.Num() * 6);
 
 		for (int32 TriangleIndex = 0; TriangleIndex < TriangleSet.Num(); TriangleIndex++)
 		{
@@ -373,6 +490,10 @@ public:
 				EndVertex = StartVertex;
 			}
 		}
+#ifdef DEBUG_BOWYERWATSON
+		DisplayEdges();
+		//Wait();
+#endif
 
 		// the bounding mesh vertices are removed
 		Vertices.SetNum(VerticesCount);
@@ -437,7 +558,7 @@ public:
 
 	void GetMesh(TArray<int32>& Triangles)
 	{
-		Triangles.Empty(TriangleSet.Num() * 3);
+		Triangles.Reset(TriangleSet.Num() * 3);
 		for (const FTriangle& Triangle : TriangleSet)
 		{
 			Triangles.Append(Triangle.VertexIndices, 3);
@@ -462,7 +583,7 @@ private:
 	TArray<FTriangle> TriangleSet;
 
 	// It's use to mark all triangles whose circumcircles contain the next vertex
-	TArray<int32> TriangleIndices;
+	TArray<int32> SelectedTriangleIndices;
 	TArray<int32> AdditionalTriangleIndices;
 
 	/**
@@ -497,12 +618,188 @@ private:
 		VerticesCount = Vertices.Num();
 		Vertices.Reserve(VerticesCount + 4);
 		TriangleSet.Reserve(VerticesCount);
-		TriangleIndices.Reserve(VerticesCount);
+		SelectedTriangleIndices.Reserve(VerticesCount);
 		AdditionalTriangleIndices.Reserve(VerticesCount);
 		EdgeVertexIndices.Reserve(4 * VerticesCount);
 		EdgeInstanceCount.Reserve(2 * VerticesCount);
 	}
 
+	// Find the boundary edges of the selected triangles:
+	// For all selected triangles, 
+	//    For each triangle edges
+	//       if the edge is not in EdgeVertexIndices: Add the edge i.e. add its vertex indices
+	//       else (the edge is in EdgeVertexIndices), remove the edge of EdgeVertexIndices
+	// As the triangles are oriented, the edge AB of a triangle is the edge BA of the adjacent triangle
+	void FindBoundaryEdgesOfSelectedTriangles(TArray<int32>& BoundaryEdgeVertexIndices)
+	{
+		BoundaryEdgeVertexIndices.Reset(SelectedTriangleIndices.Num() * 6);
+		for (int32 TriangleIndex : SelectedTriangleIndices)
+		{
+			int32 EndVertex = TriangleSet[TriangleIndex].VertexIndices[2];
+			for (int32 Index = 0; Index < 3; Index++)
+			{
+				int32 StartVertex = TriangleSet[TriangleIndex].VertexIndices[Index];
+				int32 Endex = 0;
+				// Does the edge exist
+				for (; Endex < BoundaryEdgeVertexIndices.Num(); Endex += 2)
+				{
+					if (BoundaryEdgeVertexIndices[Endex] == EndVertex && BoundaryEdgeVertexIndices[Endex + 1] == StartVertex)
+					{
+						BoundaryEdgeVertexIndices[Endex] = -1;
+						BoundaryEdgeVertexIndices[Endex + 1] = -1;
+						break;
+					}
+				}
+				if (Endex == BoundaryEdgeVertexIndices.Num())
+				{   // No
+					BoundaryEdgeVertexIndices.Add(StartVertex);
+					BoundaryEdgeVertexIndices.Add(EndVertex);
+				}
+				EndVertex = StartVertex;
+			}
+		}
+	}
+
+	bool DoSelectedTrianglesFormSinglePartition()
+	{
+		FindBoundaryEdgesOfSelectedTriangles(EdgeVertexIndices);
+
+		int32 StartIndex = -1;
+		int32 LastIndex = -1;
+		for (int32 EdgeIndex = 0; EdgeIndex < EdgeVertexIndices.Num(); EdgeIndex += 2)
+		{
+			if (EdgeVertexIndices[EdgeIndex] < 0)
+			{
+				continue;
+			}
+			StartIndex = EdgeVertexIndices[EdgeIndex];
+			LastIndex = EdgeVertexIndices[EdgeIndex + 1];
+			EdgeVertexIndices[EdgeIndex] = -1;
+			EdgeVertexIndices[EdgeIndex + 1] = -1;
+			break;
+		}
+
+		const int32 IndiceCount = EdgeVertexIndices.Num();
+		while (LastIndex != StartIndex)
+		{
+			int32 EdgeIndex = 0;
+			for (; EdgeIndex < IndiceCount; EdgeIndex += 2)
+			{
+				if (EdgeVertexIndices[EdgeIndex] == LastIndex)
+				{
+					LastIndex = EdgeVertexIndices[EdgeIndex + 1];
+					EdgeVertexIndices[EdgeIndex] = -1;
+					EdgeVertexIndices[EdgeIndex + 1] = -1;
+				}
+			}
+			if (EdgeIndex == IndiceCount)
+			{
+				return false;
+			}
+		}
+
+		// All edges are selected ?
+		bool bAllEdgesAreSelected = true;
+		for (int32 EdgeIndex = 0; EdgeIndex < EdgeVertexIndices.Num(); EdgeIndex += 2)
+		{
+			if (EdgeVertexIndices[EdgeIndex] > 0)
+			{
+				return false;
+			}
+		}
+
+		return true;
+	}
+
+	bool DoesTriangleContainingNewVertex(const FPoint2D& NewVertex, int32 CandiadateTriangle)
+	{
+		const FPoint2D& Point0 = Vertices[TriangleSet[CandiadateTriangle].VertexIndices[0]].Value;
+		const FPoint2D& Point1 = Vertices[TriangleSet[CandiadateTriangle].VertexIndices[1]].Value;
+		const FPoint2D& Point2 = Vertices[TriangleSet[CandiadateTriangle].VertexIndices[2]].Value;
+		double Slope0 = ComputePositiveSlope(Point0, Point1, NewVertex);
+		double Slope1 = ComputePositiveSlope(Point1, Point2, NewVertex);
+		double Slope2 = ComputePositiveSlope(Point2, Point0, NewVertex);
+		if (Slope0 < Slope::PiSlope && Slope1 < Slope::PiSlope && Slope2 < Slope::PiSlope)
+		{
+			return true;
+		}
+		return false;
+	}
+
+	void AddTrianglesContainingNewVertexToSelectedTriangles(const FPoint2D& NewVertex, TArray<int32>& CandiadateTriangles)
+	{
+		for (int32& TriangleIndex : CandiadateTriangles)
+		{
+			if(DoesTriangleContainingNewVertex(NewVertex, TriangleIndex))
+			{
+				SelectedTriangleIndices.Add(TriangleIndex);
+				TriangleIndex = -1;
+			}
+		}
+	}
+
+	void AddConnectedTrianglesToSelection(TArray<int32>& CandiadateTriangles)
+	{
+		EdgeVertexIndices.Reset((SelectedTriangleIndices.Num() + CandiadateTriangles.Num()) * 6);
+		FindBoundaryEdgesOfSelectedTriangles(EdgeVertexIndices);
+
+		bool bTriangleHasBeenAdded = true;
+		while (bTriangleHasBeenAdded)
+		{
+			bTriangleHasBeenAdded = false;
+			for (int32& TriangleIndex : CandiadateTriangles)
+			{
+				if (TriangleIndex < 0)
+				{
+					continue;
+				}
+
+				int32 EndVertex = TriangleSet[TriangleIndex].VertexIndices[2];
+
+				bool bEdgeIsFound = false;
+				int32 Index = 0;
+				for (; Index < 3; Index++)
+				{
+					int32 StartVertex = TriangleSet[TriangleIndex].VertexIndices[Index];
+
+					for (int32 Endex = 0; Endex < EdgeVertexIndices.Num(); Endex += 2)
+					{
+						if (EdgeVertexIndices[Endex] == EndVertex && EdgeVertexIndices[Endex + 1] == StartVertex)
+						{
+							EdgeVertexIndices[Endex] = -1;
+							EdgeVertexIndices[Endex + 1] = -1;
+							bEdgeIsFound = true;
+							break;
+						}
+					}
+					if (bEdgeIsFound)
+					{
+						break;
+					}
+					EndVertex = StartVertex;
+				}
+
+				if (bEdgeIsFound)
+				{
+					SelectedTriangleIndices.Add(TriangleIndex);
+					EndVertex = TriangleSet[TriangleIndex].VertexIndices[2];
+					for (int32 Endex = 0; Endex < 3; Endex++)
+					{
+						int32 StartVertex = TriangleSet[TriangleIndex].VertexIndices[Endex];
+						if (Endex != Index)
+						{
+							EdgeVertexIndices.Add(StartVertex);
+							EdgeVertexIndices.Add(EndVertex);
+						}
+						EndVertex = StartVertex;
+					}
+					bTriangleHasBeenAdded = true;
+					TriangleIndex = -1;
+				}
+			}
+		}
+	}
+
 #ifdef DEBUG_BOWYERWATSON
 	void DisplayEdges()
 	{
@@ -548,14 +845,37 @@ private:
 		}
 
 		F3DDebugSession _(TEXT("Selected Triangles"));
-		for (int32 Index : TriangleIndices)
+		for (int32 Index : SelectedTriangleIndices)
 		{
-			F3DDebugSession _(TEXT("Triangle"));
+			//F3DDebugSession _(TEXT("Triangle"));
 			DisplayTriangle(Index, EVisuProperty::BlueCurve);
 		}
 		//Wait();
 	};
 
+	void DisplaySelectedTrianglesAndComplements()
+	{
+		if (!bDisplay)
+		{
+			return;
+		}
+
+		DisplaySelectedTriangles(); 
+
+		{
+			F3DDebugSession A(bDisplay, TEXT("Select Triangles Add"));
+			for (int32 TriangleIndex : AdditionalTriangleIndices)
+			{
+				if (TriangleIndex < 0)
+				{
+					continue;
+				}
+				DisplayTriangle(TriangleIndex, EVisuProperty::YellowCurve);
+			}
+		}
+	}
+
+
 	void DisplayTriangle(int32 Index, EVisuProperty Property)
 	{
 		if (!bDisplay)
@@ -569,6 +889,70 @@ private:
 		//Wait();
 	};
 
+	void DisplaySelectedTrianglesBoundary(TArray<TPair<int32, double>> VertexToSlop)
+	{
+		F3DDebugSession A(TEXT("Selected Triangles Boundary"));
+
+		FPoint2D Previous = Vertices[VertexToSlop[VertexToSlop.Num() - 1].Key].Value;
+		for (int Index = 0; Index < VertexToSlop.Num(); ++Index)
+		{
+			FPoint2D VertexPoint = Vertices[VertexToSlop[Index].Key].Value;
+			DisplayPoint(VertexPoint * DisplayScale, EVisuProperty::GreenPoint);
+			DisplaySegment(VertexPoint * DisplayScale, Previous * DisplayScale, 0, EVisuProperty::GreenCurve);
+			Previous = VertexPoint;
+		}
+	}
+
+	void DisplayLocalVoronoiDiagram(const FPoint2D& CandidateVertex, const FPoint2D& NewVertex, const FPoint2D& StartPoint, const FPoint2D& EndPoint)
+	{
+		F3DDebugSession B(TEXT("Voronoi Diagram"));
+		{
+			F3DDebugSession A(TEXT("New	Vertex"));
+			DisplayPoint(NewVertex * DisplayScale, EVisuProperty::RedPoint);
+		}
+		{
+			F3DDebugSession A(TEXT("Candidate"));
+			DisplayPoint(CandidateVertex * DisplayScale, EVisuProperty::BluePoint);
+		}
+		{
+			F3DDebugSession A(TEXT("Start"));
+			DisplayPoint(StartPoint * DisplayScale, EVisuProperty::YellowPoint);
+
+			FPoint2D Normal = NewVertex - StartPoint;
+			constexpr double HalfPi = -DOUBLE_PI / 2.;
+			Normal.Normalize();
+			FPoint2D Tangent = Normal.Rotate(HalfPi) * 1000.;
+			FPoint2D TangentPoint = StartPoint + Tangent;
+
+			DisplaySegment(StartPoint * DisplayScale, TangentPoint * DisplayScale, 1, EVisuProperty::YellowCurve);
+			DisplaySegment(NewVertex * DisplayScale, StartPoint * DisplayScale, 1, EVisuProperty::YellowCurve);
+		}
+
+		{
+			F3DDebugSession A(TEXT("Start - Candidate"));
+			DisplaySegment(StartPoint * DisplayScale, CandidateVertex * DisplayScale, 1, EVisuProperty::YellowCurve);
+		}
+
+		{
+			F3DDebugSession A(TEXT("End"));
+			DisplayPoint(EndPoint * DisplayScale, EVisuProperty::OrangePoint);
+
+			FPoint2D Normal = NewVertex - EndPoint;
+			constexpr double HalfPi = DOUBLE_PI / 2.;
+			Normal.Normalize();
+			FPoint2D Tangent = Normal.Rotate(HalfPi) * 1000.;
+			FPoint2D TangentPoint = EndPoint + Tangent;
+
+			DisplaySegment(EndPoint * DisplayScale, TangentPoint * DisplayScale, 1, EVisuProperty::OrangeCurve);
+			DisplaySegment(NewVertex * DisplayScale, EndPoint * DisplayScale, 1, EVisuProperty::OrangeCurve);
+		}
+
+		{
+			F3DDebugSession A(TEXT("End - Candidate"));
+			DisplaySegment(EndPoint * DisplayScale, CandidateVertex * DisplayScale, 1, EVisuProperty::OrangeCurve);
+		}
+	}
+
 #endif
 
 };
