Rename CADKernel namespace into UE::CADKernel

diff --git a/Engine/Source/Runtime/Datasmith/CADKernel/Private/CADKernel/Geo/Surfaces/CompositeSurface.h b/Engine/Source/Runtime/Datasmith/CADKernel/Private/CADKernel/Geo/Surfaces/CompositeSurface.h
--- a/Engine/Source/Runtime/Datasmith/CADKernel/Private/CADKernel/Geo/Surfaces/CompositeSurface.h
+++ b/Engine/Source/Runtime/Datasmith/CADKernel/Private/CADKernel/Geo/Surfaces/CompositeSurface.h
@@ -4,88 +4,88 @@
 #include "CADKernel/Core/CADKernelArchive.h"
 #include "CADKernel/Geo/Surfaces/Surface.h"
 
-namespace CADKernel
+namespace UE::CADKernel
 {
 
-	class CADKERNEL_API FCompositeSurface : public FSurface
-	{
-		friend FEntity;
+class CADKERNEL_API FCompositeSurface : public FSurface
+{
+	friend FEntity;
 
-	protected:
+protected:
 
-		/**
-		 * U & V Coordinates of the composite surface's parametric space at each change of elementary surface
-		 */
-		TArray<double> GlobalCoordinates[2];
+	/**
+	 * U & V Coordinates of the composite surface's parametric space at each change of elementary surface
+	 */
+	TArray<double> GlobalCoordinates[2];
 
-		/**
-		 * Elementary surfaces
-		 */
-		TArray<TSharedPtr<FSurface>> Surfaces;
+	/**
+	 * Elementary surfaces
+	 */
+	TArray<TSharedPtr<FSurface>> Surfaces;
 
-		/**
-		 * Native UV Loops of each surface composing the composite surface
-		 */
-		TArray<FSurfacicBoundary> NativeUVBoundaries;
+	/**
+	 * Native UV Loops of each surface composing the composite surface
+	 */
+	TArray<FSurfacicBoundary> NativeUVBoundaries;
 
-		FCompositeSurface(const double InToleranceGeometric, int32 USurfaceNum, int32 VSurfaceNum, const TArray<double>& UCoordinates, const TArray<double>& VCoordinates, const TArray<TSharedPtr<FSurface>>& Surfaces);
+	FCompositeSurface(const double InToleranceGeometric, int32 USurfaceNum, int32 VSurfaceNum, const TArray<double>& UCoordinates, const TArray<double>& VCoordinates, const TArray<TSharedPtr<FSurface>>& Surfaces);
 
-		FCompositeSurface() = default;
+	FCompositeSurface() = default;
 
-	public:
+public:
 
-		virtual void Serialize(FCADKernelArchive& Ar) override
+	virtual void Serialize(FCADKernelArchive& Ar) override
+	{
+		FSurface::Serialize(Ar);
+		Ar << GlobalCoordinates[EIso::IsoU];
+		Ar << GlobalCoordinates[EIso::IsoV];
+		SerializeIdents(Ar, Surfaces);
+		Ar << NativeUVBoundaries;
+	}
+
+	virtual void SpawnIdent(FDatabase& Database) override
+	{
+		if (!FEntity::SetId(Database))
 		{
-			FSurface::Serialize(Ar);
-			Ar << GlobalCoordinates[EIso::IsoU];
-			Ar << GlobalCoordinates[EIso::IsoV];
-			SerializeIdents(Ar, Surfaces);
-			Ar << NativeUVBoundaries;
+			return;
 		}
 
-		virtual void SpawnIdent(FDatabase& Database) override
-		{
-			if (!FEntity::SetId(Database))
-			{
-				return;
-			}
-
-			SpawnIdentOnEntities(Surfaces, Database);
-		}
+		SpawnIdentOnEntities(Surfaces, Database);
+	}
 
-		virtual void ResetMarkersRecursively() override
-		{
-			ResetMarkers();
-			ResetMarkersRecursivelyOnEntities(Surfaces);
-		}
+	virtual void ResetMarkersRecursively() override
+	{
+		ResetMarkers();
+		ResetMarkersRecursivelyOnEntities(Surfaces);
+	}
 
-		ESurface GetSurfaceType() const
-		{
-			return ESurface::Composite;
-		}
+	ESurface GetSurfaceType() const
+	{
+		return ESurface::Composite;
+	}
 
 #ifdef CADKERNEL_DEV
-		virtual FInfoEntity& GetInfo(FInfoEntity&) const override;
+	virtual FInfoEntity& GetInfo(FInfoEntity&) const override;
 #endif
 
-		virtual TSharedPtr<FEntityGeom> ApplyMatrix(const FMatrixH& InMatrix) const override;
+	virtual TSharedPtr<FEntityGeom> ApplyMatrix(const FMatrixH& InMatrix) const override;
 
-		int32 GetSurfNum(EIso Iso) const
-		{
-			return GlobalCoordinates[Iso].Num() - 1;
-		}
+	int32 GetSurfNum(EIso Iso) const
+	{
+		return GlobalCoordinates[Iso].Num() - 1;
+	}
 
-		virtual void InitBoundary() override;
+	virtual void InitBoundary() override;
 
-		virtual void LinesNotDerivables(const FSurfacicBoundary& Bounds, int32 InDerivativeOrder, FCoordinateGrid& OutNotDerivableCoordinates) const override;
+	virtual void LinesNotDerivables(const FSurfacicBoundary& Bounds, int32 InDerivativeOrder, FCoordinateGrid& OutNotDerivableCoordinates) const override;
 
-		virtual void EvaluatePoint(const FPoint2D& InSurfacicCoordinate, FSurfacicPoint& OutPoint3D, int32 InDerivativeOrder = 0) const override;
+	virtual void EvaluatePoint(const FPoint2D& InSurfacicCoordinate, FSurfacicPoint& OutPoint3D, int32 InDerivativeOrder = 0) const override;
 
-		virtual void Presample(const FSurfacicBoundary& InBoundaries, FCoordinateGrid& OutCoordinates) override;
+	virtual void Presample(const FSurfacicBoundary& InBoundaries, FCoordinateGrid& OutCoordinates) override;
 
-		double LocalToGlobalCoordinate(int32 SurfaceIndex, double CoordinateU, EIso Iso) const;
-	};
+	double LocalToGlobalCoordinate(int32 SurfaceIndex, double CoordinateU, EIso Iso) const;
+};
 
-} // namespace CADKernel
+} // namespace UE::CADKernel
 
 
diff --git a/Engine/Source/Runtime/Datasmith/CADKernel/Private/CADKernel/Geo/Surfaces/ConeSurface.h b/Engine/Source/Runtime/Datasmith/CADKernel/Private/CADKernel/Geo/Surfaces/ConeSurface.h
--- a/Engine/Source/Runtime/Datasmith/CADKernel/Private/CADKernel/Geo/Surfaces/ConeSurface.h
+++ b/Engine/Source/Runtime/Datasmith/CADKernel/Private/CADKernel/Geo/Surfaces/ConeSurface.h
@@ -3,116 +3,116 @@
 
 #include "CADKernel/Geo/Surfaces/Surface.h"
 
-namespace CADKernel
+namespace UE::CADKernel
 {
-	class CADKERNEL_API FConeSurface : public FSurface
+class CADKERNEL_API FConeSurface : public FSurface
+{
+	friend FEntity;
+
+protected:
+	FMatrixH Matrix;
+	double   StartRadius;
+	double   ConeAngle;
+
+	/**
+	 * A Cone surface is defined by a point on the axis of the cone, the direction of the axis of the cone,
+	 * the radius of the cone at the axis point and the cone semi-angle.
+	 *
+	 * The local coordinate system is defined with the origin at the axis point and the Z axis in the axis direction,
+	 * The equation of the surface in this system is x2 + y2 – (StartRadius + z*tan(ConeAngle) )2 = 0
+	 * where ConeAngle is the cone semi-angle and StartRadius is the given cone radius.
+	 *
+	 * The cone is placed at its final position and orientation by the Matrix
+	 */
+	FConeSurface(const double InToleranceGeometric, const FMatrixH& InMatrix, double InStartRadius, double InConeAngle, double InStartRuleLength = -HUGE_VALUE, double InEndRuleLength = HUGE_VALUE, double InStartAngle = 0.0, double InEndAngle = DOUBLE_TWO_PI)
+		: FConeSurface(InToleranceGeometric, InMatrix, InStartRadius, InConeAngle, FSurfacicBoundary(InStartAngle, InEndAngle, InStartRuleLength, InEndRuleLength))
+	{
+	}
+
+	/**
+	 * A Cone surface is defined by a point on the axis of the cone, the direction of the axis of the cone,
+	 * the radius of the cone at the axis point and the cone semi-angle.
+	 *
+	 * The local coordinate system is defined with the origin at the axis point and the Z axis in the axis direction,
+	 * The equation of the surface in this system is x2 + y2 – (StartRadius + z*tan(ConeAngle) )2 = 0
+	 * where ConeAngle is the cone semi-angle and StartRadius is the given cone radius.
+	 *
+	 * The cone is placed at its final position and orientation by the Matrix
+	 *
+	 * The bounds of the cone are defined as follow:
+	 * Bounds[EIso::IsoU].Min = StartRuleLength;
+	 * Bounds[EIso::IsoU].Max = EndRuleLength;
+	 * Bounds[EIso::IsoV].Min = StartAngle;
+	 * Bounds[EIso::IsoV].Max = EndAngle;
+	 */
+	FConeSurface(const double InToleranceGeometric, const FMatrixH& InMatrix, double InStartRadius, double InConeAngle, const FSurfacicBoundary& InBoundary)
+		: FSurface(InToleranceGeometric, InBoundary)
+		, Matrix(InMatrix)
+		, StartRadius(InStartRadius)
+		, ConeAngle(InConeAngle)
+	{
+		ComputeMinToleranceIso();
+	}
+
+	FConeSurface() = default;
+
+	void ComputeMinToleranceIso()
+	{
+		FPoint Origin = Matrix.Multiply(FPoint::ZeroPoint);
+
+		double DeltaVR = tan(ConeAngle);
+		double Radius1 = FMath::Abs(StartRadius + Boundary[EIso::IsoV].Max * DeltaVR);
+		double Radius2 = FMath::Abs(StartRadius + Boundary[EIso::IsoV].Min * DeltaVR);
+		double Radius = FMath::Max(Radius2, Radius1);
+
+		FPoint Point2DU{ 1 , 0, 0 };
+		double ScaleU = ComputeScaleAlongAxis(Point2DU, Matrix, Origin);
+		double ToleranceU = Tolerance3D / Radius;
+		ToleranceU /= ScaleU;
+
+		FPoint Point2DV{ 0, 1, 0 };
+		double ScaleV = ComputeScaleAlongAxis(Point2DV, Matrix, Origin);
+		double ToleranceV = Tolerance3D / FMath::Sqrt(FMath::Square(DeltaVR) + 1);
+		ToleranceV /= ScaleV;
+
+		MinToleranceIso.Set(ToleranceU, ToleranceV);
+	}
+
+public:
+
+	virtual void Serialize(FCADKernelArchive& Ar) override
 	{
-		friend FEntity;
-
-	protected:
-		FMatrixH Matrix;
-		double   StartRadius;
-		double   ConeAngle;
-
-		/**
-		 * A Cone surface is defined by a point on the axis of the cone, the direction of the axis of the cone, 
-		 * the radius of the cone at the axis point and the cone semi-angle. 
-		 * 
-		 * The local coordinate system is defined with the origin at the axis point and the Z axis in the axis direction, 
-		 * The equation of the surface in this system is x2 + y2 – (StartRadius + z*tan(ConeAngle) )2 = 0
-		 * where ConeAngle is the cone semi-angle and StartRadius is the given cone radius. 
-		 * 
-		 * The cone is placed at its final position and orientation by the Matrix
-		 */
-		FConeSurface(const double InToleranceGeometric, const FMatrixH& InMatrix, double InStartRadius, double InConeAngle, double InStartRuleLength = -HUGE_VALUE, double InEndRuleLength = HUGE_VALUE, double InStartAngle = 0.0, double InEndAngle = DOUBLE_TWO_PI)
-			: FConeSurface(InToleranceGeometric, InMatrix, InStartRadius, InConeAngle, FSurfacicBoundary(InStartAngle, InEndAngle, InStartRuleLength, InEndRuleLength))
-		{
-		}
-
-		/**
-		 * A Cone surface is defined by a point on the axis of the cone, the direction of the axis of the cone,
-		 * the radius of the cone at the axis point and the cone semi-angle.
-		 *
-		 * The local coordinate system is defined with the origin at the axis point and the Z axis in the axis direction,
-		 * The equation of the surface in this system is x2 + y2 – (StartRadius + z*tan(ConeAngle) )2 = 0
-		 * where ConeAngle is the cone semi-angle and StartRadius is the given cone radius.
-		 *
-		 * The cone is placed at its final position and orientation by the Matrix
-		 *
-		 * The bounds of the cone are defined as follow:
-		 * Bounds[EIso::IsoU].Min = StartRuleLength;
-		 * Bounds[EIso::IsoU].Max = EndRuleLength;
-		 * Bounds[EIso::IsoV].Min = StartAngle;
-		 * Bounds[EIso::IsoV].Max = EndAngle;
-		 */
-		FConeSurface(const double InToleranceGeometric, const FMatrixH& InMatrix, double InStartRadius, double InConeAngle, const FSurfacicBoundary& InBoundary)
-			: FSurface(InToleranceGeometric, InBoundary)
-			, Matrix(InMatrix)
-			, StartRadius(InStartRadius)
-			, ConeAngle(InConeAngle)
-		{
-			ComputeMinToleranceIso();
-		}
-
-		FConeSurface() = default;
-
-		void ComputeMinToleranceIso()
-		{
-			FPoint Origin = Matrix.Multiply(FPoint::ZeroPoint);
-
-			double DeltaVR = tan(ConeAngle);
-			double Radius1 = FMath::Abs(StartRadius + Boundary[EIso::IsoV].Max * DeltaVR);
-			double Radius2 = FMath::Abs(StartRadius + Boundary[EIso::IsoV].Min * DeltaVR);
-			double Radius = FMath::Max(Radius2, Radius1);
-
-			FPoint Point2DU{ 1 , 0, 0 };
-			double ScaleU = ComputeScaleAlongAxis(Point2DU, Matrix, Origin);
-			double ToleranceU = Tolerance3D / Radius;
-			ToleranceU /= ScaleU;
-
-			FPoint Point2DV{ 0, 1, 0 };
-			double ScaleV = ComputeScaleAlongAxis(Point2DV, Matrix, Origin);
-			double ToleranceV = Tolerance3D / FMath::Sqrt(FMath::Square(DeltaVR) + 1);
-			ToleranceV /= ScaleV;
-
-			MinToleranceIso.Set(ToleranceU, ToleranceV);
-		}
-
-	public:
-
-		virtual void Serialize(FCADKernelArchive& Ar) override
-		{
-			FSurface::Serialize(Ar);
-			Ar << Matrix;
-			Ar << StartRadius;
-			Ar << ConeAngle;
-		}
-
-
-		ESurface GetSurfaceType() const
-		{
-			return ESurface::Cone;
-		}
+		FSurface::Serialize(Ar);
+		Ar << Matrix;
+		Ar << StartRadius;
+		Ar << ConeAngle;
+	}
+
+
+	ESurface GetSurfaceType() const
+	{
+		return ESurface::Cone;
+	}
 
 #ifdef CADKERNEL_DEV
-		virtual FInfoEntity& GetInfo(FInfoEntity&) const override;
+	virtual FInfoEntity& GetInfo(FInfoEntity&) const override;
 #endif
 
-		virtual TSharedPtr<FEntityGeom> ApplyMatrix(const FMatrixH& InMatrix) const override;
+	virtual TSharedPtr<FEntityGeom> ApplyMatrix(const FMatrixH& InMatrix) const override;
 
-		virtual void EvaluatePoint(const FPoint2D& InSurfacicCoordinate, FSurfacicPoint& OutPoint3D, int32 InDerivativeOrder = 0) const override;
-		virtual void EvaluatePointGrid(const FCoordinateGrid& Coordinates, FSurfacicSampling& OutPoints, bool bComputeNormals = false) const override;
+	virtual void EvaluatePoint(const FPoint2D& InSurfacicCoordinate, FSurfacicPoint& OutPoint3D, int32 InDerivativeOrder = 0) const override;
+	virtual void EvaluatePointGrid(const FCoordinateGrid& Coordinates, FSurfacicSampling& OutPoints, bool bComputeNormals = false) const override;
 
-		virtual void Presample(const FSurfacicBoundary& InBoundaries, FCoordinateGrid& OutCoordinates) override 
-		{
-			PresampleIsoCircle(InBoundaries, OutCoordinates, EIso::IsoU);
+	virtual void Presample(const FSurfacicBoundary& InBoundaries, FCoordinateGrid& OutCoordinates) override
+	{
+		PresampleIsoCircle(InBoundaries, OutCoordinates, EIso::IsoU);
 
-			OutCoordinates[EIso::IsoV].Empty(3);
-			OutCoordinates[EIso::IsoV].Add(InBoundaries[EIso::IsoV].Min);
-			OutCoordinates[EIso::IsoV].Add((InBoundaries[EIso::IsoV].Max + InBoundaries[EIso::IsoV].Min) / 2.0);
-			OutCoordinates[EIso::IsoV].Add(InBoundaries[EIso::IsoV].Max);
-		}
+		OutCoordinates[EIso::IsoV].Empty(3);
+		OutCoordinates[EIso::IsoV].Add(InBoundaries[EIso::IsoV].Min);
+		OutCoordinates[EIso::IsoV].Add((InBoundaries[EIso::IsoV].Max + InBoundaries[EIso::IsoV].Min) / 2.0);
+		OutCoordinates[EIso::IsoV].Add(InBoundaries[EIso::IsoV].Max);
+	}
 
-	};
+};
 }
 
diff --git a/Engine/Source/Runtime/Datasmith/CADKernel/Private/CADKernel/Geo/Surfaces/CoonsSurface.h b/Engine/Source/Runtime/Datasmith/CADKernel/Private/CADKernel/Geo/Surfaces/CoonsSurface.h
--- a/Engine/Source/Runtime/Datasmith/CADKernel/Private/CADKernel/Geo/Surfaces/CoonsSurface.h
+++ b/Engine/Source/Runtime/Datasmith/CADKernel/Private/CADKernel/Geo/Surfaces/CoonsSurface.h
@@ -4,51 +4,51 @@
 #include "CADKernel/Core/CADKernelArchive.h"
 #include "CADKernel/Geo/Surfaces/Surface.h"
 
-namespace CADKernel
+namespace UE::CADKernel
 {
-	class FCurve;
-	class FNURBSSurface;
+class FCurve;
+class FNURBSSurface;
 
-	class CADKERNEL_API FCoonsSurface : public FSurface
-	{
-		friend FEntity;
+class CADKERNEL_API FCoonsSurface : public FSurface
+{
+	friend FEntity;
 
-	protected:
-		TSharedPtr<FCurve> Curves[4];
-		TArray<FPoint> Corners;
+protected:
+	TSharedPtr<FCurve> Curves[4];
+	TArray<FPoint> Corners;
 
-		FCoonsSurface(const double InToleranceGeometric, TSharedPtr<FCurve> InCurve1, TSharedPtr<FCurve> InCurve2, TSharedPtr<FCurve> InCurve3, TSharedPtr<FCurve> InCurve4);
-		FCoonsSurface(const double InToleranceGeometric, TSharedPtr<FCurve> InCurves[4]);
-		FCoonsSurface() = default;
+	FCoonsSurface(const double InToleranceGeometric, TSharedPtr<FCurve> InCurve1, TSharedPtr<FCurve> InCurve2, TSharedPtr<FCurve> InCurve3, TSharedPtr<FCurve> InCurve4);
+	FCoonsSurface(const double InToleranceGeometric, TSharedPtr<FCurve> InCurves[4]);
+	FCoonsSurface() = default;
 
-	public:
+public:
 
-		virtual void Serialize(FCADKernelArchive& Ar) override
-		{
-			FSurface::Serialize(Ar);
-			SerializeIdent(Ar, Curves[0]);
-			SerializeIdent(Ar, Curves[1]);
-			SerializeIdent(Ar, Curves[2]);
-			SerializeIdent(Ar, Curves[3]);
-			Ar << Corners;
-		}
+	virtual void Serialize(FCADKernelArchive& Ar) override
+	{
+		FSurface::Serialize(Ar);
+		SerializeIdent(Ar, Curves[0]);
+		SerializeIdent(Ar, Curves[1]);
+		SerializeIdent(Ar, Curves[2]);
+		SerializeIdent(Ar, Curves[3]);
+		Ar << Corners;
+	}
 
 
 #ifdef CADKERNEL_DEV
-		virtual FInfoEntity& GetInfo(FInfoEntity&) const override;
+	virtual FInfoEntity& GetInfo(FInfoEntity&) const override;
 #endif
 
-		virtual ESurface GetSurfaceType() const override
-		{
-			return ESurface::Coons;
-		}
+	virtual ESurface GetSurfaceType() const override
+	{
+		return ESurface::Coons;
+	}
 
-		virtual TSharedPtr<FEntityGeom> ApplyMatrix(const FMatrixH& InMatrix) const override;
+	virtual TSharedPtr<FEntityGeom> ApplyMatrix(const FMatrixH& InMatrix) const override;
 
-		virtual void EvaluatePoint(const FPoint2D& InSurfacicCoordinate, FSurfacicPoint& OutPoint3D, int32 InDerivativeOrder = 0) const override;
+	virtual void EvaluatePoint(const FPoint2D& InSurfacicCoordinate, FSurfacicPoint& OutPoint3D, int32 InDerivativeOrder = 0) const override;
 
-		virtual void LinesNotDerivables(const FSurfacicBoundary& Bounds, int32 InDerivativeOrder, FCoordinateGrid& OutNotDerivableCoordinates) const override;
-	};
+	virtual void LinesNotDerivables(const FSurfacicBoundary& Bounds, int32 InDerivativeOrder, FCoordinateGrid& OutNotDerivableCoordinates) const override;
+};
 
-} // namespace CADKernel
+} // namespace UE::CADKernel
 
diff --git a/Engine/Source/Runtime/Datasmith/CADKernel/Private/CADKernel/Geo/Surfaces/CylinderSurface.h b/Engine/Source/Runtime/Datasmith/CADKernel/Private/CADKernel/Geo/Surfaces/CylinderSurface.h
--- a/Engine/Source/Runtime/Datasmith/CADKernel/Private/CADKernel/Geo/Surfaces/CylinderSurface.h
+++ b/Engine/Source/Runtime/Datasmith/CADKernel/Private/CADKernel/Geo/Surfaces/CylinderSurface.h
@@ -3,85 +3,85 @@
 
 #include "CADKernel/Geo/Surfaces/Surface.h"
 
-namespace CADKernel
+namespace UE::CADKernel
 {
-	class FNURBSSurface;
+class FNURBSSurface;
 
-	class CADKERNEL_API FCylinderSurface : public FSurface
+class CADKERNEL_API FCylinderSurface : public FSurface
+{
+	friend FEntity;
+
+protected:
+	FMatrixH Matrix;
+	double Radius;
+
+	/**
+	 * The Cylinder surface is a right circular cylinder surface starting at the axis origin and with z axis as the axis direction.
+	 * It is defined by its Radius and its bounds (StartAngle, EndAngle, StartLength, EndLength)
+	 *
+	 * The surface is placed at its final position and orientation by the Matrix
+	 *
+	 * The bounds of the cylinder are defined as follow:
+	 * Bounds[EIso::IsoU].Min = StartAngle;
+	 * Bounds[EIso::IsoU].Max = EndAngle;
+	 * Bounds[EIso::IsoV].Min = StartLength;
+	 * Bounds[EIso::IsoV].Max = EndLength;
+	 */
+	FCylinderSurface(const double InToleranceGeometric, const FMatrixH& InMatrix, double InRadius, double InStartLength = -HUGE_VALUE, double InEndLength = HUGE_VALUE, double InStartAngle = 0.0, double InEndAngle = DOUBLE_TWO_PI);
+	FCylinderSurface(const double InToleranceGeometric, const FMatrixH& InMatrix, const double InRadius, const FSurfacicBoundary& Boundary);
+
+	FCylinderSurface() = default;
+
+	void ComputeMinToleranceIso()
 	{
-		friend FEntity;
-
-	protected:
-		FMatrixH Matrix;
-		double Radius;
-
-		/**
-		 * The Cylinder surface is a right circular cylinder surface starting at the axis origin and with z axis as the axis direction. 
-		 * It is defined by its Radius and its bounds (StartAngle, EndAngle, StartLength, EndLength)
-		 * 
-		 * The surface is placed at its final position and orientation by the Matrix
-		 *
-		 * The bounds of the cylinder are defined as follow:
-		 * Bounds[EIso::IsoU].Min = StartAngle;
-		 * Bounds[EIso::IsoU].Max = EndAngle;
-		 * Bounds[EIso::IsoV].Min = StartLength;
-		 * Bounds[EIso::IsoV].Max = EndLength;
-		 */
-		FCylinderSurface(const double InToleranceGeometric, const FMatrixH& InMatrix, double InRadius, double InStartLength = -HUGE_VALUE, double InEndLength = HUGE_VALUE, double InStartAngle = 0.0, double InEndAngle = DOUBLE_TWO_PI);
-		FCylinderSurface(const double InToleranceGeometric, const FMatrixH& InMatrix, const double InRadius, const FSurfacicBoundary& Boundary);
-
-		FCylinderSurface() = default;
+		FPoint Origin = Matrix.Multiply(FPoint::ZeroPoint);
 
-		void ComputeMinToleranceIso()
-		{
-			FPoint Origin = Matrix.Multiply(FPoint::ZeroPoint);
+		FPoint Point2DU{ 1 , 0, 0 };
+		double ToleranceU = Tolerance3D / Radius;
 
-			FPoint Point2DU { 1 , 0, 0 };
-			double ToleranceU = Tolerance3D / Radius;
+		double ScaleU = ComputeScaleAlongAxis(Point2DU, Matrix, Origin);
+		ToleranceU /= ScaleU;
 
-			double ScaleU = ComputeScaleAlongAxis(Point2DU, Matrix, Origin);
-			ToleranceU /= ScaleU;
+		FPoint Point2DV{ 0, 1, 0 };
+		double ToleranceV = Tolerance3D / ComputeScaleAlongAxis(Point2DV, Matrix, Origin);
 
-			FPoint Point2DV{ 0, 1, 0 };
-			double ToleranceV = Tolerance3D / ComputeScaleAlongAxis(Point2DV, Matrix, Origin);
+		MinToleranceIso.Set(ToleranceU, ToleranceV);
+	}
 
-			MinToleranceIso.Set(ToleranceU, ToleranceV);
-		}
+public:
 
-	public:
-
-		virtual void Serialize(FCADKernelArchive& Ar) override
-		{
-			FSurface::Serialize(Ar);
-			Ar << Matrix;
-			Ar << Radius;
-		}
+	virtual void Serialize(FCADKernelArchive& Ar) override
+	{
+		FSurface::Serialize(Ar);
+		Ar << Matrix;
+		Ar << Radius;
+	}
 
-		virtual TSharedPtr<FEntityGeom> ApplyMatrix(const FMatrixH& InMatrix) const override;
+	virtual TSharedPtr<FEntityGeom> ApplyMatrix(const FMatrixH& InMatrix) const override;
 
 #ifdef CADKERNEL_DEV
-		virtual FInfoEntity& GetInfo(FInfoEntity&) const override;
+	virtual FInfoEntity& GetInfo(FInfoEntity&) const override;
 #endif
 
-		ESurface GetSurfaceType() const
-		{
-			return ESurface::Cylinder;
-		}
+	ESurface GetSurfaceType() const
+	{
+		return ESurface::Cylinder;
+	}
 
-		virtual void InitBoundary() override;
+	virtual void InitBoundary() override;
 
-		virtual void EvaluatePoint(const FPoint2D& InSurfacicCoordinate, FSurfacicPoint& OutPoint3D, int32 InDerivativeOrder = 0) const override;
+	virtual void EvaluatePoint(const FPoint2D& InSurfacicCoordinate, FSurfacicPoint& OutPoint3D, int32 InDerivativeOrder = 0) const override;
 
-		virtual void Presample(const FSurfacicBoundary& InBoundaries, FCoordinateGrid& OutCoordinates) override
-		{
-			PresampleIsoCircle(InBoundaries, OutCoordinates, EIso::IsoU);
+	virtual void Presample(const FSurfacicBoundary& InBoundaries, FCoordinateGrid& OutCoordinates) override
+	{
+		PresampleIsoCircle(InBoundaries, OutCoordinates, EIso::IsoU);
 
-			OutCoordinates[EIso::IsoV].Empty(3);
-			OutCoordinates[EIso::IsoV].Add(InBoundaries[EIso::IsoV].Min);
-			OutCoordinates[EIso::IsoV].Add((InBoundaries[EIso::IsoV].Max + InBoundaries[EIso::IsoV].Min) / 2.0);
-			OutCoordinates[EIso::IsoV].Add(InBoundaries[EIso::IsoV].Max);
-		}
+		OutCoordinates[EIso::IsoV].Empty(3);
+		OutCoordinates[EIso::IsoV].Add(InBoundaries[EIso::IsoV].Min);
+		OutCoordinates[EIso::IsoV].Add((InBoundaries[EIso::IsoV].Max + InBoundaries[EIso::IsoV].Min) / 2.0);
+		OutCoordinates[EIso::IsoV].Add(InBoundaries[EIso::IsoV].Max);
+	}
 
-	};
+};
 }
 
diff --git a/Engine/Source/Runtime/Datasmith/CADKernel/Private/CADKernel/Geo/Surfaces/OffsetSurface.h b/Engine/Source/Runtime/Datasmith/CADKernel/Private/CADKernel/Geo/Surfaces/OffsetSurface.h
--- a/Engine/Source/Runtime/Datasmith/CADKernel/Private/CADKernel/Geo/Surfaces/OffsetSurface.h
+++ b/Engine/Source/Runtime/Datasmith/CADKernel/Private/CADKernel/Geo/Surfaces/OffsetSurface.h
@@ -3,92 +3,92 @@
 
 #include "CADKernel/Geo/Surfaces/Surface.h"
 
-namespace CADKernel
+namespace UE::CADKernel
 {
 
-	class CADKERNEL_API FOffsetSurface : public FSurface
-	{
-		friend FEntity;
+class CADKERNEL_API FOffsetSurface : public FSurface
+{
+	friend FEntity;
 
-	protected:
+protected:
 
-		TSharedPtr<FSurface> BaseSurface;
-		double Offset;
+	TSharedPtr<FSurface> BaseSurface;
+	double Offset;
 
-		FOffsetSurface(const double InToleranceGeometric, TSharedRef<FSurface> InBaseSurface, double InOffset)
-			: FSurface(InToleranceGeometric)
-			, BaseSurface(InBaseSurface)
-			, Offset(InOffset)
-		{
-			ComputeMinToleranceIso();
-		}
+	FOffsetSurface(const double InToleranceGeometric, TSharedRef<FSurface> InBaseSurface, double InOffset)
+		: FSurface(InToleranceGeometric)
+		, BaseSurface(InBaseSurface)
+		, Offset(InOffset)
+	{
+		ComputeMinToleranceIso();
+	}
 
-		FOffsetSurface() = default;
+	FOffsetSurface() = default;
 
-		void ComputeMinToleranceIso()
-		{
-			MinToleranceIso = BaseSurface->GetIsoTolerances();
-		}
+	void ComputeMinToleranceIso()
+	{
+		MinToleranceIso = BaseSurface->GetIsoTolerances();
+	}
 
-	public:
+public:
 
-		virtual void Serialize(FCADKernelArchive& Ar) override
-		{
-			FSurface::Serialize(Ar);
-			SerializeIdent(Ar, BaseSurface);
-			Ar << Offset;
-		}
+	virtual void Serialize(FCADKernelArchive& Ar) override
+	{
+		FSurface::Serialize(Ar);
+		SerializeIdent(Ar, BaseSurface);
+		Ar << Offset;
+	}
 
-		virtual void SpawnIdent(FDatabase& Database) override
+	virtual void SpawnIdent(FDatabase& Database) override
+	{
+		if (!FEntity::SetId(Database))
 		{
-			if (!FEntity::SetId(Database))
-			{
-				return;
-			}
-
-			BaseSurface->SpawnIdent(Database);
+			return;
 		}
 
-		virtual void ResetMarkersRecursively() override
-		{
-			ResetMarkers();
-			BaseSurface->ResetMarkersRecursively();
-		}
+		BaseSurface->SpawnIdent(Database);
+	}
 
-		ESurface GetSurfaceType() const
-		{
-			return ESurface::Offset;
-		}
+	virtual void ResetMarkersRecursively() override
+	{
+		ResetMarkers();
+		BaseSurface->ResetMarkersRecursively();
+	}
 
-		const double GetOffset() const
-		{
-			return Offset;
-		}
+	ESurface GetSurfaceType() const
+	{
+		return ESurface::Offset;
+	}
 
-		TSharedPtr<FSurface> GetBaseSurface() const
-		{
-			return BaseSurface;
-		}
+	const double GetOffset() const
+	{
+		return Offset;
+	}
+
+	TSharedPtr<FSurface> GetBaseSurface() const
+	{
+		return BaseSurface;
+	}
 
-		virtual void InitBoundary() override;
+	virtual void InitBoundary() override;
 
 #ifdef CADKERNEL_DEV
-		virtual FInfoEntity& GetInfo(FInfoEntity&) const override;
+	virtual FInfoEntity& GetInfo(FInfoEntity&) const override;
 #endif
 
-		virtual TSharedPtr<FEntityGeom> ApplyMatrix(const FMatrixH& InMatrix) const override;
+	virtual TSharedPtr<FEntityGeom> ApplyMatrix(const FMatrixH& InMatrix) const override;
 
-		virtual void LinesNotDerivables(const FSurfacicBoundary& Bounds, int32 InDerivativeOrder, FCoordinateGrid& OutNotDerivableCoordinates) const override;
+	virtual void LinesNotDerivables(const FSurfacicBoundary& Bounds, int32 InDerivativeOrder, FCoordinateGrid& OutNotDerivableCoordinates) const override;
 
-		/**
-		 * The Gradient and Laplacian are approximated by their value on the carrier surface.
-		 */
-		virtual void EvaluatePoint(const FPoint2D& InSurfacicCoordinate, FSurfacicPoint& OutPoint3D, int32 InDerivativeOrder = 0) const override;
-		virtual void EvaluatePoints(const TArray<FPoint2D>& InSurfacicCoordinates, TArray<FSurfacicPoint>& OutPoint3D, int32 InDerivativeOrder = 0) const override;
-		virtual void EvaluatePointGrid(const FCoordinateGrid& Coordinates, FSurfacicSampling& OutPoints, bool bComputeNormals = false) const override;
+	/**
+	 * The Gradient and Laplacian are approximated by their value on the carrier surface.
+	 */
+	virtual void EvaluatePoint(const FPoint2D& InSurfacicCoordinate, FSurfacicPoint& OutPoint3D, int32 InDerivativeOrder = 0) const override;
+	virtual void EvaluatePoints(const TArray<FPoint2D>& InSurfacicCoordinates, TArray<FSurfacicPoint>& OutPoint3D, int32 InDerivativeOrder = 0) const override;
+	virtual void EvaluatePointGrid(const FCoordinateGrid& Coordinates, FSurfacicSampling& OutPoints, bool bComputeNormals = false) const override;
 
-		virtual void Presample(const FSurfacicBoundary& InBoundaries, FCoordinateGrid& OutCoordinates) override;
-	};
+	virtual void Presample(const FSurfacicBoundary& InBoundaries, FCoordinateGrid& OutCoordinates) override;
+};
 
-} // namespace CADKernel
+} // namespace UE::CADKernel
 
diff --git a/Engine/Source/Runtime/Datasmith/CADKernel/Private/CADKernel/Geo/Surfaces/PlaneSurface.h b/Engine/Source/Runtime/Datasmith/CADKernel/Private/CADKernel/Geo/Surfaces/PlaneSurface.h
--- a/Engine/Source/Runtime/Datasmith/CADKernel/Private/CADKernel/Geo/Surfaces/PlaneSurface.h
+++ b/Engine/Source/Runtime/Datasmith/CADKernel/Private/CADKernel/Geo/Surfaces/PlaneSurface.h
@@ -3,108 +3,108 @@
 
 #include "CADKernel/Geo/Surfaces/Surface.h"
 
-namespace CADKernel
+namespace UE::CADKernel
 {
-	class CADKERNEL_API FPlaneSurface : public FSurface
+class CADKERNEL_API FPlaneSurface : public FSurface
+{
+	friend FEntity;
+
+protected:
+
+	FMatrixH Matrix;
+	FMatrixH InverseMatrix;
+
+	/**
+	 * The plane surface is the plane XY.
+	 * The surface is placed at its final position and orientation by the Matrix
+	 */
+	FPlaneSurface(const double InToleranceGeometric, const FMatrixH& InMatrix, const FSurfacicBoundary& InBoundary);
+
+	/**
+	 * The plane surface is the plane XY.
+	 * The surface is placed at its final position and orientation by the Matrix,
+	 * The matrix is calculated from the normal plane at its final position and its distance from the origin along the normal.
+	 */
+	FPlaneSurface(const double InToleranceGeometric, double InDistanceFromOrigin, FPoint InNormal, const FSurfacicBoundary& InBoundary)
+		: FPlaneSurface(InToleranceGeometric, InNormal* InDistanceFromOrigin, InNormal, InBoundary)
 	{
-		friend FEntity;
-
-	protected:
-
-		FMatrixH Matrix;
-		FMatrixH InverseMatrix;
-
-		/**
-		 * The plane surface is the plane XY.
-		 * The surface is placed at its final position and orientation by the Matrix
-		 */
-		FPlaneSurface(const double InToleranceGeometric, const FMatrixH& InMatrix, const FSurfacicBoundary& InBoundary);
-
-		/**
-		 * The plane surface is the plane XY.
-		 * The surface is placed at its final position and orientation by the Matrix,
-		 * The matrix is calculated from the normal plane at its final position and its distance from the origin along the normal.
-		 */
-		FPlaneSurface(const double InToleranceGeometric, double InDistanceFromOrigin, FPoint InNormal, const FSurfacicBoundary& InBoundary)
-			: FPlaneSurface(InToleranceGeometric, InNormal* InDistanceFromOrigin, InNormal, InBoundary)
-		{
-		}
-
-		/**
-		 * The plane surface is the plane XY.
-		 * The surface is placed at its final position and orientation by the Matrix 
-		 * The matrix is calculated from the plan origin at its final position and its final normal
-		 */
-		FPlaneSurface(const double InToleranceGeometric, const FPoint& Position, FPoint Normal, const FSurfacicBoundary& InBoundary);
-
-		FPlaneSurface() = default;
-		
-		void ComputeMinToleranceIso()
-		{
-			FPoint Origin = Matrix.Multiply(FPoint::ZeroPoint);
-
-			FPoint Point2DU{ 1 , 0, 0 };
-			FPoint Point2DV{ 0, 1, 0 };
-
-			double ToleranceU = Tolerance3D / ComputeScaleAlongAxis(Point2DU, Matrix, Origin);
-			double ToleranceV = Tolerance3D / ComputeScaleAlongAxis(Point2DV, Matrix, Origin);
-
-			MinToleranceIso.Set(ToleranceU, ToleranceV);
-		}
-
-	public:
-
-		virtual void Serialize(FCADKernelArchive& Ar) override
-		{
-			FSurface::Serialize(Ar);
-			Ar << Matrix;
-			Ar << InverseMatrix;
-		}
-
-		ESurface GetSurfaceType() const
-		{
-			return ESurface::Plane;
-		}
-
-		const FMatrixH& GetMatrix() const
-		{
-			return Matrix;
-		}
-
-		FPlane GetPlane() const;
+	}
 
-#ifdef CADKERNEL_DEV
-		virtual FInfoEntity& GetInfo(FInfoEntity&) const override;
-#endif
+	/**
+	 * The plane surface is the plane XY.
+	 * The surface is placed at its final position and orientation by the Matrix
+	 * The matrix is calculated from the plan origin at its final position and its final normal
+	 */
+	FPlaneSurface(const double InToleranceGeometric, const FPoint& Position, FPoint Normal, const FSurfacicBoundary& InBoundary);
+
+	FPlaneSurface() = default;
+
+	void ComputeMinToleranceIso()
+	{
+		FPoint Origin = Matrix.Multiply(FPoint::ZeroPoint);
+
+		FPoint Point2DU{ 1 , 0, 0 };
+		FPoint Point2DV{ 0, 1, 0 };
+
+		double ToleranceU = Tolerance3D / ComputeScaleAlongAxis(Point2DU, Matrix, Origin);
+		double ToleranceV = Tolerance3D / ComputeScaleAlongAxis(Point2DV, Matrix, Origin);
+
+		MinToleranceIso.Set(ToleranceU, ToleranceV);
+	}
+
+public:
 
-		virtual TSharedPtr<FEntityGeom> ApplyMatrix(const FMatrixH& InMatrix) const override;
+	virtual void Serialize(FCADKernelArchive& Ar) override
+	{
+		FSurface::Serialize(Ar);
+		Ar << Matrix;
+		Ar << InverseMatrix;
+	}
+
+	ESurface GetSurfaceType() const
+	{
+		return ESurface::Plane;
+	}
+
+	const FMatrixH& GetMatrix() const
+	{
+		return Matrix;
+	}
 
-		virtual void EvaluatePoint(const FPoint2D& InSurfacicCoordinate, FSurfacicPoint& OutPoint3D, int32 InDerivativeOrder = 0) const override;
-		virtual void EvaluatePoints(const TArray<FPoint2D>& InSurfacicCoordinates, TArray<FSurfacicPoint>& OutPoint3D, int32 InDerivativeOrder = 0) const override;
-		virtual void EvaluatePointGrid(const FCoordinateGrid& Coordinates, FSurfacicSampling& OutPoints, bool bComputeNormals = false) const override;
+	FPlane GetPlane() const;
 
-		virtual FPoint ProjectPoint(const FPoint& InPoint, FPoint* OutProjectedPoint = nullptr) const;// override;
-		virtual void ProjectPoints(const TArray<FPoint>& InPoints, TArray<FPoint>* OutProjectedPointCoordinates, TArray<FPoint>* OutProjectedPoints = nullptr) const;// override;
+#ifdef CADKERNEL_DEV
+	virtual FInfoEntity& GetInfo(FInfoEntity&) const override;
+#endif
+
+	virtual TSharedPtr<FEntityGeom> ApplyMatrix(const FMatrixH& InMatrix) const override;
 
-		virtual void Presample(const FSurfacicBoundary& InBoundaries, FCoordinateGrid& OutCoordinates) override
-		{
-			OutCoordinates[EIso::IsoU].Empty(3);
-			OutCoordinates[EIso::IsoU].Add(InBoundaries[EIso::IsoU].Min);
-			OutCoordinates[EIso::IsoU].Add(InBoundaries[EIso::IsoU].GetMiddle());
-			OutCoordinates[EIso::IsoU].Add(InBoundaries[EIso::IsoU].Max);
+	virtual void EvaluatePoint(const FPoint2D& InSurfacicCoordinate, FSurfacicPoint& OutPoint3D, int32 InDerivativeOrder = 0) const override;
+	virtual void EvaluatePoints(const TArray<FPoint2D>& InSurfacicCoordinates, TArray<FSurfacicPoint>& OutPoint3D, int32 InDerivativeOrder = 0) const override;
+	virtual void EvaluatePointGrid(const FCoordinateGrid& Coordinates, FSurfacicSampling& OutPoints, bool bComputeNormals = false) const override;
 
-			OutCoordinates[EIso::IsoV].Empty(3);
-			OutCoordinates[EIso::IsoV].Add(InBoundaries[EIso::IsoV].Min);
-			OutCoordinates[EIso::IsoV].Add(InBoundaries[EIso::IsoV].GetMiddle());
-			OutCoordinates[EIso::IsoV].Add(InBoundaries[EIso::IsoV].Max);
-		}
+	virtual FPoint ProjectPoint(const FPoint& InPoint, FPoint* OutProjectedPoint = nullptr) const;// override;
+	virtual void ProjectPoints(const TArray<FPoint>& InPoints, TArray<FPoint>* OutProjectedPointCoordinates, TArray<FPoint>* OutProjectedPoints = nullptr) const;// override;
 
-		virtual void IsSurfaceClosed(bool& bOutClosedAlongU, bool& bOutClosedAlongV) const override
-		{
-			bOutClosedAlongU = false;
-			bOutClosedAlongV = false;
-		}
+	virtual void Presample(const FSurfacicBoundary& InBoundaries, FCoordinateGrid& OutCoordinates) override
+	{
+		OutCoordinates[EIso::IsoU].Empty(3);
+		OutCoordinates[EIso::IsoU].Add(InBoundaries[EIso::IsoU].Min);
+		OutCoordinates[EIso::IsoU].Add(InBoundaries[EIso::IsoU].GetMiddle());
+		OutCoordinates[EIso::IsoU].Add(InBoundaries[EIso::IsoU].Max);
+
+		OutCoordinates[EIso::IsoV].Empty(3);
+		OutCoordinates[EIso::IsoV].Add(InBoundaries[EIso::IsoV].Min);
+		OutCoordinates[EIso::IsoV].Add(InBoundaries[EIso::IsoV].GetMiddle());
+		OutCoordinates[EIso::IsoV].Add(InBoundaries[EIso::IsoV].Max);
+	}
+
+	virtual void IsSurfaceClosed(bool& bOutClosedAlongU, bool& bOutClosedAlongV) const override
+	{
+		bOutClosedAlongU = false;
+		bOutClosedAlongV = false;
+	}
 
-	};
+};
 
-} // namespace CADKernel
+} // namespace UE::CADKernel
diff --git a/Engine/Source/Runtime/Datasmith/CADKernel/Private/CADKernel/Geo/Surfaces/RevolutionSurface.h b/Engine/Source/Runtime/Datasmith/CADKernel/Private/CADKernel/Geo/Surfaces/RevolutionSurface.h
--- a/Engine/Source/Runtime/Datasmith/CADKernel/Private/CADKernel/Geo/Surfaces/RevolutionSurface.h
+++ b/Engine/Source/Runtime/Datasmith/CADKernel/Private/CADKernel/Geo/Surfaces/RevolutionSurface.h
@@ -3,79 +3,79 @@
 
 #include "CADKernel/Geo/Surfaces/Surface.h"
 
-namespace CADKernel
+namespace UE::CADKernel
 {
-	class FCurve;
-	class FSegmentCurve;
+class FCurve;
+class FSegmentCurve;
 
-	class CADKERNEL_API FRevolutionSurface : public FSurface
+class CADKERNEL_API FRevolutionSurface : public FSurface
+{
+	friend FEntity;
+
+protected:
+
+	TSharedPtr<FSegmentCurve> Axis;
+	TSharedPtr<FCurve> Generatrix;
+
+	FPoint RotationAxis;
+
+	/**
+	 * A surface of revolution is defined by an axis of rotation (Axis) and a generatrix.
+	 * The start and terminate rotation angles can limited the surface.
+	 * The surface is created by rotating the generatrix about the axis of rotation through the start and terminating angles.
+	 *
+	 * The bounds of the surface of revolution are defined as follow:
+	 * Bounds[EIso::IsoU].Min = Start rotation angle;
+	 * Bounds[EIso::IsoU].Max = Terminate rotation angle;
+	 * Bounds[EIso::IsoV].Min = Min Generatrix bounds;
+	 * Bounds[EIso::IsoV].Max = Max Generatrix bounds;
+	 */
+	FRevolutionSurface(const double InToleranceGeometric, TSharedRef<FSegmentCurve> Axe, TSharedRef<FCurve> Generatrix, double MinAngle, double MaxAngle);
+
+	FRevolutionSurface() = default;
+
+public:
+
+	virtual void Serialize(FCADKernelArchive& Ar) override
+	{
+		FSurface::Serialize(Ar);
+		SerializeIdent(Ar, Axis);
+		SerializeIdent(Ar, Generatrix);
+		Ar << RotationAxis;
+	}
+
+	virtual void SpawnIdent(FDatabase& Database) override;
+	virtual void ResetMarkersRecursively() override;
+
+	ESurface GetSurfaceType() const
+	{
+		return ESurface::Revolution;
+	}
+
+	TSharedPtr<FCurve> GetGeneratrix() const
+	{
+		return Generatrix;
+	}
+
+	virtual TSharedPtr<FSegmentCurve> GetAxe() const
 	{
-		friend FEntity;
-
-	protected:
-
-		TSharedPtr<FSegmentCurve> Axis;
-		TSharedPtr<FCurve> Generatrix;
-
-		FPoint RotationAxis;
-
-		/**
-		 * A surface of revolution is defined by an axis of rotation (Axis) and a generatrix.
-		 * The start and terminate rotation angles can limited the surface. 
-		 * The surface is created by rotating the generatrix about the axis of rotation through the start and terminating angles. 
-		 * 
-		 * The bounds of the surface of revolution are defined as follow:
-		 * Bounds[EIso::IsoU].Min = Start rotation angle;
-		 * Bounds[EIso::IsoU].Max = Terminate rotation angle;
-		 * Bounds[EIso::IsoV].Min = Min Generatrix bounds;
-		 * Bounds[EIso::IsoV].Max = Max Generatrix bounds;
-		 */
-		FRevolutionSurface(const double InToleranceGeometric, TSharedRef<FSegmentCurve> Axe, TSharedRef<FCurve> Generatrix, double MinAngle, double MaxAngle);
-
-		FRevolutionSurface() = default;
-
-	public:
-
-		virtual void Serialize(FCADKernelArchive& Ar) override
-		{
-			FSurface::Serialize(Ar);
-			SerializeIdent(Ar, Axis);
-			SerializeIdent(Ar, Generatrix);
-			Ar << RotationAxis;
-		}
-
-		virtual void SpawnIdent(FDatabase& Database) override;
-		virtual void ResetMarkersRecursively() override;
-
-		ESurface GetSurfaceType() const
-		{
-			return ESurface::Revolution;
-		}
-
-		TSharedPtr<FCurve> GetGeneratrix() const
-		{
-			return Generatrix;
-		}
-
-		virtual TSharedPtr<FSegmentCurve> GetAxe() const
-		{
-			return Axis;
-		}
+		return Axis;
+	}
 
 #ifdef CADKERNEL_DEV
-		virtual FInfoEntity& GetInfo(FInfoEntity&) const override;
+	virtual FInfoEntity& GetInfo(FInfoEntity&) const override;
 #endif
 
-		virtual TSharedPtr<FEntityGeom> ApplyMatrix(const FMatrixH& InMatrix) const override;
+	virtual TSharedPtr<FEntityGeom> ApplyMatrix(const FMatrixH& InMatrix) const override;
 
-		virtual void EvaluatePoint(const FPoint2D& InSurfacicCoordinate, FSurfacicPoint& OutPoint3D, int32 InDerivativeOrder = 0) const override;
-		virtual void EvaluatePointGrid(const FCoordinateGrid& Coordinates, FSurfacicSampling& OutPoints, bool bComputeNormals = false) const override;
+	virtual void EvaluatePoint(const FPoint2D& InSurfacicCoordinate, FSurfacicPoint& OutPoint3D, int32 InDerivativeOrder = 0) const override;
+	virtual void EvaluatePointGrid(const FCoordinateGrid& Coordinates, FSurfacicSampling& OutPoints, bool bComputeNormals = false) const override;
 
-		virtual void LinesNotDerivables(const FSurfacicBoundary& Bounds, int32 InDerivativeOrder, FCoordinateGrid& OutCoordinates) const override;
+	virtual void LinesNotDerivables(const FSurfacicBoundary& Bounds, int32 InDerivativeOrder, FCoordinateGrid& OutCoordinates) const override;
 
-		virtual void Presample(const FSurfacicBoundary& InBoundaries, FCoordinateGrid& OutCoordinates) override;
+	virtual void Presample(const FSurfacicBoundary& InBoundaries, FCoordinateGrid& OutCoordinates) override;
 
-	};
+};
 
-} // namespace CADKernel
+} // namespace UE::CADKernel
 
diff --git a/Engine/Source/Runtime/Datasmith/CADKernel/Private/CADKernel/Geo/Surfaces/RuledSurface.h b/Engine/Source/Runtime/Datasmith/CADKernel/Private/CADKernel/Geo/Surfaces/RuledSurface.h
--- a/Engine/Source/Runtime/Datasmith/CADKernel/Private/CADKernel/Geo/Surfaces/RuledSurface.h
+++ b/Engine/Source/Runtime/Datasmith/CADKernel/Private/CADKernel/Geo/Surfaces/RuledSurface.h
@@ -3,64 +3,64 @@
 
 #include "CADKernel/Geo/Surfaces/Surface.h"
 
-namespace CADKernel
+namespace UE::CADKernel
 {
-	class FCurve;
-	class FNURBSSurface;
+class FCurve;
+class FNURBSSurface;
 
-	class CADKERNEL_API FRuledSurface : public FSurface
-	{
-		friend FEntity;
+class CADKERNEL_API FRuledSurface : public FSurface
+{
+	friend FEntity;
 
-	protected:
-		TSharedPtr<FCurve> Curves[2];
+protected:
+	TSharedPtr<FCurve> Curves[2];
 
-		FRuledSurface(const double InToleranceGeometric, TSharedPtr<FCurve> InCurveU, TSharedPtr<FCurve> InCurveV)
-			: FSurface(InToleranceGeometric, 0., 1., 0., 1.)
-		{
-			Curves[0] = InCurveU;
-			Curves[1] = InCurveV;
-			ComputeDefaultMinToleranceIso();
-		}
+	FRuledSurface(const double InToleranceGeometric, TSharedPtr<FCurve> InCurveU, TSharedPtr<FCurve> InCurveV)
+		: FSurface(InToleranceGeometric, 0., 1., 0., 1.)
+	{
+		Curves[0] = InCurveU;
+		Curves[1] = InCurveV;
+		ComputeDefaultMinToleranceIso();
+	}
 
-		FRuledSurface() = default;
+	FRuledSurface() = default;
 
-	public:
+public:
 
-		virtual void Serialize(FCADKernelArchive& Ar) override
-		{
-			FSurface::Serialize(Ar);
-			SerializeIdent(Ar, Curves[0]);
-			SerializeIdent(Ar, Curves[1]);
-		}
+	virtual void Serialize(FCADKernelArchive& Ar) override
+	{
+		FSurface::Serialize(Ar);
+		SerializeIdent(Ar, Curves[0]);
+		SerializeIdent(Ar, Curves[1]);
+	}
 
-		virtual void SpawnIdent(FDatabase& Database) override;
-		virtual void ResetMarkersRecursively() override;
+	virtual void SpawnIdent(FDatabase& Database) override;
+	virtual void ResetMarkersRecursively() override;
 
-		ESurface GetSurfaceType() const
-		{
-			return ESurface::Ruled;
-		}
+	ESurface GetSurfaceType() const
+	{
+		return ESurface::Ruled;
+	}
 
-		const TSharedPtr<FCurve>& GetCurve(EIso Iso) const
-		{
-			return Curves[Iso];
-		}
+	const TSharedPtr<FCurve>& GetCurve(EIso Iso) const
+	{
+		return Curves[Iso];
+	}
 
 #ifdef CADKERNEL_DEV
-		virtual FInfoEntity& GetInfo(FInfoEntity&) const override;
+	virtual FInfoEntity& GetInfo(FInfoEntity&) const override;
 #endif
 
-		virtual TSharedPtr<FEntityGeom> ApplyMatrix(const FMatrixH& InMatrix) const override;
+	virtual TSharedPtr<FEntityGeom> ApplyMatrix(const FMatrixH& InMatrix) const override;
 
-		virtual void EvaluatePoint(const FPoint2D& InSurfacicCoordinate, FSurfacicPoint& OutPoint3D, int32 InDerivativeOrder = 0) const override;
-		virtual void EvaluatePointGrid(const FCoordinateGrid& Coordinates, FSurfacicSampling& OutPoints, bool bComputeNormals = false) const override;
+	virtual void EvaluatePoint(const FPoint2D& InSurfacicCoordinate, FSurfacicPoint& OutPoint3D, int32 InDerivativeOrder = 0) const override;
+	virtual void EvaluatePointGrid(const FCoordinateGrid& Coordinates, FSurfacicSampling& OutPoints, bool bComputeNormals = false) const override;
 
-		virtual void LinesNotDerivables(const FSurfacicBoundary& Bounds, int32 InDerivativeOrder, FCoordinateGrid& OutCoordinates) const override;
+	virtual void LinesNotDerivables(const FSurfacicBoundary& Bounds, int32 InDerivativeOrder, FCoordinateGrid& OutCoordinates) const override;
 
-		virtual void Presample(const FSurfacicBoundary& InBoundaries, FCoordinateGrid& OutCoordinates) override;
+	virtual void Presample(const FSurfacicBoundary& InBoundaries, FCoordinateGrid& OutCoordinates) override;
 
-	};
+};
 
-} // namespace CADKernel
+} // namespace UE::CADKernel
 
diff --git a/Engine/Source/Runtime/Datasmith/CADKernel/Private/CADKernel/Geo/Surfaces/SphericalSurface.h b/Engine/Source/Runtime/Datasmith/CADKernel/Private/CADKernel/Geo/Surfaces/SphericalSurface.h
--- a/Engine/Source/Runtime/Datasmith/CADKernel/Private/CADKernel/Geo/Surfaces/SphericalSurface.h
+++ b/Engine/Source/Runtime/Datasmith/CADKernel/Private/CADKernel/Geo/Surfaces/SphericalSurface.h
@@ -3,138 +3,138 @@
 
 #include "CADKernel/Geo/Surfaces/Surface.h"
 
-namespace CADKernel
+namespace UE::CADKernel
 {
-	class CADKERNEL_API FSphericalSurface : public FSurface
+class CADKERNEL_API FSphericalSurface : public FSurface
+{
+	friend FEntity;
+
+protected:
+	FMatrixH Matrix;
+	double Radius;
+
+	/**
+	 * The spherical surface is defined its radius.
+	 *
+	 * It's defined as the rotation around Z axis of an semicircle defined in the plan XY centered at the origin.
+	 *
+	 * The spherical surface is placed at its final position and orientation by the Matrix
+	 */
+	FSphericalSurface(const double InToleranceGeometric, const FMatrixH& InMatrix, double InRadius, double InParallelStartAngle = 0.0, double InParallelEndAngle = DOUBLE_TWO_PI, double InMeridianStartAngle = -DOUBLE_HALF_PI, double InMeridianEndAngle = DOUBLE_HALF_PI)
+		: FSphericalSurface(InToleranceGeometric, InMatrix, InRadius, FSurfacicBoundary(InParallelStartAngle, InParallelEndAngle, InMeridianStartAngle, InMeridianEndAngle))
 	{
-		friend FEntity;
-
-	protected:
-		FMatrixH Matrix;
-		double Radius;
-
-		/**
-		 * The spherical surface is defined its radius. 
-		 *
-		 * It's defined as the rotation around Z axis of an semicircle defined in the plan XY centered at the origin. 
-		 *
-		 * The spherical surface is placed at its final position and orientation by the Matrix
-		 */
-		FSphericalSurface(const double InToleranceGeometric, const FMatrixH& InMatrix, double InRadius, double InParallelStartAngle = 0.0, double InParallelEndAngle = DOUBLE_TWO_PI, double InMeridianStartAngle = -DOUBLE_HALF_PI, double InMeridianEndAngle = DOUBLE_HALF_PI)
-			: FSphericalSurface(InToleranceGeometric, InMatrix, InRadius, FSurfacicBoundary(InParallelStartAngle, InParallelEndAngle, InMeridianStartAngle, InMeridianEndAngle))
-		{
-		}
+	}
+
+	/**
+	 * The spherical surface is defined its radius.
+	 *
+	 * It's defined as the rotation around Z axis of an semicircle defined in the plan XY centered at the origin.
+	 *
+	 * The spherical surface is placed at its final position and orientation by the Matrix
+	 *
+	 * The bounds of the spherical surface are defined as follow:
+	 * Bounds[EIso::IsoU].Min = MeridianStartAngle
+	 * Bounds[EIso::IsoU].Max = MeridianEndAngle
+	 * Bounds[EIso::IsoV].Min = ParallelStartAngle
+	 * Bounds[EIso::IsoV].Max = ParallelEndAngle
+	 */
+	FSphericalSurface(const double InToleranceGeometric, const FMatrixH& InMatrix, double InRadius, const FSurfacicBoundary& Boundary)
+		: FSurface(InToleranceGeometric, Boundary)
+		, Matrix(InMatrix)
+		, Radius(InRadius)
+	{
+		ComputeMinToleranceIso();
+	}
 
-		/**
-		 * The spherical surface is defined its radius.
-		 *
-		 * It's defined as the rotation around Z axis of an semicircle defined in the plan XY centered at the origin.
-		 *
-		 * The spherical surface is placed at its final position and orientation by the Matrix
-		 *
-		 * The bounds of the spherical surface are defined as follow:
-		 * Bounds[EIso::IsoU].Min = MeridianStartAngle
-		 * Bounds[EIso::IsoU].Max = MeridianEndAngle
-		 * Bounds[EIso::IsoV].Min = ParallelStartAngle
-		 * Bounds[EIso::IsoV].Max = ParallelEndAngle
-		 */
-		FSphericalSurface(const double InToleranceGeometric, const FMatrixH& InMatrix, double InRadius, const FSurfacicBoundary& Boundary)
-			: FSurface(InToleranceGeometric, Boundary)
-			, Matrix(InMatrix)
-			, Radius(InRadius)
-		{
-			ComputeMinToleranceIso();
-		}
+	FSphericalSurface() = default;
 
-		FSphericalSurface() = default;
+	void ComputeMinToleranceIso()
+	{
+		double Tolerance2D = Tolerance3D / Radius;
 
-		void ComputeMinToleranceIso()
-		{
-			double Tolerance2D = Tolerance3D / Radius;
+		FPoint Origin = Matrix.Multiply(FPoint::ZeroPoint);
 
-			FPoint Origin = Matrix.Multiply(FPoint::ZeroPoint);
+		FPoint Point2DU{ 1 , 0, 0 };
+		FPoint Point2DV{ 0, 1, 0 };
 
-			FPoint Point2DU{ 1 , 0, 0 };
-			FPoint Point2DV{ 0, 1, 0 };
+		double ToleranceU = Tolerance2D / ComputeScaleAlongAxis(Point2DU, Matrix, Origin);
+		double ToleranceV = Tolerance2D / ComputeScaleAlongAxis(Point2DV, Matrix, Origin);
 
-			double ToleranceU = Tolerance2D / ComputeScaleAlongAxis(Point2DU, Matrix, Origin);
-			double ToleranceV = Tolerance2D / ComputeScaleAlongAxis(Point2DV, Matrix, Origin);
+		MinToleranceIso.Set(ToleranceU, ToleranceV);
+	}
 
-			MinToleranceIso.Set(ToleranceU, ToleranceV);
-		}
-	
-	public:
+public:
 
-		virtual void Serialize(FCADKernelArchive& Ar) override
-		{
-			FSurface::Serialize(Ar);
-			Ar << Matrix;
-			Ar << Radius;
-		}
+	virtual void Serialize(FCADKernelArchive& Ar) override
+	{
+		FSurface::Serialize(Ar);
+		Ar << Matrix;
+		Ar << Radius;
+	}
 
-		ESurface GetSurfaceType() const
-		{
-			return ESurface::Sphere;
-		}
+	ESurface GetSurfaceType() const
+	{
+		return ESurface::Sphere;
+	}
 
 #ifdef CADKERNEL_DEV
-		virtual FInfoEntity& GetInfo(FInfoEntity&) const override;
+	virtual FInfoEntity& GetInfo(FInfoEntity&) const override;
 #endif
 
-		virtual TSharedPtr<FEntityGeom> ApplyMatrix(const FMatrixH& InMatrix) const override;
+	virtual TSharedPtr<FEntityGeom> ApplyMatrix(const FMatrixH& InMatrix) const override;
 
-		virtual void EvaluatePoint(const FPoint2D& InSurfacicCoordinate, FSurfacicPoint& OutPoint3D, int32 InDerivativeOrder = 0) const override
-		{
-			double CosU = cos(InSurfacicCoordinate.U);
-			double CosV = cos(InSurfacicCoordinate.V);
-
-			double SinU = sin(InSurfacicCoordinate.U);
-			double SinV = sin(InSurfacicCoordinate.V);
-
-			OutPoint3D.DerivativeOrder = InDerivativeOrder;
-			OutPoint3D.Point.Set(Radius * CosV * CosU, Radius * CosV * SinU, Radius * SinV);
-			OutPoint3D.Point = Matrix.Multiply(OutPoint3D.Point);
-
-			if (InDerivativeOrder > 0)
-			{
-				OutPoint3D.GradientU = FPoint(-Radius * CosV * SinU, Radius * CosV * CosU, 0.0);
-				OutPoint3D.GradientV = FPoint(-Radius * SinV * CosU, -Radius * SinV * SinU, Radius * CosV);
-
-				OutPoint3D.GradientU = Matrix.MultiplyVector(OutPoint3D.GradientU);
-				OutPoint3D.GradientV = Matrix.MultiplyVector(OutPoint3D.GradientV);
-			}
-
-			if (InDerivativeOrder > 1)
-			{
-				OutPoint3D.LaplacianU = FPoint(-Radius * CosV * CosU, -Radius * CosV * SinU, 0.0);
-				OutPoint3D.LaplacianV = FPoint(OutPoint3D.LaplacianU.X, OutPoint3D.LaplacianU.Y, -Radius * SinV);
-				OutPoint3D.LaplacianUV = FPoint(Radius * SinV * SinU, -Radius * SinV * CosU, 0.);
-
-				OutPoint3D.LaplacianU = Matrix.MultiplyVector(OutPoint3D.LaplacianU);
-				OutPoint3D.LaplacianV = Matrix.MultiplyVector(OutPoint3D.LaplacianV);
-				OutPoint3D.LaplacianUV = Matrix.MultiplyVector(OutPoint3D.LaplacianUV);
-			}
-		}
+	virtual void EvaluatePoint(const FPoint2D& InSurfacicCoordinate, FSurfacicPoint& OutPoint3D, int32 InDerivativeOrder = 0) const override
+	{
+		double CosU = cos(InSurfacicCoordinate.U);
+		double CosV = cos(InSurfacicCoordinate.V);
 
-		virtual void EvaluatePointGrid(const FCoordinateGrid& Coordinates, FSurfacicSampling& OutPoints, bool bComputeNormals = false) const override;
+		double SinU = sin(InSurfacicCoordinate.U);
+		double SinV = sin(InSurfacicCoordinate.V);
 
-		virtual FPoint2D EvaluatePointInCylindricalSpace(const FPoint2D& InSurfacicCoordinate) const override
-		{
-			double CosU = cos(InSurfacicCoordinate.U);
-			double CosV = cos(InSurfacicCoordinate.V);
+		OutPoint3D.DerivativeOrder = InDerivativeOrder;
+		OutPoint3D.Point.Set(Radius * CosV * CosU, Radius * CosV * SinU, Radius * SinV);
+		OutPoint3D.Point = Matrix.Multiply(OutPoint3D.Point);
 
-			double SinU = sin(InSurfacicCoordinate.U);
-			double SwapOrientation = (InSurfacicCoordinate.V < DOUBLE_PI && InSurfacicCoordinate.V >= 0) ? 1.0 : -1.0;
+		if (InDerivativeOrder > 0)
+		{
+			OutPoint3D.GradientU = FPoint(-Radius * CosV * SinU, Radius * CosV * CosU, 0.0);
+			OutPoint3D.GradientV = FPoint(-Radius * SinV * CosU, -Radius * SinV * SinU, Radius * CosV);
 
-			return FPoint(Radius * CosV * CosU * SwapOrientation, Radius * CosV * SinU );
+			OutPoint3D.GradientU = Matrix.MultiplyVector(OutPoint3D.GradientU);
+			OutPoint3D.GradientV = Matrix.MultiplyVector(OutPoint3D.GradientV);
 		}
 
-		virtual void EvaluatePointGridInCylindricalSpace(const FCoordinateGrid& Coordinates, TArray<FPoint2D>&) const override;
-
-		virtual void Presample(const FSurfacicBoundary& InBoundaries, FCoordinateGrid& OutCoordinates) override
+		if (InDerivativeOrder > 1)
 		{
-			PresampleIsoCircle(InBoundaries, OutCoordinates, EIso::IsoU);
-			PresampleIsoCircle(InBoundaries, OutCoordinates, EIso::IsoV);
+			OutPoint3D.LaplacianU = FPoint(-Radius * CosV * CosU, -Radius * CosV * SinU, 0.0);
+			OutPoint3D.LaplacianV = FPoint(OutPoint3D.LaplacianU.X, OutPoint3D.LaplacianU.Y, -Radius * SinV);
+			OutPoint3D.LaplacianUV = FPoint(Radius * SinV * SinU, -Radius * SinV * CosU, 0.);
+
+			OutPoint3D.LaplacianU = Matrix.MultiplyVector(OutPoint3D.LaplacianU);
+			OutPoint3D.LaplacianV = Matrix.MultiplyVector(OutPoint3D.LaplacianV);
+			OutPoint3D.LaplacianUV = Matrix.MultiplyVector(OutPoint3D.LaplacianUV);
 		}
-	};
+	}
+
+	virtual void EvaluatePointGrid(const FCoordinateGrid& Coordinates, FSurfacicSampling& OutPoints, bool bComputeNormals = false) const override;
+
+	virtual FPoint2D EvaluatePointInCylindricalSpace(const FPoint2D& InSurfacicCoordinate) const override
+	{
+		double CosU = cos(InSurfacicCoordinate.U);
+		double CosV = cos(InSurfacicCoordinate.V);
+
+		double SinU = sin(InSurfacicCoordinate.U);
+		double SwapOrientation = (InSurfacicCoordinate.V < DOUBLE_PI&& InSurfacicCoordinate.V >= 0) ? 1.0 : -1.0;
+
+		return FPoint(Radius * CosV * CosU * SwapOrientation, Radius * CosV * SinU);
+	}
+
+	virtual void EvaluatePointGridInCylindricalSpace(const FCoordinateGrid& Coordinates, TArray<FPoint2D>&) const override;
+
+	virtual void Presample(const FSurfacicBoundary& InBoundaries, FCoordinateGrid& OutCoordinates) override
+	{
+		PresampleIsoCircle(InBoundaries, OutCoordinates, EIso::IsoU);
+		PresampleIsoCircle(InBoundaries, OutCoordinates, EIso::IsoV);
+	}
+};
 }
diff --git a/Engine/Source/Runtime/Datasmith/CADKernel/Private/CADKernel/Geo/Surfaces/SurfaceMaker.cpp b/Engine/Source/Runtime/Datasmith/CADKernel/Private/CADKernel/Geo/Surfaces/SurfaceMaker.cpp
--- a/Engine/Source/Runtime/Datasmith/CADKernel/Private/CADKernel/Geo/Surfaces/SurfaceMaker.cpp
+++ b/Engine/Source/Runtime/Datasmith/CADKernel/Private/CADKernel/Geo/Surfaces/SurfaceMaker.cpp
@@ -17,47 +17,47 @@
 #include "CADKernel/Geo/Surfaces/ConeSurface.h"
 
 
-namespace CADKernel
+namespace UE::CADKernel
 {
 TSharedPtr<FSurface> FSurface::MakeBezierSurface(const double InToleranceGeometric, int32 InUDegre, int32 InVDegre, const TArray<FPoint>& InPoles)
 {
-	return CADKernel::FEntity::MakeShared<CADKernel::FBezierSurface>(InToleranceGeometric, InUDegre, InVDegre, InPoles);
+	return UE::CADKernel::FEntity::MakeShared<UE::CADKernel::FBezierSurface>(InToleranceGeometric, InUDegre, InVDegre, InPoles);
 }
 
 TSharedPtr<FSurface> FSurface::MakeConeSurface(const double InToleranceGeometric, const FMatrixH& InMatrix, double InStartRadius, double InConeAngle, const FSurfacicBoundary& InBoundary)
 {
-	return CADKernel::FEntity::MakeShared<CADKernel::FConeSurface>(InToleranceGeometric, InMatrix, InStartRadius, InConeAngle, InBoundary);
+	return UE::CADKernel::FEntity::MakeShared<UE::CADKernel::FConeSurface>(InToleranceGeometric, InMatrix, InStartRadius, InConeAngle, InBoundary);
 }
 
 TSharedPtr<FSurface> FSurface::MakeCylinderSurface(const double InToleranceGeometric, const FMatrixH& InMatrix, const double InRadius, const FSurfacicBoundary& InBoundary)
 {
-	return CADKernel::FEntity::MakeShared<CADKernel::FCylinderSurface>(InToleranceGeometric, InMatrix, InRadius, InBoundary);
+	return UE::CADKernel::FEntity::MakeShared<UE::CADKernel::FCylinderSurface>(InToleranceGeometric, InMatrix, InRadius, InBoundary);
 }
 
 TSharedPtr<FSurface> FSurface::MakeNurbsSurface(const double InToleranceGeometric, const FNurbsSurfaceHomogeneousData& NurbsData)
 {
-	return CADKernel::FEntity::MakeShared<CADKernel::FNURBSSurface>(InToleranceGeometric, NurbsData);
+	return UE::CADKernel::FEntity::MakeShared<UE::CADKernel::FNURBSSurface>(InToleranceGeometric, NurbsData);
 }
 
 TSharedPtr<FSurface> FSurface::MakeNurbsSurface(const double InToleranceGeometric, const FNurbsSurfaceData& NurbsData)
 {
-	return CADKernel::FEntity::MakeShared<CADKernel::FNURBSSurface>(InToleranceGeometric, NurbsData);
+	return UE::CADKernel::FEntity::MakeShared<UE::CADKernel::FNURBSSurface>(InToleranceGeometric, NurbsData);
 }
 
 TSharedPtr<FSurface> FSurface::MakePlaneSurface(const double InToleranceGeometric, const FMatrixH& InMatrix, const FSurfacicBoundary& InBoundary)
 {
-	return CADKernel::FEntity::MakeShared<CADKernel::FPlaneSurface>(InToleranceGeometric, InMatrix, InBoundary);
+	return UE::CADKernel::FEntity::MakeShared<UE::CADKernel::FPlaneSurface>(InToleranceGeometric, InMatrix, InBoundary);
 }
 
 TSharedPtr<FSurface> FSurface::MakeSphericalSurface(const double InToleranceGeometric, const FMatrixH& InMatrix, double InRadius, const FSurfacicBoundary& InBoundary)
 {
-	return CADKernel::FEntity::MakeShared<CADKernel::FSphericalSurface>(InToleranceGeometric, InMatrix, InRadius, InBoundary);
+	return UE::CADKernel::FEntity::MakeShared<UE::CADKernel::FSphericalSurface>(InToleranceGeometric, InMatrix, InRadius, InBoundary);
 }
 
 TSharedPtr<FSurface> FSurface::MakeTorusSurface(const double InToleranceGeometric, const FMatrixH& InMatrix, double InMajorRadius, double InMinorRadius, const FSurfacicBoundary& InBoundary)
 {
-	return CADKernel::FEntity::MakeShared<CADKernel::FTorusSurface>(InToleranceGeometric, InMatrix, InMajorRadius, InMinorRadius, InBoundary);
+	return UE::CADKernel::FEntity::MakeShared<UE::CADKernel::FTorusSurface>(InToleranceGeometric, InMatrix, InMajorRadius, InMinorRadius, InBoundary);
 }
 
-} // namespace CADKernel
+} // namespace UE::CADKernel
 
diff --git a/Engine/Source/Runtime/Datasmith/CADKernel/Private/CADKernel/Geo/Surfaces/SurfaceSerialize.cpp b/Engine/Source/Runtime/Datasmith/CADKernel/Private/CADKernel/Geo/Surfaces/SurfaceSerialize.cpp
--- a/Engine/Source/Runtime/Datasmith/CADKernel/Private/CADKernel/Geo/Surfaces/SurfaceSerialize.cpp
+++ b/Engine/Source/Runtime/Datasmith/CADKernel/Private/CADKernel/Geo/Surfaces/SurfaceSerialize.cpp
@@ -19,46 +19,46 @@
 #include "CADKernel/Geo/Surfaces/ConeSurface.h"
 
 
-namespace CADKernel
+namespace UE::CADKernel
 {
-	TSharedPtr<FSurface> FSurface::Deserialize(FCADKernelArchive& Archive)
-	{
-		ESurface SurfaceType = ESurface::None;
-		Archive << SurfaceType;
+TSharedPtr<FSurface> FSurface::Deserialize(FCADKernelArchive& Archive)
+{
+	ESurface SurfaceType = ESurface::None;
+	Archive << SurfaceType;
 
-		switch (SurfaceType)
-		{
-		case ESurface::Bezier:
-			return FEntity::MakeShared<FBezierSurface>(Archive);
-		case ESurface::Cone:
-			return FEntity::MakeShared<FConeSurface>(Archive);
-		case ESurface::Composite:
-			return FEntity::MakeShared<FCompositeSurface>(Archive);
-		case ESurface::Coons:
-			return FEntity::MakeShared<FCoonsSurface>(Archive);
-		case ESurface::Cylinder:
-			return FEntity::MakeShared<FCylinderSurface>(Archive);
-		case ESurface::Nurbs:
-			return FEntity::MakeShared<FNURBSSurface>(Archive);
-		case ESurface::Offset:
-			return FEntity::MakeShared<FOffsetSurface>(Archive);
-		case ESurface::Plane:
-			return FEntity::MakeShared<FPlaneSurface>(Archive);
-		case ESurface::Revolution:
-			return FEntity::MakeShared<FRevolutionSurface>(Archive);
-		case ESurface::Ruled:
-			return FEntity::MakeShared<FRuledSurface>(Archive);
-		case ESurface::Sphere:
-			return FEntity::MakeShared<FSphericalSurface>(Archive);
-		case ESurface::TabulatedCylinder:
-			return FEntity::MakeShared<FTabulatedCylinderSurface>(Archive);
-		case ESurface::Torus:
-			return FEntity::MakeShared<FTorusSurface>(Archive);
-		default:
-			return TSharedPtr<FSurface>();
-		}
+	switch (SurfaceType)
+	{
+	case ESurface::Bezier:
+		return FEntity::MakeShared<FBezierSurface>(Archive);
+	case ESurface::Cone:
+		return FEntity::MakeShared<FConeSurface>(Archive);
+	case ESurface::Composite:
+		return FEntity::MakeShared<FCompositeSurface>(Archive);
+	case ESurface::Coons:
+		return FEntity::MakeShared<FCoonsSurface>(Archive);
+	case ESurface::Cylinder:
+		return FEntity::MakeShared<FCylinderSurface>(Archive);
+	case ESurface::Nurbs:
+		return FEntity::MakeShared<FNURBSSurface>(Archive);
+	case ESurface::Offset:
+		return FEntity::MakeShared<FOffsetSurface>(Archive);
+	case ESurface::Plane:
+		return FEntity::MakeShared<FPlaneSurface>(Archive);
+	case ESurface::Revolution:
+		return FEntity::MakeShared<FRevolutionSurface>(Archive);
+	case ESurface::Ruled:
+		return FEntity::MakeShared<FRuledSurface>(Archive);
+	case ESurface::Sphere:
+		return FEntity::MakeShared<FSphericalSurface>(Archive);
+	case ESurface::TabulatedCylinder:
+		return FEntity::MakeShared<FTabulatedCylinderSurface>(Archive);
+	case ESurface::Torus:
+		return FEntity::MakeShared<FTorusSurface>(Archive);
+	default:
 		return TSharedPtr<FSurface>();
 	}
+	return TSharedPtr<FSurface>();
+}
 
-} // namespace CADKernel
+} // namespace UE::CADKernel
 
diff --git a/Engine/Source/Runtime/Datasmith/CADKernel/Private/CADKernel/Geo/Surfaces/TabulatedCylinderSurface.h b/Engine/Source/Runtime/Datasmith/CADKernel/Private/CADKernel/Geo/Surfaces/TabulatedCylinderSurface.h
--- a/Engine/Source/Runtime/Datasmith/CADKernel/Private/CADKernel/Geo/Surfaces/TabulatedCylinderSurface.h
+++ b/Engine/Source/Runtime/Datasmith/CADKernel/Private/CADKernel/Geo/Surfaces/TabulatedCylinderSurface.h
@@ -4,77 +4,77 @@
 #include "CADKernel/Geo/Surfaces/Surface.h"
 #include "CADKernel/Math/Point.h"
 
-namespace CADKernel
+namespace UE::CADKernel
 {
-	class FCurve;
+class FCurve;
 
-	class CADKERNEL_API FTabulatedCylinderSurface : public FSurface
+class CADKERNEL_API FTabulatedCylinderSurface : public FSurface
+{
+	friend FEntity;
+
+protected:
+	TSharedPtr<FCurve> GuideCurve;
+	FPoint DirectorVector;
+
+	/**
+	 * A tabulated cylinder is a surface formed by moving a curve segment called GuideCurve parallel to
+	 * itself along a vector called DirectorVector.
+	 *
+	 * The bounds of the cylinder are defined as follow:
+	 * Bounds[EIso::IsoU] = GuideCurve bounds;
+	 * Bounds[EIso::IsoV].Min = InVMin;
+	 * Bounds[EIso::IsoV].Max = InVMax;
+	 */
+	FTabulatedCylinderSurface(const double InToleranceGeometric, TSharedPtr<FCurve> InGuideCurve, const FPoint& InDirectorVector, double InVMin = 0., double InVMax = 1.)
+		: FSurface(InToleranceGeometric, 0.0, 1.0, InVMin, InVMax)
+		, GuideCurve(InGuideCurve)
+		, DirectorVector(InDirectorVector)
+	{
+		ComputeDefaultMinToleranceIso();
+	}
+
+	FTabulatedCylinderSurface() = default;
+
+public:
+
+	virtual void Serialize(FCADKernelArchive& Ar) override
+	{
+		FSurface::Serialize(Ar);
+		SerializeIdent(Ar, GuideCurve);
+		Ar << DirectorVector;
+	}
+
+	virtual void SpawnIdent(FDatabase& Database) override;
+	virtual void ResetMarkersRecursively() override;
+
+	ESurface GetSurfaceType() const
 	{
-		friend FEntity;
-
-	protected:
-		TSharedPtr<FCurve> GuideCurve;
-		FPoint DirectorVector;
-
-		/**
-		 * A tabulated cylinder is a surface formed by moving a curve segment called GuideCurve parallel to
-		 * itself along a vector called DirectorVector. 
-		 *
-		 * The bounds of the cylinder are defined as follow:
-		 * Bounds[EIso::IsoU] = GuideCurve bounds;
-		 * Bounds[EIso::IsoV].Min = InVMin;
-		 * Bounds[EIso::IsoV].Max = InVMax;
-		 */
-		FTabulatedCylinderSurface(const double InToleranceGeometric, TSharedPtr<FCurve> InGuideCurve, const FPoint& InDirectorVector, double InVMin = 0., double InVMax = 1.)
-			: FSurface(InToleranceGeometric, 0.0, 1.0, InVMin, InVMax)
-			, GuideCurve(InGuideCurve)
-			, DirectorVector(InDirectorVector)
-		{
-			ComputeDefaultMinToleranceIso();
-		}
-
-		FTabulatedCylinderSurface() = default;
-
-	public:
-
-		virtual void Serialize(FCADKernelArchive& Ar) override
-		{
-			FSurface::Serialize(Ar);
-			SerializeIdent(Ar, GuideCurve);
-			Ar << DirectorVector;
-		}	
-
-		virtual void SpawnIdent(FDatabase& Database) override;
-		virtual void ResetMarkersRecursively() override;
-		
-		ESurface GetSurfaceType() const
-		{
-			return ESurface::TabulatedCylinder;
-		}
-
-		TSharedPtr<FCurve> GetGuideCurve() const
-		{
-			return GuideCurve;
-		}
-
-		FPoint GetDirectorVector() const
-		{
-			return DirectorVector;
-		}
-
-		virtual void EvaluatePoint(const FPoint2D& InSurfacicCoordinate, FSurfacicPoint& OutPoint3D, int32 InDerivativeOrder = 0) const override;
-		virtual void EvaluatePointGrid(const FCoordinateGrid& Coordinates, FSurfacicSampling& OutPoints, bool bComputeNormals = false) const override;
-
-		virtual void LinesNotDerivables(const FSurfacicBoundary& Bounds, int32 InDerivativeOrder, FCoordinateGrid& OutNotDerivables) const override;
-
-		virtual void Presample(const FSurfacicBoundary& InBoundaries, FCoordinateGrid& OutCoordinates) override;
-
-		virtual TSharedPtr<FEntityGeom> ApplyMatrix(const FMatrixH& InMatrix) const override;
+		return ESurface::TabulatedCylinder;
+	}
+
+	TSharedPtr<FCurve> GetGuideCurve() const
+	{
+		return GuideCurve;
+	}
+
+	FPoint GetDirectorVector() const
+	{
+		return DirectorVector;
+	}
+
+	virtual void EvaluatePoint(const FPoint2D& InSurfacicCoordinate, FSurfacicPoint& OutPoint3D, int32 InDerivativeOrder = 0) const override;
+	virtual void EvaluatePointGrid(const FCoordinateGrid& Coordinates, FSurfacicSampling& OutPoints, bool bComputeNormals = false) const override;
+
+	virtual void LinesNotDerivables(const FSurfacicBoundary& Bounds, int32 InDerivativeOrder, FCoordinateGrid& OutNotDerivables) const override;
+
+	virtual void Presample(const FSurfacicBoundary& InBoundaries, FCoordinateGrid& OutCoordinates) override;
+
+	virtual TSharedPtr<FEntityGeom> ApplyMatrix(const FMatrixH& InMatrix) const override;
 
 #ifdef CADKERNEL_DEV
-		virtual FInfoEntity& GetInfo(FInfoEntity&) const override;
+	virtual FInfoEntity& GetInfo(FInfoEntity&) const override;
 #endif
 
-	};
+};
 
-} // namespace CADKernel
+} // namespace UE::CADKernel
diff --git a/Engine/Source/Runtime/Datasmith/CADKernel/Private/CADKernel/Geo/Surfaces/TorusSurface.h b/Engine/Source/Runtime/Datasmith/CADKernel/Private/CADKernel/Geo/Surfaces/TorusSurface.h
--- a/Engine/Source/Runtime/Datasmith/CADKernel/Private/CADKernel/Geo/Surfaces/TorusSurface.h
+++ b/Engine/Source/Runtime/Datasmith/CADKernel/Private/CADKernel/Geo/Surfaces/TorusSurface.h
@@ -3,151 +3,151 @@
 
 #include "CADKernel/Geo/Surfaces/Surface.h"
 
-namespace CADKernel
+namespace UE::CADKernel
 {
-	class CADKERNEL_API FTorusSurface : public FSurface
+class CADKERNEL_API FTorusSurface : public FSurface
+{
+	friend FEntity;
+
+protected:
+	FMatrixH Matrix;
+	double MajorRadius;
+	double MinorRadius;
+
+	/**
+	 * A torus is the solid formed by revolving a circular disc about a specified coplanar axis.
+	 * MajorRadius is the distance from the axis to the center of the defining disc, and MinorRadius is the radius of the defining disc,
+	 * where MajorRadius > MinorRadius > 0.0.
+	 *
+	 * The torus computed at the origin with Z axis.
+	 * It is placed at its final position and orientation by the Matrix
+	 */
+	FTorusSurface(const double InToleranceGeometric, const FMatrixH& InMatrix, double InMajorRadius, double InMinorRadius, double InMajorStartAngle = 0.0, double InMajorEndAngle = DOUBLE_TWO_PI, double InMinorStartAngle = 0.0, double InMinorEndAngle = DOUBLE_TWO_PI)
+		: FTorusSurface(InToleranceGeometric, InMatrix, InMajorRadius, InMinorRadius, FSurfacicBoundary(InMajorStartAngle, InMajorEndAngle, InMinorStartAngle, InMinorEndAngle))
 	{
-		friend FEntity;
-
-	protected:
-		FMatrixH Matrix;
-		double MajorRadius;
-		double MinorRadius;
-
-		/**
-		 * A torus is the solid formed by revolving a circular disc about a specified coplanar axis. 
-		 * MajorRadius is the distance from the axis to the center of the defining disc, and MinorRadius is the radius of the defining disc,
-		 * where MajorRadius > MinorRadius > 0.0. 
-		 *
-		 * The torus computed at the origin with Z axis.
-		 * It is placed at its final position and orientation by the Matrix
-		 */
-		FTorusSurface(const double InToleranceGeometric, const FMatrixH& InMatrix, double InMajorRadius, double InMinorRadius, double InMajorStartAngle = 0.0, double InMajorEndAngle = DOUBLE_TWO_PI, double InMinorStartAngle = 0.0, double InMinorEndAngle = DOUBLE_TWO_PI)
-			: FTorusSurface(InToleranceGeometric, InMatrix, InMajorRadius, InMinorRadius, FSurfacicBoundary(InMajorStartAngle, InMajorEndAngle, InMinorStartAngle, InMinorEndAngle))
-		{
-		}
+	}
+
+	/**
+	 * A torus is the solid formed by revolving a circular disc about a specified coplanar axis.
+	 * MajorRadius is the distance from the axis to the center of the defining disc, and MinorRadius is the radius of the defining disc,
+	 * where MajorRadius > MinorRadius > 0.0.
+	 *
+	 * The torus computed at the origin with Z axis.
+	 * It is placed at its final position and orientation by the Matrix
+	 *
+	 * The bounds of the cone are defined as follow:
+	 * Bounds[EIso::IsoU].Min = MajorStartAngle
+	 * Bounds[EIso::IsoU].Max = MajorEndAngle
+	 * Bounds[EIso::IsoV].Min = MinorStartAngle
+	 * Bounds[EIso::IsoV].Max = MinorEndAngle
+	 */
+	FTorusSurface(const double InToleranceGeometric, const FMatrixH& InMatrix, double InMajorRadius, double InMinorRadius, const FSurfacicBoundary& InBoundary)
+		: FSurface(InToleranceGeometric, InBoundary)
+		, Matrix(InMatrix)
+		, MajorRadius(InMajorRadius)
+		, MinorRadius(InMinorRadius)
+	{
+		ComputeMinToleranceIso();
+	}
 
-		/**
-		 * A torus is the solid formed by revolving a circular disc about a specified coplanar axis.
-		 * MajorRadius is the distance from the axis to the center of the defining disc, and MinorRadius is the radius of the defining disc,
-		 * where MajorRadius > MinorRadius > 0.0.
-		 *
-		 * The torus computed at the origin with Z axis.
-		 * It is placed at its final position and orientation by the Matrix
-		 *
-		 * The bounds of the cone are defined as follow:
-		 * Bounds[EIso::IsoU].Min = MajorStartAngle
-		 * Bounds[EIso::IsoU].Max = MajorEndAngle
-		 * Bounds[EIso::IsoV].Min = MinorStartAngle
-		 * Bounds[EIso::IsoV].Max = MinorEndAngle
-		 */
-		FTorusSurface(const double InToleranceGeometric, const FMatrixH & InMatrix, double InMajorRadius, double InMinorRadius, const FSurfacicBoundary& InBoundary)
-			: FSurface(InToleranceGeometric, InBoundary)
-			, Matrix(InMatrix)
-			, MajorRadius(InMajorRadius)
-			, MinorRadius(InMinorRadius)
-		{
-			ComputeMinToleranceIso();
-		}
+	FTorusSurface() = default;
 
-		FTorusSurface() = default;
+	void ComputeMinToleranceIso()
+	{
+		double Tolerance2DU = Tolerance3D / MajorRadius;
+		double Tolerance2DV = Tolerance3D / MinorRadius;
 
-		void ComputeMinToleranceIso()
-		{
-			double Tolerance2DU = Tolerance3D / MajorRadius;
-			double Tolerance2DV = Tolerance3D / MinorRadius;
+		FPoint Origin = Matrix.Multiply(FPoint::ZeroPoint);
 
-			FPoint Origin = Matrix.Multiply(FPoint::ZeroPoint);
+		FPoint Point2DU{ 1 , 0, 0 };
+		FPoint Point2DV{ 0, 1, 0 };
 
-			FPoint Point2DU{ 1 , 0, 0 };
-			FPoint Point2DV{ 0, 1, 0 };
+		Tolerance2DU /= ComputeScaleAlongAxis(Point2DU, Matrix, Origin);
+		Tolerance2DV /= ComputeScaleAlongAxis(Point2DV, Matrix, Origin);
 
-			Tolerance2DU /= ComputeScaleAlongAxis(Point2DU, Matrix, Origin);
-			Tolerance2DV /= ComputeScaleAlongAxis(Point2DV, Matrix, Origin);
+		MinToleranceIso.Set(FMath::Abs(Tolerance2DU), FMath::Abs(Tolerance2DV));
+	}
 
-			MinToleranceIso.Set(FMath::Abs(Tolerance2DU), FMath::Abs(Tolerance2DV));
-		}
-		
-	public:
+public:
 
-		virtual void Serialize(FCADKernelArchive& Ar) override
-		{
-			FSurface::Serialize(Ar);
-			Ar << Matrix;
-			Ar << MajorRadius;
-			Ar << MinorRadius;
-		}
+	virtual void Serialize(FCADKernelArchive& Ar) override
+	{
+		FSurface::Serialize(Ar);
+		Ar << Matrix;
+		Ar << MajorRadius;
+		Ar << MinorRadius;
+	}
 
-		ESurface GetSurfaceType() const
-		{
-			return ESurface::Torus;
-		}
+	ESurface GetSurfaceType() const
+	{
+		return ESurface::Torus;
+	}
 
 #ifdef CADKERNEL_DEV
-		virtual FInfoEntity& GetInfo(FInfoEntity&) const override;
+	virtual FInfoEntity& GetInfo(FInfoEntity&) const override;
 #endif
 
-		virtual TSharedPtr<FEntityGeom> ApplyMatrix(const FMatrixH& InMatrix) const override;
+	virtual TSharedPtr<FEntityGeom> ApplyMatrix(const FMatrixH& InMatrix) const override;
 
-		virtual void EvaluatePoint(const FPoint2D& InSurfacicCoordinate, FSurfacicPoint& OutPoint3D, int32 InDerivativeOrder = 0) const override
-		{
-			const double CosU = cos(InSurfacicCoordinate.U);
-			const double CosV = cos(InSurfacicCoordinate.V);
+	virtual void EvaluatePoint(const FPoint2D& InSurfacicCoordinate, FSurfacicPoint& OutPoint3D, int32 InDerivativeOrder = 0) const override
+	{
+		const double CosU = cos(InSurfacicCoordinate.U);
+		const double CosV = cos(InSurfacicCoordinate.V);
 
-			const double SinU = sin(InSurfacicCoordinate.U);
-			const double SinV = sin(InSurfacicCoordinate.V);
-			double Rho = (MajorRadius + MinorRadius * CosV);
+		const double SinU = sin(InSurfacicCoordinate.U);
+		const double SinV = sin(InSurfacicCoordinate.V);
+		double Rho = (MajorRadius + MinorRadius * CosV);
 
-			OutPoint3D.DerivativeOrder = InDerivativeOrder;
-			OutPoint3D.Point.Set(Rho * CosU, Rho * SinU, MinorRadius * SinV);
+		OutPoint3D.DerivativeOrder = InDerivativeOrder;
+		OutPoint3D.Point.Set(Rho * CosU, Rho * SinU, MinorRadius * SinV);
 
-			OutPoint3D.Point = Matrix.Multiply(OutPoint3D.Point);
+		OutPoint3D.Point = Matrix.Multiply(OutPoint3D.Point);
 
-			if (InDerivativeOrder > 0)
+		if (InDerivativeOrder > 0)
+		{
+			OutPoint3D.GradientU = FPoint((MajorRadius + MinorRadius * CosV) * -SinU, (MajorRadius + MinorRadius * CosV) * CosU, 0.0);
+			OutPoint3D.GradientV = FPoint((MinorRadius * -SinV) * CosU, (MinorRadius * -SinV) * SinU, MinorRadius * CosV);
+
+			OutPoint3D.GradientU = Matrix.MultiplyVector(OutPoint3D.GradientU);
+			OutPoint3D.GradientV = Matrix.MultiplyVector(OutPoint3D.GradientV);
+
+			if (InDerivativeOrder > 1)
 			{
-				OutPoint3D.GradientU = FPoint((MajorRadius + MinorRadius * CosV) * -SinU, (MajorRadius + MinorRadius * CosV) * CosU, 0.0);
-				OutPoint3D.GradientV = FPoint((MinorRadius * -SinV) * CosU, (MinorRadius * -SinV) * SinU, MinorRadius * CosV);
-
-				OutPoint3D.GradientU = Matrix.MultiplyVector(OutPoint3D.GradientU);
-				OutPoint3D.GradientV = Matrix.MultiplyVector(OutPoint3D.GradientV);
-
-				if (InDerivativeOrder > 1)
-				{
-					OutPoint3D.LaplacianU = FPoint((MajorRadius + MinorRadius * CosV) * -CosU, (MajorRadius + MinorRadius * CosV) * -SinU, 0.0);
-					OutPoint3D.LaplacianV = FPoint((MinorRadius * -CosV) * CosU, (MinorRadius * -CosV) * SinU, MinorRadius * -SinV);
-					OutPoint3D.LaplacianUV = FPoint((MinorRadius * -SinV) * -SinU, (MinorRadius * -SinV) * CosU, 0.0);
-
-					OutPoint3D.LaplacianU = Matrix.MultiplyVector(OutPoint3D.LaplacianU);
-					OutPoint3D.LaplacianV = Matrix.MultiplyVector(OutPoint3D.LaplacianV);
-					OutPoint3D.LaplacianUV = Matrix.MultiplyVector(OutPoint3D.LaplacianUV);
-				}
+				OutPoint3D.LaplacianU = FPoint((MajorRadius + MinorRadius * CosV) * -CosU, (MajorRadius + MinorRadius * CosV) * -SinU, 0.0);
+				OutPoint3D.LaplacianV = FPoint((MinorRadius * -CosV) * CosU, (MinorRadius * -CosV) * SinU, MinorRadius * -SinV);
+				OutPoint3D.LaplacianUV = FPoint((MinorRadius * -SinV) * -SinU, (MinorRadius * -SinV) * CosU, 0.0);
+
+				OutPoint3D.LaplacianU = Matrix.MultiplyVector(OutPoint3D.LaplacianU);
+				OutPoint3D.LaplacianV = Matrix.MultiplyVector(OutPoint3D.LaplacianV);
+				OutPoint3D.LaplacianUV = Matrix.MultiplyVector(OutPoint3D.LaplacianUV);
 			}
 		}
-	
-		virtual void EvaluatePointGrid(const FCoordinateGrid& Coordinates, FSurfacicSampling& OutPoints, bool bComputeNormals = false) const override;
+	}
 
-		virtual void EvaluatePointGridInCylindricalSpace(const FCoordinateGrid& Coordinates, TArray<FPoint2D>&) const override;
+	virtual void EvaluatePointGrid(const FCoordinateGrid& Coordinates, FSurfacicSampling& OutPoints, bool bComputeNormals = false) const override;
 
-		virtual FPoint2D EvaluatePointInCylindricalSpace(const FPoint2D& InSurfacicCoordinate) const override
-		{
-			const double CosU = cos(InSurfacicCoordinate.U);
-			const double CosV = cos(InSurfacicCoordinate.V);
+	virtual void EvaluatePointGridInCylindricalSpace(const FCoordinateGrid& Coordinates, TArray<FPoint2D>&) const override;
 
-			const double SinU = sin(InSurfacicCoordinate.U);
+	virtual FPoint2D EvaluatePointInCylindricalSpace(const FPoint2D& InSurfacicCoordinate) const override
+	{
+		const double CosU = cos(InSurfacicCoordinate.U);
+		const double CosV = cos(InSurfacicCoordinate.V);
 
-			double Rho = (MajorRadius + MinorRadius * CosV);
+		const double SinU = sin(InSurfacicCoordinate.U);
 
-			double SwapOrientation = (InSurfacicCoordinate.V < DOUBLE_PI&& InSurfacicCoordinate.V >=0) ? 1.0 : -1.0;
+		double Rho = (MajorRadius + MinorRadius * CosV);
 
-			return FPoint2D(Rho * CosU * SwapOrientation, Rho * SinU);
-		}
+		double SwapOrientation = (InSurfacicCoordinate.V < DOUBLE_PI&& InSurfacicCoordinate.V >= 0) ? 1.0 : -1.0;
 
-		virtual void Presample(const FSurfacicBoundary& InBoundaries, FCoordinateGrid& OutCoordinates) override
-		{
-			PresampleIsoCircle(InBoundaries, OutCoordinates, EIso::IsoU);
-			PresampleIsoCircle(InBoundaries, OutCoordinates, EIso::IsoV);
-		}
+		return FPoint2D(Rho * CosU * SwapOrientation, Rho * SinU);
+	}
+
+	virtual void Presample(const FSurfacicBoundary& InBoundaries, FCoordinateGrid& OutCoordinates) override
+	{
+		PresampleIsoCircle(InBoundaries, OutCoordinates, EIso::IsoU);
+		PresampleIsoCircle(InBoundaries, OutCoordinates, EIso::IsoV);
+	}
 
-	};
+};
 }
 
