Rename CADKernel namespace into UE::CADKernel

diff --git a/Engine/Source/Runtime/Datasmith/CADKernel/Private/CADKernel/Math/Geometry.cpp b/Engine/Source/Runtime/Datasmith/CADKernel/Private/CADKernel/Math/Geometry.cpp
--- a/Engine/Source/Runtime/Datasmith/CADKernel/Private/CADKernel/Math/Geometry.cpp
+++ b/Engine/Source/Runtime/Datasmith/CADKernel/Private/CADKernel/Math/Geometry.cpp
@@ -8,129 +8,129 @@
 #include "CADKernel/Math/MatrixH.h"
 #include "CADKernel/Utils/Util.h"
 
-namespace CADKernel
+namespace UE::CADKernel
 {
-	double ComputeCurvature(const FPoint& Gradient, const FPoint& Laplacian)
-	{
-		FPoint GradientCopy = Gradient;
-		FPoint LaplacianCopy = Laplacian;
-		GradientCopy.Normalize();
-		LaplacianCopy.Normalize();
-		FPoint Normal = GradientCopy ^ LaplacianCopy;
-		return (Normal.Length() * Laplacian.Length()) / (Gradient.Length() * Gradient.Length());
-	}
+double ComputeCurvature(const FPoint& Gradient, const FPoint& Laplacian)
+{
+	FPoint GradientCopy = Gradient;
+	FPoint LaplacianCopy = Laplacian;
+	GradientCopy.Normalize();
+	LaplacianCopy.Normalize();
+	FPoint Normal = GradientCopy ^ LaplacianCopy;
+	return (Normal.Length() * Laplacian.Length()) / (Gradient.Length() * Gradient.Length());
+}
 
-	double ComputeCurvature(const FPoint& Normal, const FPoint& Gradient, const FPoint& Laplacian)
-	{
-		FPoint GradientCopy = Gradient;
-		FPoint LaplacianCopy = Laplacian;
-		FPoint NormalCopy = Normal;
-		GradientCopy.Normalize();
-		LaplacianCopy.Normalize();
-		NormalCopy.Normalize();
-		FPoint Coef = (LaplacianCopy ^ GradientCopy) ^ NormalCopy;
-		return (Coef.Length() * Laplacian.Length()) / Gradient.SquareLength();
-	}
+double ComputeCurvature(const FPoint& Normal, const FPoint& Gradient, const FPoint& Laplacian)
+{
+	FPoint GradientCopy = Gradient;
+	FPoint LaplacianCopy = Laplacian;
+	FPoint NormalCopy = Normal;
+	GradientCopy.Normalize();
+	LaplacianCopy.Normalize();
+	NormalCopy.Normalize();
+	FPoint Coef = (LaplacianCopy ^ GradientCopy) ^ NormalCopy;
+	return (Coef.Length() * Laplacian.Length()) / Gradient.SquareLength();
+}
 
-	void FindLoopIntersectionsWithIso(const EIso Iso, const double IsoParameter, const TArray<TArray<FPoint2D>>& Loops, TArray<double>& OutIntersections)
-	{
-		OutIntersections.Empty(8);
+void FindLoopIntersectionsWithIso(const EIso Iso, const double IsoParameter, const TArray<TArray<FPoint2D>>& Loops, TArray<double>& OutIntersections)
+{
+	OutIntersections.Empty(8);
 
-		int32 UIndex = Iso == EIso::IsoU ? 0 : 1;
-		int32 VIndex = Iso == EIso::IsoU ? 1 : 0;
+	int32 UIndex = Iso == EIso::IsoU ? 0 : 1;
+	int32 VIndex = Iso == EIso::IsoU ? 1 : 0;
 
-		TFunction<void(const FPoint2D&, const FPoint2D&)> ComputeIntersection = [&](const FPoint2D& Point1, const FPoint2D& Point2)
+	TFunction<void(const FPoint2D&, const FPoint2D&)> ComputeIntersection = [&](const FPoint2D& Point1, const FPoint2D& Point2)
+	{
+		if (IsoParameter > Point1[UIndex] && IsoParameter <= Point2[UIndex])
 		{
-			if (IsoParameter > Point1[UIndex] && IsoParameter <= Point2[UIndex])
-			{
-				double Intersection = (IsoParameter - Point1[UIndex]) / (Point2[UIndex] - Point1[UIndex]) * (Point2[VIndex] - Point1[VIndex]) + Point1[VIndex];
-				OutIntersections.Add((IsoParameter - Point1[UIndex]) / (Point2[UIndex] - Point1[UIndex]) * (Point2[VIndex] - Point1[VIndex]) + Point1[VIndex]);
-			}
-		};
+			double Intersection = (IsoParameter - Point1[UIndex]) / (Point2[UIndex] - Point1[UIndex]) * (Point2[VIndex] - Point1[VIndex]) + Point1[VIndex];
+			OutIntersections.Add((IsoParameter - Point1[UIndex]) / (Point2[UIndex] - Point1[UIndex]) * (Point2[VIndex] - Point1[VIndex]) + Point1[VIndex]);
+		}
+	};
 
-		for (const TArray<FPoint2D>& Loop : Loops)
+	for (const TArray<FPoint2D>& Loop : Loops)
+	{
+		const FPoint2D* Point1 = &Loop.Last();
+		for (const FPoint2D& Point2 : Loop)
 		{
-			const FPoint2D* Point1 = &Loop.Last();
-			for (const FPoint2D& Point2 : Loop)
+			if (!FMath::IsNearlyEqual((*Point1)[UIndex], Point2[UIndex]))
 			{
-				if (!FMath::IsNearlyEqual((*Point1)[UIndex], Point2[UIndex]))
+				if ((*Point1)[UIndex] < Point2[UIndex])
 				{
-					if ((*Point1)[UIndex] < Point2[UIndex])
-					{
-						ComputeIntersection(*Point1, Point2);
-					}
-					else
-					{
-						ComputeIntersection(Point2, *Point1);
-					}
+					ComputeIntersection(*Point1, Point2);
+				}
+				else
+				{
+					ComputeIntersection(Point2, *Point1);
 				}
-				Point1 = &Point2;
 			}
+			Point1 = &Point2;
 		}
-		Algo::Sort(OutIntersections);
 	}
+	Algo::Sort(OutIntersections);
+}
 
-	bool IntersectSegments2D(const TSegment<FPoint2D>& SegmentAB, const TSegment<FPoint2D>& SegmentCD)
-	{
-		constexpr const double Min = -DOUBLE_SMALL_NUMBER;
-		constexpr const double Max = 1. + DOUBLE_SMALL_NUMBER;
+bool IntersectSegments2D(const TSegment<FPoint2D>& SegmentAB, const TSegment<FPoint2D>& SegmentCD)
+{
+	constexpr const double Min = -DOUBLE_SMALL_NUMBER;
+	constexpr const double Max = 1. + DOUBLE_SMALL_NUMBER;
 
-		TFunction<bool(double, double, double, double)> Intersect = [](double A, double B, double C, double D)
-		{
-			return !((D < A) || (B < C));
-		};
+	TFunction<bool(double, double, double, double)> Intersect = [](double A, double B, double C, double D)
+	{
+		return !((D < A) || (B < C));
+	};
 
-		TFunction<bool(double, double, double, double)> TestWhenParallel = [&](double A, double B, double C, double D)
+	TFunction<bool(double, double, double, double)> TestWhenParallel = [&](double A, double B, double C, double D)
+	{
+		if (A < B)
 		{
-			if (A < B)
+			if (C < D)
 			{
-				if (C < D)
-				{
-					return Intersect(A, B, C, D);
-				}
-				else
-				{
-					return Intersect(A, B, D, C);
-				}
+				return Intersect(A, B, C, D);
 			}
 			else
 			{
-				if (C < D)
-				{
-					return Intersect(B, A, C, D);
-				}
-				else
-				{
-					return Intersect(B, A, D, C);
-				}
+				return Intersect(A, B, D, C);
 			}
-		};
-
-		FPoint2D AB = SegmentAB[1] - SegmentAB[0];
-		FPoint2D CD = SegmentCD[1] - SegmentCD[0];
-		FPoint2D CA = SegmentAB[0] - SegmentCD[0];
-
-		const double ParallelCoef = CD ^ AB;
-		if (FMath::IsNearlyZero(ParallelCoef))
+		}
+		else
 		{
-			const double ParallelCoef2 = CA ^ AB;
-			if (!FMath::IsNearlyZero(ParallelCoef2))
+			if (C < D)
 			{
-				return false;
-			}
-
-			if (fabs(AB.U) > fabs(AB.V))
-			{
-				return TestWhenParallel(SegmentAB[0].U, SegmentAB[1].U, SegmentCD[0].U, SegmentCD[1].U);
+				return Intersect(B, A, C, D);
 			}
 			else
 			{
-				return TestWhenParallel(SegmentAB[0].V, SegmentAB[1].V, SegmentCD[0].V, SegmentCD[1].V);
+				return Intersect(B, A, D, C);
 			}
 		}
+	};
+
+	FPoint2D AB = SegmentAB[1] - SegmentAB[0];
+	FPoint2D CD = SegmentCD[1] - SegmentCD[0];
+	FPoint2D CA = SegmentAB[0] - SegmentCD[0];
+
+	const double ParallelCoef = CD ^ AB;
+	if (FMath::IsNearlyZero(ParallelCoef))
+	{
+		const double ParallelCoef2 = CA ^ AB;
+		if (!FMath::IsNearlyZero(ParallelCoef2))
+		{
+			return false;
+		}
 
-		const double ABIntersectionCoordinate = (CA ^ CD) / ParallelCoef;
-		const double CDIntersectionCoordinate = (CA ^ AB) / ParallelCoef;
-		return (ABIntersectionCoordinate <= Max && ABIntersectionCoordinate >= Min && CDIntersectionCoordinate <= Max && CDIntersectionCoordinate >= Min);
+		if (fabs(AB.U) > fabs(AB.V))
+		{
+			return TestWhenParallel(SegmentAB[0].U, SegmentAB[1].U, SegmentCD[0].U, SegmentCD[1].U);
+		}
+		else
+		{
+			return TestWhenParallel(SegmentAB[0].V, SegmentAB[1].V, SegmentCD[0].V, SegmentCD[1].V);
+		}
 	}
+
+	const double ABIntersectionCoordinate = (CA ^ CD) / ParallelCoef;
+	const double CDIntersectionCoordinate = (CA ^ AB) / ParallelCoef;
+	return (ABIntersectionCoordinate <= Max && ABIntersectionCoordinate >= Min && CDIntersectionCoordinate <= Max && CDIntersectionCoordinate >= Min);
+}
 }
diff --git a/Engine/Source/Runtime/Datasmith/CADKernel/Private/CADKernel/Math/MatrixH.cpp b/Engine/Source/Runtime/Datasmith/CADKernel/Private/CADKernel/Math/MatrixH.cpp
--- a/Engine/Source/Runtime/Datasmith/CADKernel/Private/CADKernel/Math/MatrixH.cpp
+++ b/Engine/Source/Runtime/Datasmith/CADKernel/Private/CADKernel/Math/MatrixH.cpp
@@ -5,251 +5,251 @@
 #include "CADKernel/Math/MathConst.h"
 #include "CADKernel/UI/Message.h"
 
-namespace CADKernel
+namespace UE::CADKernel
 {
 
-	const FMatrixH FMatrixH::Identity;
+const FMatrixH FMatrixH::Identity;
 
-	void FMatrixH::BuildChangeOfCoordinateSystemMatrix(const FPoint& Xaxis, const FPoint& Yaxis, const FPoint& Zaxis, const FPoint& Origin)
-	{
-		Get(0, 0) = Xaxis[0];
-		Get(1, 0) = Xaxis[1];
-		Get(2, 0) = Xaxis[2];
-		Get(3, 0) = 0.0;
-
-		Get(0, 1) = Yaxis[0];
-		Get(1, 1) = Yaxis[1];
-		Get(2, 1) = Yaxis[2];
-		Get(3, 1) = 0.0;
-
-		Get(0, 2) = Zaxis[0];
-		Get(1, 2) = Zaxis[1];
-		Get(2, 2) = Zaxis[2];
-		Get(3, 2) = 0.0;
-
-		Get(0, 3) = Origin[0];
-		Get(1, 3) = Origin[1];
-		Get(2, 3) = Origin[2];
-		Get(3, 3) = 1.0;
-	}
+void FMatrixH::BuildChangeOfCoordinateSystemMatrix(const FPoint& Xaxis, const FPoint& Yaxis, const FPoint& Zaxis, const FPoint& Origin)
+{
+	Get(0, 0) = Xaxis[0];
+	Get(1, 0) = Xaxis[1];
+	Get(2, 0) = Xaxis[2];
+	Get(3, 0) = 0.0;
+
+	Get(0, 1) = Yaxis[0];
+	Get(1, 1) = Yaxis[1];
+	Get(2, 1) = Yaxis[2];
+	Get(3, 1) = 0.0;
+
+	Get(0, 2) = Zaxis[0];
+	Get(1, 2) = Zaxis[1];
+	Get(2, 2) = Zaxis[2];
+	Get(3, 2) = 0.0;
+
+	Get(0, 3) = Origin[0];
+	Get(1, 3) = Origin[1];
+	Get(2, 3) = Origin[2];
+	Get(3, 3) = 1.0;
+}
+
+void FMatrixH::FromAxisOrigin(const FPoint& Axis, const FPoint& Origin)
+{
+	FPoint Zaxis = FPoint(0, 1, 0);
+	FPoint Xaxis = Zaxis ^ Axis;
 
-	void FMatrixH::FromAxisOrigin(const FPoint& Axis, const FPoint& Origin)
+	//on cherche le vecteur ox
+	Xaxis.Normalize();
+	if (FMath::Abs(Xaxis.Length()) < DOUBLE_SMALL_NUMBER)
 	{
-		FPoint Zaxis = FPoint(0, 1, 0);
-		FPoint Xaxis = Zaxis ^ Axis;
-
-		//on cherche le vecteur ox
+		Zaxis = FPoint(1, 0, 0);
+		Xaxis = Axis ^ Zaxis;
 		Xaxis.Normalize();
 		if (FMath::Abs(Xaxis.Length()) < DOUBLE_SMALL_NUMBER)
 		{
-			Zaxis = FPoint(1, 0, 0);
+			Zaxis = FPoint(0, 0, 1);
 			Xaxis = Axis ^ Zaxis;
-			Xaxis.Normalize();
-			if (FMath::Abs(Xaxis.Length()) < DOUBLE_SMALL_NUMBER)
-			{
-				Zaxis = FPoint(0, 0, 1);
-				Xaxis = Axis ^ Zaxis;
 
-				Xaxis.Normalize();
-				ensureCADKernel(FMath::Abs(Xaxis.Length()) > DOUBLE_SMALL_NUMBER);
-			}
+			Xaxis.Normalize();
+			ensureCADKernel(FMath::Abs(Xaxis.Length()) > DOUBLE_SMALL_NUMBER);
 		}
+	}
 
-		FPoint YAxis = Axis ^ Xaxis;
-		Zaxis = Axis;
+	FPoint YAxis = Axis ^ Xaxis;
+	Zaxis = Axis;
 
-		YAxis.Normalize();
-		Zaxis.Normalize();
+	YAxis.Normalize();
+	Zaxis.Normalize();
 
-		BuildChangeOfCoordinateSystemMatrix(Xaxis, YAxis, Zaxis, Origin);
-	}
+	BuildChangeOfCoordinateSystemMatrix(Xaxis, YAxis, Zaxis, Origin);
+}
 
-	FMatrixH FMatrixH::MakeRotationMatrix(double Angle, FPoint Axe)
-	{
-		FMatrixH Matrix;
-		Matrix.SetIdentity();
+FMatrixH FMatrixH::MakeRotationMatrix(double Angle, FPoint Axe)
+{
+	FMatrixH Matrix;
+	Matrix.SetIdentity();
 
-		ensureCADKernel(Axe.Length() > DOUBLE_SMALL_NUMBER);
-		Axe.Normalize();
+	ensureCADKernel(Axe.Length() > DOUBLE_SMALL_NUMBER);
+	Axe.Normalize();
 
-		Matrix.Get(0, 0) = Axe[0] * Axe[0] + (double)cos(Angle) * ((double)1.0 - Axe[0] * Axe[0]);
-		Matrix.Get(0, 1) = (double)(1.0 - cos(Angle)) * Axe[0] * Axe[1] - (double)sin(Angle) * Axe[2];
-		Matrix.Get(0, 2) = (double)(1.0 - cos(Angle)) * Axe[0] * Axe[2] + (double)sin(Angle) * Axe[1];
+	Matrix.Get(0, 0) = Axe[0] * Axe[0] + (double)cos(Angle) * ((double)1.0 - Axe[0] * Axe[0]);
+	Matrix.Get(0, 1) = (double)(1.0 - cos(Angle)) * Axe[0] * Axe[1] - (double)sin(Angle) * Axe[2];
+	Matrix.Get(0, 2) = (double)(1.0 - cos(Angle)) * Axe[0] * Axe[2] + (double)sin(Angle) * Axe[1];
 
-		Matrix.Get(1, 0) = (double)(1.0 - cos(Angle)) * Axe[1] * Axe[0] + (double)sin(Angle) * Axe[2];
-		Matrix.Get(1, 1) = Axe[1] * Axe[1] + (double)cos(Angle) * ((double)1.0 - Axe[1] * Axe[1]);
-		Matrix.Get(1, 2) = (double)(1.0 - cos(Angle)) * Axe[1] * Axe[2] - (double)sin(Angle) * Axe[0];
+	Matrix.Get(1, 0) = (double)(1.0 - cos(Angle)) * Axe[1] * Axe[0] + (double)sin(Angle) * Axe[2];
+	Matrix.Get(1, 1) = Axe[1] * Axe[1] + (double)cos(Angle) * ((double)1.0 - Axe[1] * Axe[1]);
+	Matrix.Get(1, 2) = (double)(1.0 - cos(Angle)) * Axe[1] * Axe[2] - (double)sin(Angle) * Axe[0];
 
-		Matrix.Get(2, 0) = (double)(1.0 - cos(Angle)) * Axe[2] * Axe[0] - (double)sin(Angle) * Axe[1];
-		Matrix.Get(2, 1) = (double)(1.0 - cos(Angle)) * Axe[2] * Axe[1] + (double)sin(Angle) * Axe[0];
-		Matrix.Get(2, 2) = Axe[2] * Axe[2] + (double)cos(Angle) * ((double)1.0 - Axe[2] * Axe[2]);
-		return Matrix;
-	}
+	Matrix.Get(2, 0) = (double)(1.0 - cos(Angle)) * Axe[2] * Axe[0] - (double)sin(Angle) * Axe[1];
+	Matrix.Get(2, 1) = (double)(1.0 - cos(Angle)) * Axe[2] * Axe[1] + (double)sin(Angle) * Axe[0];
+	Matrix.Get(2, 2) = Axe[2] * Axe[2] + (double)cos(Angle) * ((double)1.0 - Axe[2] * Axe[2]);
+	return Matrix;
+}
 
-	FMatrixH FMatrixH::MakeTranslationMatrix(const FPoint& Point)
-	{
-		FMatrixH Matrix;
-		Matrix.SetIdentity();
-		Matrix.Get(0, 3) = Point.X;
-		Matrix.Get(1, 3) = Point.Y;
-		Matrix.Get(2, 3) = Point.Z;
-		return Matrix;
-	}
+FMatrixH FMatrixH::MakeTranslationMatrix(const FPoint& Point)
+{
+	FMatrixH Matrix;
+	Matrix.SetIdentity();
+	Matrix.Get(0, 3) = Point.X;
+	Matrix.Get(1, 3) = Point.Y;
+	Matrix.Get(2, 3) = Point.Z;
+	return Matrix;
+}
+
+FMatrixH FMatrixH::MakeScaleMatrix(double XScale, double YScale, double ZScale)
+{
+	FMatrixH Matrix;
+	Matrix.SetIdentity();
+	Matrix.Get(0, 0) = XScale;
+	Matrix.Get(1, 1) = YScale;
+	Matrix.Get(2, 2) = ZScale;
+	return Matrix;
+}
+
+FMatrixH FMatrixH::MakeScaleMatrix(FPoint& Scale)
+{
+	return MakeScaleMatrix(Scale.X, Scale.Y, Scale.Z);
+}
 
-	FMatrixH FMatrixH::MakeScaleMatrix(double XScale, double YScale, double ZScale)
-	{
-		FMatrixH Matrix;
-		Matrix.SetIdentity();
-		Matrix.Get(0, 0) = XScale;
-		Matrix.Get(1, 1) = YScale;
-		Matrix.Get(2, 2) = ZScale;
-		return Matrix;
-	}
+void FMatrixH::Inverse()
+{
+	InverseMatrixN(Matrix, 4);
+}
 
-	FMatrixH FMatrixH::MakeScaleMatrix(FPoint& Scale)
+void FMatrixH::Print(EVerboseLevel level) const
+{
+	FMessage::Printf(level, TEXT(" - Matrix\n"));
+	for (int32 Row = 0; Row < 4; Row++)
 	{
-		return MakeScaleMatrix(Scale.X, Scale.Y, Scale.Z);
+		FMessage::Printf(level, TEXT("	- "));
+		for (int32 Column = 0; Column < 4; Column++)
+		{
+			FMessage::Printf(level, TEXT("%f "), Get(Row, Column));
+		}
+		FMessage::Printf(level, TEXT("\n"));
 	}
+}
 
-	void FMatrixH::Inverse()
+void InverseMatrixN(double* Matrice, int32 Rank)
+{
+	const double One = 1.0;
+	const double Zero = 0.0;
+
+	TArray<double> TempMatrix;
+	TempMatrix.Append(Matrice, Rank * Rank);
+
+	double Determinant = One;
+
+	TArray<int32> ColumnToRow;
+	ColumnToRow.SetNum(Rank);
+	for (int32 Index = 0; Index < Rank; Index++)
 	{
-		InverseMatrixN(Matrix, 4);
+		ColumnToRow[Index] = Index;
 	}
 
-	void FMatrixH::Print(EVerboseLevel level) const
+	for (int32 Column = 0; Column < Rank; ++Column)
 	{
-		FMessage::Printf(level, TEXT(" - Matrix\n"));
-		for (int32 Row = 0; Row < 4; Row++)
+		double Pivot = 0;
+
+		int32 Row = Column;
+		while (Row < Rank)
 		{
-			FMessage::Printf(level, TEXT("	- "));
-			for (int32 Column = 0; Column < 4; Column++)
+			Pivot = TempMatrix[Rank * Row + Column];
+			if (!FMath::IsNearlyZero(Pivot))
 			{
-				FMessage::Printf(level, TEXT("%f "), Get(Row, Column));
+				break;
 			}
-			FMessage::Printf(level, TEXT("\n"));
+			Row++;
 		}
-	}
-
-	void InverseMatrixN(double* Matrice, int32 Rank)
-	{
-		const double One = 1.0;
-		const double Zero = 0.0;
 
-		TArray<double> TempMatrix;
-		TempMatrix.Append(Matrice, Rank * Rank);
+		Determinant = Determinant * Pivot;
+		if (Row != Column)
+		{
+			Swap(ColumnToRow[Column], ColumnToRow[Row]);
+			for (int32 Index = 0; Index < Rank; Index++)
+			{
+				Swap(TempMatrix[Row * Rank + Index], TempMatrix[Column * Rank + Index]);
+			}
+			Determinant = -Determinant;
+		}
 
-		double Determinant = One;
+		double InvPivot = One / Pivot;
+		TempMatrix[Column * Rank + Column] = One;
 
-		TArray<int32> ColumnToRow;
-		ColumnToRow.SetNum(Rank);
 		for (int32 Index = 0; Index < Rank; Index++)
 		{
-			ColumnToRow[Index] = Index;
+			TempMatrix[Column * Rank + Index] = TempMatrix[Column * Rank + Index] * InvPivot;
 		}
 
-		for (int32 Column = 0; Column < Rank; ++Column)
+		for (Row = 0; Row < Rank; Row++)
 		{
-			double Pivot = 0;
-
-			int32 Row = Column;
-			while (Row < Rank)
+			if (Row == Column)
 			{
-				Pivot = TempMatrix[Rank * Row + Column];
-				if (!FMath::IsNearlyZero(Pivot))
-				{
-					break;
-				}
-				Row++;
-			}
-
-			Determinant = Determinant * Pivot;
-			if (Row != Column)
-			{
-				Swap(ColumnToRow[Column], ColumnToRow[Row]);
-				for (int32 Index = 0; Index < Rank; Index++)
-				{
-					Swap(TempMatrix[Row * Rank + Index], TempMatrix[Column * Rank + Index]);
-				}
-				Determinant = -Determinant;
+				continue;
 			}
 
-			double InvPivot = One / Pivot;
-			TempMatrix[Column * Rank + Column] = One;
+			double ValueRC = TempMatrix[Row * Rank + Column];
 
+			TempMatrix[Row * Rank + Column] = Zero;
 			for (int32 Index = 0; Index < Rank; Index++)
 			{
-				TempMatrix[Column * Rank + Index] = TempMatrix[Column * Rank + Index] * InvPivot;
-			}
-
-			for (Row = 0; Row < Rank; Row++)
-			{
-				if (Row == Column)
-				{
-					continue;
-				}
-
-				double ValueRC = TempMatrix[Row * Rank + Column];
-
-				TempMatrix[Row * Rank + Column] = Zero;
-				for (int32 Index = 0; Index < Rank; Index++)
-				{
-					TempMatrix[Row * Rank + Index] = TempMatrix[Row * Rank + Index] - ValueRC * TempMatrix[Column * Rank + Index];
-				}
+				TempMatrix[Row * Rank + Index] = TempMatrix[Row * Rank + Index] - ValueRC * TempMatrix[Column * Rank + Index];
 			}
 		}
+	}
 
-		for (int32 Column = 0; Column < Rank; Column++)
+	for (int32 Column = 0; Column < Rank; Column++)
+	{
+		int32 Row = Column;
+		while (Row < Rank)
 		{
-			int32 Row = Column;
-			while (Row < Rank)
+			if (ColumnToRow[Row] == Column)
 			{
-				if (ColumnToRow[Row] == Column)
-				{
-					break;
-				}
-				Row++;
-			}
-
-			if (Column == Row)
-			{
-				continue;
+				break;
 			}
+			Row++;
+		}
 
-			ColumnToRow[Row] = ColumnToRow[Column];
-			for (int32 Index = 0; Index < Rank; Index++)
-			{
-				Swap(TempMatrix[Index * Rank + Column], TempMatrix[Index * Rank + Row]);
-			}
+		if (Column == Row)
+		{
+			continue;
 		}
 
-		memcpy(Matrice, TempMatrix.GetData(), Rank * Rank * sizeof(double));
+		ColumnToRow[Row] = ColumnToRow[Column];
+		for (int32 Index = 0; Index < Rank; Index++)
+		{
+			Swap(TempMatrix[Index * Rank + Column], TempMatrix[Index * Rank + Row]);
+		}
 	}
 
-	void MatrixProduct(int32 ARowNum, int32 AColumnNum, int32 ResultRank, const double* MatrixA, const double* MatrixB, double* MatrixResult)
+	memcpy(Matrice, TempMatrix.GetData(), Rank * Rank * sizeof(double));
+}
+
+void MatrixProduct(int32 ARowNum, int32 AColumnNum, int32 ResultRank, const double* MatrixA, const double* MatrixB, double* MatrixResult)
+{
+	for (int32 RowA = 0; RowA < ARowNum; RowA++)
 	{
-		for (int32 RowA = 0; RowA < ARowNum; RowA++)
+		for (int32 ColumnB = 0; ColumnB < ResultRank; ColumnB++)
 		{
-			for (int32 ColumnB = 0; ColumnB < ResultRank; ColumnB++)
-			{
-				int32 ResultIndex = RowA * ResultRank + ColumnB;
+			int32 ResultIndex = RowA * ResultRank + ColumnB;
 
-				MatrixResult[ResultIndex] = 0.0;
-				for (int32 k = 0; k < AColumnNum; k++)
-				{
-					MatrixResult[ResultIndex] = MatrixResult[ResultIndex] + MatrixA[RowA * AColumnNum + k] * MatrixB[k * ResultRank + ColumnB];
-				}
+			MatrixResult[ResultIndex] = 0.0;
+			for (int32 k = 0; k < AColumnNum; k++)
+			{
+				MatrixResult[ResultIndex] = MatrixResult[ResultIndex] + MatrixA[RowA * AColumnNum + k] * MatrixB[k * ResultRank + ColumnB];
 			}
 		}
 	}
+}
 
-	void TransposeMatrix(int32 RowNum, int32 ColumnNum, const double* InMatrix, double* OutMatrix)
+void TransposeMatrix(int32 RowNum, int32 ColumnNum, const double* InMatrix, double* OutMatrix)
+{
+	for (int32 Row = 0; Row < RowNum; ++Row)
 	{
-		for (int32 Row = 0; Row < RowNum; ++Row)
+		for (int32 Column = 0; Column < ColumnNum; ++Column)
 		{
-			for (int32 Column = 0; Column < ColumnNum; ++Column)
-			{
-				OutMatrix[RowNum * Column + Row] = InMatrix[ColumnNum * Row + Column];
-			}
+			OutMatrix[RowNum * Column + Row] = InMatrix[ColumnNum * Row + Column];
 		}
 	}
-} // namespace CADKernel
+}
+} // namespace UE::CADKernel
diff --git a/Engine/Source/Runtime/Datasmith/CADKernel/Private/CADKernel/Mesh/Criteria/AngleCriterion.h b/Engine/Source/Runtime/Datasmith/CADKernel/Private/CADKernel/Mesh/Criteria/AngleCriterion.h
--- a/Engine/Source/Runtime/Datasmith/CADKernel/Private/CADKernel/Mesh/Criteria/AngleCriterion.h
+++ b/Engine/Source/Runtime/Datasmith/CADKernel/Private/CADKernel/Mesh/Criteria/AngleCriterion.h
@@ -8,69 +8,69 @@
  * https://docs.google.com/presentation/d/1bUnrRFWCW3sDn9ngb9ftfQS-2JxNJaUZlh783hZMMEw/edit?usp=sharing
 */
 
-namespace CADKernel
+namespace UE::CADKernel
 {
 
-	class FAngleCriterion : public FCriterion
+class FAngleCriterion : public FCriterion
+{
+	friend class FEntity;
+
+protected:
+	double AngleCriterionValue;
+	double SinMaxAngle;
+
+	/**
+	 * @param DegreeAngle the max allowed angle between tow elements in degree
+	 */
+	FAngleCriterion(double DegreeAngle)
+	{
+		AngleCriterionValue = FMath::DegreesToRadians(DegreeAngle);
+		SinMaxAngle = sin(AngleCriterionValue * 0.5);
+	}
+
+	FAngleCriterion() = default;
+
+public:
+
+	virtual void Serialize(FCADKernelArchive& Ar) override
+	{
+		FCriterion::Serialize(Ar);
+		Ar << AngleCriterionValue;
+		Ar << SinMaxAngle;
+	}
+
+	virtual double Value() const override
+	{
+		return AngleCriterionValue;
+	}
+
+	static double DefaultValue()
+	{
+		return FMath::DegreesToRadians(15.0);
+	}
+
+	virtual bool IsAppliedBetweenBreaks() const override
+	{
+		return true;
+	}
+
+	virtual ECriterion GetCriterionType() const override
 	{
-		friend class FEntity;
-
-	protected:
-		double AngleCriterionValue;
-		double SinMaxAngle;
-
-		/**
-		 * @param DegreeAngle the max allowed angle between tow elements in degree
-		 */
-		FAngleCriterion(double DegreeAngle)
-		{
-			AngleCriterionValue = FMath::DegreesToRadians(DegreeAngle);
-			SinMaxAngle = sin(AngleCriterionValue * 0.5);
-		}
-
-		FAngleCriterion()  = default;
-
-	public:
-
-		virtual void Serialize(FCADKernelArchive& Ar) override
-		{
-			FCriterion::Serialize(Ar);
-			Ar << AngleCriterionValue;
-			Ar << SinMaxAngle;
-		}
-
-		virtual double Value() const override
-		{
-			return AngleCriterionValue;
-		}
-
-		static double DefaultValue()
-		{
-			return FMath::DegreesToRadians(15.0);
-		}
-
-		virtual bool IsAppliedBetweenBreaks() const override
-		{
-			return true;
-		}
-
-		virtual ECriterion GetCriterionType() const override
-		{
-			return ECriterion::Angle;
-		}
-
-
-	protected:
-
-		/**
-		 * Sag & Angle criterion.pdf
-		 * https://docs.google.com/presentation/d/1bUnrRFWCW3sDn9ngb9ftfQS-2JxNJaUZlh783hZMMEw/edit?usp=sharing
-		*/
-		virtual double ComputeDeltaU(double ChordLength, double DeltaU, double Sag) const override
-		{
-			return 0.25 * SinMaxAngle * ChordLength * DeltaU / Sag;
-		}
-	};
-
-} // namespace CADKernel
+		return ECriterion::Angle;
+	}
+
+
+protected:
+
+	/**
+	 * Sag & Angle criterion.pdf
+	 * https://docs.google.com/presentation/d/1bUnrRFWCW3sDn9ngb9ftfQS-2JxNJaUZlh783hZMMEw/edit?usp=sharing
+	*/
+	virtual double ComputeDeltaU(double ChordLength, double DeltaU, double Sag) const override
+	{
+		return 0.25 * SinMaxAngle * ChordLength * DeltaU / Sag;
+	}
+};
+
+} // namespace UE::CADKernel
 
diff --git a/Engine/Source/Runtime/Datasmith/CADKernel/Private/CADKernel/Mesh/Criteria/CriteriaGrid.cpp b/Engine/Source/Runtime/Datasmith/CADKernel/Private/CADKernel/Mesh/Criteria/CriteriaGrid.cpp
--- a/Engine/Source/Runtime/Datasmith/CADKernel/Private/CADKernel/Mesh/Criteria/CriteriaGrid.cpp
+++ b/Engine/Source/Runtime/Datasmith/CADKernel/Private/CADKernel/Mesh/Criteria/CriteriaGrid.cpp
@@ -11,7 +11,7 @@
 #include "CADKernel/Mesh/Meshers/IsoTriangulator/DefineForDebug.h"
 #endif
 
-namespace CADKernel
+namespace UE::CADKernel
 {
 
 const FPoint& FCriteriaGrid::GetPoint(int32 UIndex, int32 VIndex, bool bIsInternalU, bool bIsInternalV, FVector3f* OutNormal) const
@@ -127,7 +127,7 @@ void FCriteriaGrid::Display() const
 		{
 			for (int32 UIndex = 0; UIndex < GetCoordinateCount(EIso::IsoU); ++UIndex)
 			{
-				CADKernel::DisplayPoint(GetPoint(UIndex, VIndex), UIndex);
+				UE::CADKernel::DisplayPoint(GetPoint(UIndex, VIndex), UIndex);
 			}
 		}
 	}
@@ -138,7 +138,7 @@ void FCriteriaGrid::Display() const
 		{
 			for (int32 UIndex = 0; UIndex < GetCoordinateCount(EIso::IsoU) - 1; ++UIndex)
 			{
-				CADKernel::DisplayPoint(GetIntermediateU(UIndex, VIndex), EVisuProperty::PurplePoint);
+				UE::CADKernel::DisplayPoint(GetIntermediateU(UIndex, VIndex), EVisuProperty::PurplePoint);
 			}
 		}
 	}
@@ -149,7 +149,7 @@ void FCriteriaGrid::Display() const
 		{
 			for (int32 UIndex = 0; UIndex < GetCoordinateCount(EIso::IsoU); ++UIndex)
 			{
-				CADKernel::DisplayPoint(GetIntermediateV(UIndex, VIndex), EVisuProperty::PurplePoint);
+				UE::CADKernel::DisplayPoint(GetIntermediateV(UIndex, VIndex), EVisuProperty::PurplePoint);
 			}
 		}
 	}
@@ -160,7 +160,7 @@ void FCriteriaGrid::Display() const
 		{
 			for (int32 UIndex = 0; UIndex < GetCoordinateCount(EIso::IsoU) - 1; ++UIndex)
 			{
-				CADKernel::DisplayPoint(GetIntermediateUV(UIndex, VIndex), EVisuProperty::PurplePoint);
+				UE::CADKernel::DisplayPoint(GetIntermediateUV(UIndex, VIndex), EVisuProperty::PurplePoint);
 			}
 		}
 	}
@@ -169,7 +169,7 @@ void FCriteriaGrid::Display() const
 		F3DDebugSession A(TEXT("Loop 3D"));
 		for (const TSharedPtr<FTopologicalLoop>& Loop : Face.GetLoops())
 		{
-			CADKernel::Display(*Loop);
+			UE::CADKernel::Display(*Loop);
 		}
 	}
 
@@ -177,7 +177,7 @@ void FCriteriaGrid::Display() const
 		F3DDebugSession A(TEXT("Loop 2D"));
 		for (const TSharedPtr<FTopologicalLoop>& Loop : Face.GetLoops())
 		{
-			CADKernel::Display2D(*Loop);
+			UE::CADKernel::Display2D(*Loop);
 		}
 	}
 
@@ -187,7 +187,7 @@ void FCriteriaGrid::Display() const
 		{
 			for (int32 UIndex = 0; UIndex < GetCoordinateCount(EIso::IsoU); ++UIndex)
 			{
-				CADKernel::DisplayPoint(FPoint2D(GetCoordinate(EIso::IsoU, UIndex), GetCoordinate(EIso::IsoV, VIndex)));
+				UE::CADKernel::DisplayPoint(FPoint2D(GetCoordinate(EIso::IsoU, UIndex), GetCoordinate(EIso::IsoV, VIndex)));
 			}
 		}
 	}
@@ -198,21 +198,21 @@ void FCriteriaGrid::Display() const
 		{
 			for (int32 UIndex = 1; UIndex < GetCoordinateCount(EIso::IsoU); ++UIndex)
 			{
-				CADKernel::DisplayPoint(FPoint2D((GetCoordinate(EIso::IsoU, UIndex) + GetCoordinate(EIso::IsoU, UIndex - 1)) * 0.5, GetCoordinate(EIso::IsoV, VIndex)), EVisuProperty::PurplePoint);
+				UE::CADKernel::DisplayPoint(FPoint2D((GetCoordinate(EIso::IsoU, UIndex) + GetCoordinate(EIso::IsoU, UIndex - 1)) * 0.5, GetCoordinate(EIso::IsoV, VIndex)), EVisuProperty::PurplePoint);
 			}
 		}
 		for (int32 VIndex = 1; VIndex < GetCoordinateCount(EIso::IsoV); ++VIndex)
 		{
 			for (int32 UIndex = 0; UIndex < GetCoordinateCount(EIso::IsoU); ++UIndex)
 			{
-				CADKernel::DisplayPoint(FPoint2D(GetCoordinate(EIso::IsoU, UIndex), (GetCoordinate(EIso::IsoV, VIndex) + GetCoordinate(EIso::IsoV, VIndex - 1)) * 0.5), EVisuProperty::PurplePoint);
+				UE::CADKernel::DisplayPoint(FPoint2D(GetCoordinate(EIso::IsoU, UIndex), (GetCoordinate(EIso::IsoV, VIndex) + GetCoordinate(EIso::IsoV, VIndex - 1)) * 0.5), EVisuProperty::PurplePoint);
 			}
 		}
 		for (int32 VIndex = 1; VIndex < GetCoordinateCount(EIso::IsoV); ++VIndex)
 		{
 			for (int32 UIndex = 1; UIndex < GetCoordinateCount(EIso::IsoU); ++UIndex)
 			{
-				CADKernel::DisplayPoint(FPoint2D((GetCoordinate(EIso::IsoU, UIndex) + GetCoordinate(EIso::IsoU, UIndex - 1)) * 0.5, (GetCoordinate(EIso::IsoV, VIndex) + GetCoordinate(EIso::IsoV, VIndex - 1)) * 0.5), EVisuProperty::PurplePoint);
+				UE::CADKernel::DisplayPoint(FPoint2D((GetCoordinate(EIso::IsoU, UIndex) + GetCoordinate(EIso::IsoU, UIndex - 1)) * 0.5, (GetCoordinate(EIso::IsoV, VIndex) + GetCoordinate(EIso::IsoV, VIndex - 1)) * 0.5), EVisuProperty::PurplePoint);
 			}
 		}
 	}
@@ -220,5 +220,5 @@ void FCriteriaGrid::Display() const
 	Wait();
 }
 
-} // namespace CADKernel
+} // namespace UE::CADKernel
 
diff --git a/Engine/Source/Runtime/Datasmith/CADKernel/Private/CADKernel/Mesh/Criteria/CriteriaGrid.h b/Engine/Source/Runtime/Datasmith/CADKernel/Private/CADKernel/Mesh/Criteria/CriteriaGrid.h
--- a/Engine/Source/Runtime/Datasmith/CADKernel/Private/CADKernel/Mesh/Criteria/CriteriaGrid.h
+++ b/Engine/Source/Runtime/Datasmith/CADKernel/Private/CADKernel/Mesh/Criteria/CriteriaGrid.h
@@ -6,94 +6,94 @@
 #include "CADKernel/Geo/Sampling/SurfacicSampling.h"
 #include "CADKernel/Math/Point.h"
 
-namespace CADKernel
+namespace UE::CADKernel
 {
-	class FTopologicalFace;
-	class FCriterion;
+class FTopologicalFace;
+class FCriterion;
 
-	class FCriteriaGrid
-	{
-	protected:
-		FTopologicalFace& Face;
+class FCriteriaGrid
+{
+protected:
+	FTopologicalFace& Face;
 
-		const FCoordinateGrid& CoordinateGrid;
-		int32 TrueUcoorindateCount; 
-		FSurfacicSampling Grid;
+	const FCoordinateGrid& CoordinateGrid;
+	int32 TrueUcoorindateCount;
+	FSurfacicSampling Grid;
 
-		int32 GetIndex(int32 UIndex, int32 VIndex, bool bIsInternalU, bool bIsInternalV) const
+	int32 GetIndex(int32 UIndex, int32 VIndex, bool bIsInternalU, bool bIsInternalV) const
+	{
+		int32 TrueUIndex = UIndex * 2;
+		int32 TrueVIndex = VIndex * 2;
+		if (bIsInternalU)
+		{
+			++TrueUIndex;
+		}
+		if (bIsInternalV)
 		{
-			int32 TrueUIndex = UIndex * 2;
-			int32 TrueVIndex = VIndex * 2;
-			if (bIsInternalU)
-			{
-				++TrueUIndex;
-			}
-			if (bIsInternalV)
-			{
-				++TrueVIndex;
-			}
-			return TrueVIndex * TrueUcoorindateCount + TrueUIndex;
+			++TrueVIndex;
 		}
+		return TrueVIndex * TrueUcoorindateCount + TrueUIndex;
+	}
 
-		const FPoint& GetPoint(int32 UIndex, int32 VIndex, bool bIsInterU, bool bIsInterV, FVector3f* normal = nullptr) const;
-		void Init();
+	const FPoint& GetPoint(int32 UIndex, int32 VIndex, bool bIsInterU, bool bIsInterV, FVector3f* normal = nullptr) const;
+	void Init();
 
-	public:
-		FCriteriaGrid(FTopologicalFace& Surface);
+public:
+	FCriteriaGrid(FTopologicalFace& Surface);
 
-		FTopologicalFace& GetSurface()
-		{
-			return Face;
-		}
+	FTopologicalFace& GetSurface()
+	{
+		return Face;
+	}
 
-		void ApplyCriteria(const TArray<TSharedPtr<FCriterion>>& GetCriteria) const;
+	void ApplyCriteria(const TArray<TSharedPtr<FCriterion>>& GetCriteria) const;
 
-		double GetCoordinate(EIso Iso, int32 ind) const
-		{
-			return CoordinateGrid[Iso][ind];
-		}
+	double GetCoordinate(EIso Iso, int32 ind) const
+	{
+		return CoordinateGrid[Iso][ind];
+	}
 
-		constexpr const TArray<double>& GetCoordinates(EIso Iso) const
-		{
-			return CoordinateGrid[Iso];
-		}
+	constexpr const TArray<double>& GetCoordinates(EIso Iso) const
+	{
+		return CoordinateGrid[Iso];
+	}
 
-		int32 GetCoordinateCount(EIso Iso) const
-		{
-			return CoordinateGrid.IsoCount(Iso);
-		}
+	int32 GetCoordinateCount(EIso Iso) const
+	{
+		return CoordinateGrid.IsoCount(Iso);
+	}
 
-		const TArray<FPoint>& GetGridPoints() const
-		{
-			return Grid.Points3D;
-		}
+	const TArray<FPoint>& GetGridPoints() const
+	{
+		return Grid.Points3D;
+	}
 
-		const TArray<FVector3f>& GetGridNormals() const
-		{
-			return Grid.Normals;
-		}
+	const TArray<FVector3f>& GetGridNormals() const
+	{
+		return Grid.Normals;
+	}
 
-		const FPoint& GetPoint(int32 iU, int32 iV, FVector3f* normal = nullptr) const
-		{
-			return GetPoint(iU, iV, false, false, normal);
-		}
+	const FPoint& GetPoint(int32 iU, int32 iV, FVector3f* normal = nullptr) const
+	{
+		return GetPoint(iU, iV, false, false, normal);
+	}
 
-		const FPoint& GetIntermediateU(int32 iU, int32 iV, FVector3f* normal = nullptr) const
-		{
-			return GetPoint(iU, iV, true, false, normal);
-		}
+	const FPoint& GetIntermediateU(int32 iU, int32 iV, FVector3f* normal = nullptr) const
+	{
+		return GetPoint(iU, iV, true, false, normal);
+	}
 
-		const FPoint& GetIntermediateV(int32 iU, int32 iV, FVector3f* normal = nullptr) const
-		{
-			return GetPoint(iU, iV, false, true, normal);
-		}
+	const FPoint& GetIntermediateV(int32 iU, int32 iV, FVector3f* normal = nullptr) const
+	{
+		return GetPoint(iU, iV, false, true, normal);
+	}
 
-		const FPoint& GetIntermediateUV(int32 iU, int32 iV, FVector3f* normal = nullptr) const
-		{
-			return GetPoint(iU, iV, true, true, normal);
-		}
+	const FPoint& GetIntermediateUV(int32 iU, int32 iV, FVector3f* normal = nullptr) const
+	{
+		return GetPoint(iU, iV, true, true, normal);
+	}
 
-		void Display() const;
-	};
+	void Display() const;
+};
 }
 
diff --git a/Engine/Source/Runtime/Datasmith/CADKernel/Private/CADKernel/Mesh/Criteria/Criterion.cpp b/Engine/Source/Runtime/Datasmith/CADKernel/Private/CADKernel/Mesh/Criteria/Criterion.cpp
--- a/Engine/Source/Runtime/Datasmith/CADKernel/Private/CADKernel/Mesh/Criteria/Criterion.cpp
+++ b/Engine/Source/Runtime/Datasmith/CADKernel/Private/CADKernel/Mesh/Criteria/Criterion.cpp
@@ -11,153 +11,153 @@
 #include "CADKernel/Geo/GeoPoint.h"
 #include "CADKernel/Topo/TopologicalEdge.h"
 
-namespace CADKernel
+namespace UE::CADKernel
 {
 
-	void FCriterion::ApplyOnEdgeParameters(FTopologicalEdge& Edge, const TArray<double>& Coordinates, const TArray<FCurvePoint>& Points) const
-	{
-		TArray<double>& DeltaUMaxs = Edge.GetDeltaUMaxs();
+void FCriterion::ApplyOnEdgeParameters(FTopologicalEdge& Edge, const TArray<double>& Coordinates, const TArray<FCurvePoint>& Points) const
+{
+	TArray<double>& DeltaUMaxs = Edge.GetDeltaUMaxs();
 
-		double NumericPrecision = Edge.GetTolerance3D();
-		NumericPrecision /= 10;
+	double NumericPrecision = Edge.GetTolerance3D();
+	NumericPrecision /= 10;
 
-		if (Edge.Length() <= NumericPrecision)
-		{
-			return;
-		}
+	if (Edge.Length() <= NumericPrecision)
+	{
+		return;
+	}
 
-		int32 PreviousIndex = 0;
-		for (int32 Index = 1; Index < Coordinates.Num(); Index++)
-		{
-			int32 MiddleIndex = PreviousIndex + Index;
-			double ChordLength;
-			double Sag = EvaluateSag(Points[2 * PreviousIndex].Point, Points[2 * Index].Point, Points[MiddleIndex].Point, ChordLength);
+	int32 PreviousIndex = 0;
+	for (int32 Index = 1; Index < Coordinates.Num(); Index++)
+	{
+		int32 MiddleIndex = PreviousIndex + Index;
+		double ChordLength;
+		double Sag = EvaluateSag(Points[2 * PreviousIndex].Point, Points[2 * Index].Point, Points[MiddleIndex].Point, ChordLength);
 
-			double AiAm = Points[2 * PreviousIndex].Point.Distance(Points[MiddleIndex].Point);
-			double AmAi2 = Points[2 * Index].Point.Distance(Points[MiddleIndex].Point);
-			double AiAi2 = Points[2 * PreviousIndex].Point.Distance(Points[2 * Index].Point);
+		double AiAm = Points[2 * PreviousIndex].Point.Distance(Points[MiddleIndex].Point);
+		double AmAi2 = Points[2 * Index].Point.Distance(Points[MiddleIndex].Point);
+		double AiAi2 = Points[2 * PreviousIndex].Point.Distance(Points[2 * Index].Point);
 
-			double DeltaU = Coordinates[Index] - Coordinates[PreviousIndex];
+		double DeltaU = Coordinates[Index] - Coordinates[PreviousIndex];
 
-			if (ChordLength > NumericPrecision && Sag > DOUBLE_SMALL_NUMBER)
-			{
-				double DeltaUMax = ComputeDeltaU(ChordLength, DeltaU, Sag);
-				if (DeltaUMax < DeltaUMaxs[PreviousIndex]) DeltaUMaxs[PreviousIndex] = DeltaUMax;
-			}
+		if (ChordLength > NumericPrecision && Sag > DOUBLE_SMALL_NUMBER)
+		{
+			double DeltaUMax = ComputeDeltaU(ChordLength, DeltaU, Sag);
+			if (DeltaUMax < DeltaUMaxs[PreviousIndex]) DeltaUMaxs[PreviousIndex] = DeltaUMax;
+		}
 
-			Sag = EvaluateSag(Points[2 * Index].Point, Points[2 * PreviousIndex].Point, Points[MiddleIndex].Point, ChordLength);
-			if (ChordLength > NumericPrecision && Sag > DOUBLE_SMALL_NUMBER)
+		Sag = EvaluateSag(Points[2 * Index].Point, Points[2 * PreviousIndex].Point, Points[MiddleIndex].Point, ChordLength);
+		if (ChordLength > NumericPrecision && Sag > DOUBLE_SMALL_NUMBER)
+		{
+			double DeltaUMax = ComputeDeltaU(ChordLength, DeltaU, Sag);
+			if (DeltaUMax < DeltaUMaxs[PreviousIndex])
 			{
-				double DeltaUMax = ComputeDeltaU(ChordLength, DeltaU, Sag);
-				if (DeltaUMax < DeltaUMaxs[PreviousIndex]) 
-				{
-					DeltaUMaxs[PreviousIndex] = DeltaUMax;
-				}
+				DeltaUMaxs[PreviousIndex] = DeltaUMax;
 			}
-			PreviousIndex = Index;
 		}
+		PreviousIndex = Index;
 	}
+}
 
-	TSharedPtr<FCriterion> FCriterion::Deserialize(FCADKernelArchive& Archive)
-	{
-		ECriterion CriterionType = ECriterion::None;
-		Archive << CriterionType;
+TSharedPtr<FCriterion> FCriterion::Deserialize(FCADKernelArchive& Archive)
+{
+	ECriterion CriterionType = ECriterion::None;
+	Archive << CriterionType;
 
-		TSharedPtr<FEntity> Entity;
-		switch (CriterionType)
-		{
-		case ECriterion::MinSize:
-			return FEntity::MakeShared<FMinSizeCriterion>(Archive);
-		case ECriterion::MaxSize:
-			return FEntity::MakeShared<FMaxSizeCriterion>(Archive);
-		case ECriterion::CADCurvature:
-			return FEntity::MakeShared<FCurvatureCriterion>(Archive);
-		case ECriterion::Sag:
-			return FEntity::MakeShared<FSagCriterion>(Archive);
-		case ECriterion::Angle:
-			return FEntity::MakeShared<FAngleCriterion>(Archive);
-		default:
-			return TSharedPtr<FCriterion>();
-		}
+	TSharedPtr<FEntity> Entity;
+	switch (CriterionType)
+	{
+	case ECriterion::MinSize:
+		return FEntity::MakeShared<FMinSizeCriterion>(Archive);
+	case ECriterion::MaxSize:
+		return FEntity::MakeShared<FMaxSizeCriterion>(Archive);
+	case ECriterion::CADCurvature:
+		return FEntity::MakeShared<FCurvatureCriterion>(Archive);
+	case ECriterion::Sag:
+		return FEntity::MakeShared<FSagCriterion>(Archive);
+	case ECriterion::Angle:
+		return FEntity::MakeShared<FAngleCriterion>(Archive);
+	default:
+		return TSharedPtr<FCriterion>();
 	}
+}
 
-	TSharedPtr<FCriterion> FCriterion::CreateCriterion(ECriterion Type, double Value)
+TSharedPtr<FCriterion> FCriterion::CreateCriterion(ECriterion Type, double Value)
+{
+	switch (Type)
 	{
-		switch (Type)
-		{
-		case ECriterion::MinSize:
-			return FEntity::MakeShared<FMinSizeCriterion>(Value);
-		case ECriterion::MaxSize:
-			return FEntity::MakeShared<FMaxSizeCriterion>(Value);
-		case ECriterion::Sag:
-			return FEntity::MakeShared<FSagCriterion>(Value);
-		case ECriterion::Angle:
-			return FEntity::MakeShared<FAngleCriterion>(Value);
-		case ECriterion::CADCurvature:
-			return FEntity::MakeShared<FCurvatureCriterion>();
-		default:
-			ERROR_FUNCTION_CALL_NOT_EXPECTED;
-			return TSharedPtr<FCriterion>();
-		}
+	case ECriterion::MinSize:
+		return FEntity::MakeShared<FMinSizeCriterion>(Value);
+	case ECriterion::MaxSize:
+		return FEntity::MakeShared<FMaxSizeCriterion>(Value);
+	case ECriterion::Sag:
+		return FEntity::MakeShared<FSagCriterion>(Value);
+	case ECriterion::Angle:
+		return FEntity::MakeShared<FAngleCriterion>(Value);
+	case ECriterion::CADCurvature:
+		return FEntity::MakeShared<FCurvatureCriterion>();
+	default:
+		ERROR_FUNCTION_CALL_NOT_EXPECTED;
+		return TSharedPtr<FCriterion>();
 	}
+}
 
-	double FCriterion::DefaultValue(ECriterion Type)
+double FCriterion::DefaultValue(ECriterion Type)
+{
+	switch (Type)
 	{
-		switch (Type)
-		{
-		case ECriterion::MinSize:
-			return FSizeCriterion::DefaultValue(Type);
-		case ECriterion::MaxSize:
-			return FSizeCriterion::DefaultValue(Type);
-		case ECriterion::Sag:
-			return FSagCriterion::DefaultValue();
-		case ECriterion::Angle:
-			return FAngleCriterion::DefaultValue();
-		default:
-			ERROR_FUNCTION_CALL_NOT_EXPECTED;
-			return 0.0;
-		}
+	case ECriterion::MinSize:
+		return FSizeCriterion::DefaultValue(Type);
+	case ECriterion::MaxSize:
+		return FSizeCriterion::DefaultValue(Type);
+	case ECriterion::Sag:
+		return FSagCriterion::DefaultValue();
+	case ECriterion::Angle:
+		return FAngleCriterion::DefaultValue();
+	default:
+		ERROR_FUNCTION_CALL_NOT_EXPECTED;
+		return 0.0;
 	}
+}
 
 #ifdef CADKERNEL_DEV
-	FInfoEntity& FCriterion::GetInfo(FInfoEntity& Info) const
-	{
-		FEntity::GetInfo(Info)
-			.Add(TEXT("CriterionType"), CriterionTypeNames[(uint8)GetCriterionType()])
-			.Add(TEXT("Value"), Value());
-		return Info;
-	}
+FInfoEntity& FCriterion::GetInfo(FInfoEntity& Info) const
+{
+	FEntity::GetInfo(Info)
+		.Add(TEXT("CriterionType"), CriterionTypeNames[(uint8)GetCriterionType()])
+		.Add(TEXT("Value"), Value());
+	return Info;
+}
 #endif
 
-	const TCHAR* CriterionTypeNames[] = {
-		TEXT("Size"),
-		TEXT("MaxSize"),
-		TEXT("MinSize"),
-		TEXT("Angle"),
-		TEXT("Chord Error"),
-		TEXT("CAD Curvature"),
-		TEXT("")
-	};
-
-	// Defined for Python purpose
-	const char* CriterionTypeConstNames[] = {
-		"CRITERION_SIZE",
-		"CRITERION_MAX_SIZE",
-		"CRITERION_MIN_SIZE",
-		"CRITERION_ANGLE",
-		"CRITERION_SAG",
-		"CRITERION_CAD_CURVATURE",
-		nullptr
-	};
-
-	// Defined for Python purpose
-	const char* CriterionTypeConstDescHelp[] = {
-		": Size of elements (only for check)",
-		": maximum size of elements",
-		": minimum size of elements",
-		": angle between the normal of elements",
-		": distance between elements and CAD surfaces/curves",
-		": CAD curvature (only for check)",
-		nullptr
-	};
+const TCHAR* CriterionTypeNames[] = {
+	TEXT("Size"),
+	TEXT("MaxSize"),
+	TEXT("MinSize"),
+	TEXT("Angle"),
+	TEXT("Chord Error"),
+	TEXT("CAD Curvature"),
+	TEXT("")
+};
+
+// Defined for Python purpose
+const char* CriterionTypeConstNames[] = {
+	"CRITERION_SIZE",
+	"CRITERION_MAX_SIZE",
+	"CRITERION_MIN_SIZE",
+	"CRITERION_ANGLE",
+	"CRITERION_SAG",
+	"CRITERION_CAD_CURVATURE",
+	nullptr
+};
+
+// Defined for Python purpose
+const char* CriterionTypeConstDescHelp[] = {
+	": Size of elements (only for check)",
+	": maximum size of elements",
+	": minimum size of elements",
+	": angle between the normal of elements",
+	": distance between elements and CAD surfaces/curves",
+	": CAD curvature (only for check)",
+	nullptr
+};
 }
\ No newline at end of file
diff --git a/Engine/Source/Runtime/Datasmith/CADKernel/Private/CADKernel/Mesh/Criteria/CurvatureCriterion.h b/Engine/Source/Runtime/Datasmith/CADKernel/Private/CADKernel/Mesh/Criteria/CurvatureCriterion.h
--- a/Engine/Source/Runtime/Datasmith/CADKernel/Private/CADKernel/Mesh/Criteria/CurvatureCriterion.h
+++ b/Engine/Source/Runtime/Datasmith/CADKernel/Private/CADKernel/Mesh/Criteria/CurvatureCriterion.h
@@ -4,78 +4,78 @@
 #include "CADKernel/Math/Curvature.h"
 #include "CADKernel/Mesh/Criteria/Criterion.h"
 
-namespace CADKernel
+namespace UE::CADKernel
 {
-	struct FCurvePoint;
+struct FCurvePoint;
 
-	class FCurvatureCriterion : public FCriterion
-	{
-		friend class FEntity;
-		
-	protected:
+class FCurvatureCriterion : public FCriterion
+{
+	friend class FEntity;
 
-		FCurvatureCriterion()  = default;
+protected:
 
-	public:
+	FCurvatureCriterion() = default;
 
-		void Serialize(FCADKernelArchive& Ar)
-		{
-			FCriterion::Serialize(Ar);
-		}
+public:
 
-		/**
-		 * Sag & Angle criterion.pdf
-		 * https://docs.google.com/presentation/d/1bUnrRFWCW3sDn9ngb9ftfQS-2JxNJaUZlh783hZMMEw/edit?usp=sharing
-		 */
-		virtual void UpdateDelta(double InDeltaU, double InUSag, double InDiagonalSag, double InVSag, double ChordLength, double DiagonalLength, double& OutSagDeltaUMax, double& OutSagDeltaUMin, FIsoCurvature& IsoCurvature) const override
+	void Serialize(FCADKernelArchive& Ar)
+	{
+		FCriterion::Serialize(Ar);
+	}
+
+	/**
+	 * Sag & Angle criterion.pdf
+	 * https://docs.google.com/presentation/d/1bUnrRFWCW3sDn9ngb9ftfQS-2JxNJaUZlh783hZMMEw/edit?usp=sharing
+	 */
+	virtual void UpdateDelta(double InDeltaU, double InUSag, double InDiagonalSag, double InVSag, double ChordLength, double DiagonalLength, double& OutSagDeltaUMax, double& OutSagDeltaUMin, FIsoCurvature& IsoCurvature) const override
+	{
+		if (DiagonalLength > DOUBLE_SMALL_NUMBER)
 		{
-			if (DiagonalLength > DOUBLE_SMALL_NUMBER)
+			double Curvature = 8 * InUSag / FMath::Square(DiagonalLength);
+			if (IsoCurvature.Max < Curvature)
 			{
-				double Curvature = 8 * InUSag / FMath::Square(DiagonalLength);
-				if (IsoCurvature.Max < Curvature)
-				{
-					IsoCurvature.Max = Curvature;
-				}
-				if (IsoCurvature.Min > Curvature)
-				{
-					IsoCurvature.Min = Curvature;
-				}
+				IsoCurvature.Max = Curvature;
+			}
+			if (IsoCurvature.Min > Curvature)
+			{
+				IsoCurvature.Min = Curvature;
 			}
 		}
+	}
 
-		virtual double Value() const override
-		{
-			return 0.;
-		}
-
-		virtual void ApplyOnEdgeParameters(FTopologicalEdge& Edge, const TArray<double>& Coordinates, const TArray<FCurvePoint>& Points) const override
-		{
-			// Do nothing
-		}
+	virtual double Value() const override
+	{
+		return 0.;
+	}
 
-		virtual ECriterion GetCriterionType() const override
-		{
-			return ECriterion::CADCurvature;
-		}
+	virtual void ApplyOnEdgeParameters(FTopologicalEdge& Edge, const TArray<double>& Coordinates, const TArray<FCurvePoint>& Points) const override
+	{
+		// Do nothing
+	}
 
+	virtual ECriterion GetCriterionType() const override
+	{
+		return ECriterion::CADCurvature;
+	}
 
-	protected:
 
-		/**
-		 * Sag & Angle criterion.pdf
-		 * https://docs.google.com/presentation/d/1bUnrRFWCW3sDn9ngb9ftfQS-2JxNJaUZlh783hZMMEw/edit?usp=sharing
-		 */
-		virtual bool IsAppliedBetweenBreaks() const override
-		{
-			return true;
-		}
+protected:
 
-		virtual double ComputeDeltaU(double ChordLength, double DeltaU, double Sag) const override
-		{
-			return 0.0;
-		};
+	/**
+	 * Sag & Angle criterion.pdf
+	 * https://docs.google.com/presentation/d/1bUnrRFWCW3sDn9ngb9ftfQS-2JxNJaUZlh783hZMMEw/edit?usp=sharing
+	 */
+	virtual bool IsAppliedBetweenBreaks() const override
+	{
+		return true;
+	}
 
+	virtual double ComputeDeltaU(double ChordLength, double DeltaU, double Sag) const override
+	{
+		return 0.0;
 	};
 
-} // namespace CADKernel
+};
+
+} // namespace UE::CADKernel
 
diff --git a/Engine/Source/Runtime/Datasmith/CADKernel/Private/CADKernel/Mesh/Criteria/SagCriterion.h b/Engine/Source/Runtime/Datasmith/CADKernel/Private/CADKernel/Mesh/Criteria/SagCriterion.h
--- a/Engine/Source/Runtime/Datasmith/CADKernel/Private/CADKernel/Mesh/Criteria/SagCriterion.h
+++ b/Engine/Source/Runtime/Datasmith/CADKernel/Private/CADKernel/Mesh/Criteria/SagCriterion.h
@@ -3,61 +3,61 @@
 
 #include "CADKernel/Mesh/Criteria/Criterion.h"
 
-namespace CADKernel
+namespace UE::CADKernel
 {
 
-	class FSagCriterion : public FCriterion
+class FSagCriterion : public FCriterion
+{
+	friend class FEntity;
+
+protected:
+	double MaxSag;
+
+	FSagCriterion(double InSag = 0.2)
+		: MaxSag(InSag)
+	{
+	}
+
+public:
+
+	void Serialize(FCADKernelArchive& Ar)
+	{
+		FCriterion::Serialize(Ar);
+		Ar << MaxSag;
+	}
+
+	virtual double Value() const override
+	{
+		return MaxSag;
+	}
+
+	static double DefaultValue()
 	{
-		friend class FEntity;
-		
-	protected:
-		double MaxSag;
-
-		FSagCriterion(double InSag = 0.2)
-			: MaxSag(InSag)
-		{
-		}
-
-	public:
-
-		void Serialize(FCADKernelArchive& Ar)
-		{
-			FCriterion::Serialize(Ar);
-			Ar << MaxSag;
-		}
-
-		virtual double Value() const override
-		{
-			return MaxSag;
-		}
-
-		static double DefaultValue()
-		{
-			return 0.2;
-		}
-
-		virtual ECriterion GetCriterionType() const override
-		{
-			return ECriterion::Sag;
-		}
-
-
-	protected:
-
-		/**
-		 * Sag & Angle criterion.pdf
-		 * https://docs.google.com/presentation/d/1bUnrRFWCW3sDn9ngb9ftfQS-2JxNJaUZlh783hZMMEw/edit?usp=sharing
-		*/
-		virtual double ComputeDeltaU(double ChordLength, double DeltaU, double Sag) const override
-		{
-			return sqrt(MaxSag / Sag) * DeltaU;
-		}
-
-		virtual bool IsAppliedBetweenBreaks() const override
-		{
-			return true;
-		}
-	};
-
-} // namespace CADKernel
+		return 0.2;
+	}
+
+	virtual ECriterion GetCriterionType() const override
+	{
+		return ECriterion::Sag;
+	}
+
+
+protected:
+
+	/**
+	 * Sag & Angle criterion.pdf
+	 * https://docs.google.com/presentation/d/1bUnrRFWCW3sDn9ngb9ftfQS-2JxNJaUZlh783hZMMEw/edit?usp=sharing
+	*/
+	virtual double ComputeDeltaU(double ChordLength, double DeltaU, double Sag) const override
+	{
+		return sqrt(MaxSag / Sag) * DeltaU;
+	}
+
+	virtual bool IsAppliedBetweenBreaks() const override
+	{
+		return true;
+	}
+};
+
+} // namespace UE::CADKernel
 
diff --git a/Engine/Source/Runtime/Datasmith/CADKernel/Private/CADKernel/Mesh/Criteria/SizeCriterion.h b/Engine/Source/Runtime/Datasmith/CADKernel/Private/CADKernel/Mesh/Criteria/SizeCriterion.h
--- a/Engine/Source/Runtime/Datasmith/CADKernel/Private/CADKernel/Mesh/Criteria/SizeCriterion.h
+++ b/Engine/Source/Runtime/Datasmith/CADKernel/Private/CADKernel/Mesh/Criteria/SizeCriterion.h
@@ -8,89 +8,89 @@
  * https://docs.google.com/presentation/d/1bUnrRFWCW3sDn9ngb9ftfQS-2JxNJaUZlh783hZMMEw/edit?usp=sharing
 */
 
-namespace CADKernel
+namespace UE::CADKernel
 {
-	struct FCurvePoint;
+struct FCurvePoint;
 
-	class FSizeCriterion : public FCriterion
-	{
-		friend class FEntity;
-	
-	protected:
-		double Size;
+class FSizeCriterion : public FCriterion
+{
+	friend class FEntity;
 
-		FSizeCriterion(double InSize = 1.)
-			: Size(InSize)
-		{
-		}
+protected:
+	double Size;
 
+	FSizeCriterion(double InSize = 1.)
+		: Size(InSize)
+	{
+	}
 
-	public:
 
-		void Serialize(FCADKernelArchive& Ar)
-		{
-			FCriterion::Serialize(Ar);
-			Ar << Size;
-		}
+public:
 
-		double Value() const override
-		{
-			return Size;
-		}
+	void Serialize(FCADKernelArchive& Ar)
+	{
+		FCriterion::Serialize(Ar);
+		Ar << Size;
+	}
 
-		static double DefaultValue(ECriterion Type)
+	double Value() const override
+	{
+		return Size;
+	}
+
+	static double DefaultValue(ECriterion Type)
+	{
+		switch (Type)
 		{
-			switch (Type)
-			{
-			case ECriterion::MinSize:
-				return 0.1;
-			case ECriterion::MaxSize:
-				return 30;
-			}
-			return 0;
+		case ECriterion::MinSize:
+			return 0.1;
+		case ECriterion::MaxSize:
+			return 30;
 		}
+		return 0;
+	}
 
-		void ApplyOnParameters(const TArray<double>& TabU, const TArray<FCurvePoint>& tabPt, TArray<double>& tabDeltaU, TFunction<void(double, double&)> Compare) const;
-	};
+	void ApplyOnParameters(const TArray<double>& TabU, const TArray<FCurvePoint>& tabPt, TArray<double>& tabDeltaU, TFunction<void(double, double&)> Compare) const;
+};
 
-	class FMinSizeCriterion : public FSizeCriterion
+class FMinSizeCriterion : public FSizeCriterion
+{
+	friend class FEntity;
+protected:
+	FMinSizeCriterion(double InSize = 0.05)
+		: FSizeCriterion(InSize)
 	{
-		friend class FEntity;
-	protected:
-		FMinSizeCriterion(double InSize = 0.05)
-			: FSizeCriterion(InSize)
-		{
-		}
+	}
 
-	public:
+public:
 
-		virtual ECriterion GetCriterionType() const override
-		{
-			return ECriterion::MinSize;
-		}
+	virtual ECriterion GetCriterionType() const override
+	{
+		return ECriterion::MinSize;
+	}
 
-		virtual void ApplyOnEdgeParameters(FTopologicalEdge& Edge, const TArray<double>& Coordinates, const TArray<FCurvePoint>& Points) const override;
-		virtual void UpdateDelta(double InDeltaU, double InUSag, double InDiagonalSag, double InVSag, double ChordLength, double DiagonalLength, double& OutSagDeltaUMax, double& OutSagDeltaUMin, FIsoCurvature& SurfaceCurvature) const override;
-	};
+	virtual void ApplyOnEdgeParameters(FTopologicalEdge& Edge, const TArray<double>& Coordinates, const TArray<FCurvePoint>& Points) const override;
+	virtual void UpdateDelta(double InDeltaU, double InUSag, double InDiagonalSag, double InVSag, double ChordLength, double DiagonalLength, double& OutSagDeltaUMax, double& OutSagDeltaUMin, FIsoCurvature& SurfaceCurvature) const override;
+};
 
-	class FMaxSizeCriterion : public FSizeCriterion
+class FMaxSizeCriterion : public FSizeCriterion
+{
+	friend class FEntity;
+protected:
+	FMaxSizeCriterion(double InSize = 10000.)
+		: FSizeCriterion(InSize)
 	{
-		friend class FEntity;
-	protected:
-		FMaxSizeCriterion(double InSize = 10000.)
-			: FSizeCriterion(InSize)
-		{
-		}
+	}
 
-	public:
-		virtual ECriterion GetCriterionType() const override
-		{
-			return ECriterion::MaxSize;
-		}
+public:
+	virtual ECriterion GetCriterionType() const override
+	{
+		return ECriterion::MaxSize;
+	}
 
-		virtual void ApplyOnEdgeParameters(FTopologicalEdge& Edge, const TArray<double>& Coordinates, const TArray<FCurvePoint>& Points) const override;
-		virtual void UpdateDelta(double InDeltaU, double InUSag, double InDiagonalSag, double InVSag, double ChordLength, double DiagonalLength, double& OutSagDeltaUMax, double& OutSagDeltaUMin, FIsoCurvature& SurfaceCurvature) const override;
-	};
+	virtual void ApplyOnEdgeParameters(FTopologicalEdge& Edge, const TArray<double>& Coordinates, const TArray<FCurvePoint>& Points) const override;
+	virtual void UpdateDelta(double InDeltaU, double InUSag, double InDiagonalSag, double InVSag, double ChordLength, double DiagonalLength, double& OutSagDeltaUMax, double& OutSagDeltaUMin, FIsoCurvature& SurfaceCurvature) const override;
+};
 
-} // namespace CADKernel
+} // namespace UE::CADKernel
 
diff --git a/Engine/Source/Runtime/Datasmith/CADKernel/Private/CADKernel/Mesh/MeshEnum.h b/Engine/Source/Runtime/Datasmith/CADKernel/Private/CADKernel/Mesh/MeshEnum.h
--- a/Engine/Source/Runtime/Datasmith/CADKernel/Private/CADKernel/Mesh/MeshEnum.h
+++ b/Engine/Source/Runtime/Datasmith/CADKernel/Private/CADKernel/Mesh/MeshEnum.h
@@ -1,23 +1,23 @@
 // Copyright Epic Games, Inc. All Rights Reserved.
 #pragma once
 
-namespace CADKernel
+namespace UE::CADKernel
 {
-	/**
-	 * 3 parametric spaces are defined :
-	 *  -EGridSpace::Default2D the natural parametric space of the carrier surface i.e. for a Bezier surface the boundary of the space is [[0, 1], [0, 1]] and for a spherical surface [[0, 2Pi], [-Pi / 2, Pi / 2]]
-	 *  -EGridSpace::UniformScaled.The average curvilinear length of the iso curves is calculated in each axis.The parametric space is then scaled in each axis such that
-	 *            the average curvilinear length between two points P(A, V) and P(B, V) = B - A.For a spherical surface [[0, 2Pi.R / 2], [-Pi.R / 2, Pi.R / 2]] with R the radius of the sphere
-	 * -EGridSpace::Scaled.The UnifomeScaled space is rescaled in its most disturbe axis e.g. if V is the most disturbe axis, whatever V, length of[P(A, V), P(B, V)] = (B - A).
-	 *            This parametric space is closed to a good flattening of the surface.
-	 */
-	enum EGridSpace : uint8
-	{
-		Default2D = 0,  // do not modified enum value as it is used for a static array
-		Scaled = 1,
-		UniformScaled = 2,
-		EndGridSpace
-	};
+/**
+ * 3 parametric spaces are defined :
+ *  -EGridSpace::Default2D the natural parametric space of the carrier surface i.e. for a Bezier surface the boundary of the space is [[0, 1], [0, 1]] and for a spherical surface [[0, 2Pi], [-Pi / 2, Pi / 2]]
+ *  -EGridSpace::UniformScaled.The average curvilinear length of the iso curves is calculated in each axis.The parametric space is then scaled in each axis such that
+ *            the average curvilinear length between two points P(A, V) and P(B, V) = B - A.For a spherical surface [[0, 2Pi.R / 2], [-Pi.R / 2, Pi.R / 2]] with R the radius of the sphere
+ * -EGridSpace::Scaled.The UnifomeScaled space is rescaled in its most disturbe axis e.g. if V is the most disturbe axis, whatever V, length of[P(A, V), P(B, V)] = (B - A).
+ *            This parametric space is closed to a good flattening of the surface.
+ */
+enum EGridSpace : uint8
+{
+	Default2D = 0,  // do not modified enum value as it is used for a static array
+	Scaled = 1,
+	UniformScaled = 2,
+	EndGridSpace
+};
 
-} // namespace CADKernel
+} // namespace UE::CADKernel
 
diff --git a/Engine/Source/Runtime/Datasmith/CADKernel/Private/CADKernel/Mesh/Meshers/IsoTriangulator.cpp b/Engine/Source/Runtime/Datasmith/CADKernel/Private/CADKernel/Mesh/Meshers/IsoTriangulator.cpp
--- a/Engine/Source/Runtime/Datasmith/CADKernel/Private/CADKernel/Mesh/Meshers/IsoTriangulator.cpp
+++ b/Engine/Source/Runtime/Datasmith/CADKernel/Private/CADKernel/Mesh/Meshers/IsoTriangulator.cpp
@@ -19,7 +19,7 @@
 #include "CADKernel/UI/Display.h"
 #include "CADKernel/Utils/ArrayUtils.h"
 
-namespace CADKernel
+namespace UE::CADKernel
 {
 
 const double FIsoTriangulator::GeometricToMeshingToleranceFactor = 10.;
@@ -3334,4 +3334,4 @@ void FIsoTriangulator::ConnectCellCornerToInnerLoop(FCell& Cell)
 	Cell.SelectSegmentInCandidateSegments(IsoSegmentFactory);
 }
 
-} //namespace CADKernel
\ No newline at end of file
+} //namespace UE::CADKernel
\ No newline at end of file
diff --git a/Engine/Source/Runtime/Datasmith/CADKernel/Private/CADKernel/Mesh/Meshers/IsoTriangulator/IntersectionSegmentTool.h b/Engine/Source/Runtime/Datasmith/CADKernel/Private/CADKernel/Mesh/Meshers/IsoTriangulator/IntersectionSegmentTool.h
--- a/Engine/Source/Runtime/Datasmith/CADKernel/Private/CADKernel/Mesh/Meshers/IsoTriangulator/IntersectionSegmentTool.h
+++ b/Engine/Source/Runtime/Datasmith/CADKernel/Private/CADKernel/Mesh/Meshers/IsoTriangulator/IntersectionSegmentTool.h
@@ -7,209 +7,209 @@
 #include "CADKernel/UI/Message.h"
 #include "CADKernel/UI/Visu.h"
 
-namespace CADKernel
+namespace UE::CADKernel
 {
-	class FGrid;
-	class FIsoNode;
-	class FIsoSegment;
-	class FPoint2D;
+class FGrid;
+class FIsoNode;
+class FIsoSegment;
+class FPoint2D;
 
-	struct FSegment4IntersectionTools
+struct FSegment4IntersectionTools
+{
+	const TSegment<FPoint2D> Segment2D;
+
+	/**
+	 * Segment's axis aligned bounding box
+	 */
+	const FSurfacicBoundary Boundary;
+
+	const FIsoSegment* IsoSegment;
+
+	/**
+	 * Uses as criterion to sort segments to optimize
+	 * AxisMin = Boundary[EIso::IsoU].Min + Boundary[EIso::IsoV].Min
+	 */
+	double AxisMin;
+
+	FSegment4IntersectionTools(const FGrid& Grid, const FIsoSegment& InSegment);
+
+	/**
+	 * WARNING StartPoint, EndPoint must be defined in EGridSpace::UniformScaled
+	 */
+	FSegment4IntersectionTools(const FPoint2D& StartPoint, const FPoint2D& EndPoint)
+		: Segment2D(StartPoint, EndPoint)
+		, Boundary(StartPoint, EndPoint)
+		, IsoSegment(nullptr)
+	{
+		AxisMin = Boundary[EIso::IsoU].Min + Boundary[EIso::IsoV].Min;
+	}
+
+	/**
+	 * @return false if their Boundaries are not intersecting
+	 */
+	bool CouldItIntersect(const FSegment4IntersectionTools& Segment) const
 	{
-		const TSegment<FPoint2D> Segment2D;
-
-		/**
-		 * Segment's axis aligned bounding box
-		 */
-		const FSurfacicBoundary Boundary;
-
-		const FIsoSegment* IsoSegment;
-
-		/**
-		 * Uses as criterion to sort segments to optimize 
-		 * AxisMin = Boundary[EIso::IsoU].Min + Boundary[EIso::IsoV].Min 
-		 */
-		double AxisMin;
-
-		FSegment4IntersectionTools(const FGrid& Grid, const FIsoSegment& InSegment);
-
-		/**
-		 * WARNING StartPoint, EndPoint must be defined in EGridSpace::UniformScaled
-		 */
-		FSegment4IntersectionTools(const FPoint2D& StartPoint, const FPoint2D& EndPoint)
-			: Segment2D(StartPoint, EndPoint)
-			, Boundary(StartPoint, EndPoint)
-			, IsoSegment(nullptr)
+		if ((Boundary[EIso::IsoU].Min > Segment.Boundary[EIso::IsoU].Max) || (Boundary[EIso::IsoV].Min > Segment.Boundary[EIso::IsoV].Max))
 		{
-			AxisMin = Boundary[EIso::IsoU].Min + Boundary[EIso::IsoV].Min;
+			return false;
 		}
 
-		/**
-		 * @return false if their Boundaries are not intersecting 
-		 */
-		bool CouldItIntersect(const FSegment4IntersectionTools& Segment) const
+		if ((Boundary[EIso::IsoU].Max < Segment.Boundary[EIso::IsoU].Min) || (Boundary[EIso::IsoV].Max < Segment.Boundary[EIso::IsoV].Min))
 		{
-			if ((Boundary[EIso::IsoU].Min > Segment.Boundary[EIso::IsoU].Max) || (Boundary[EIso::IsoV].Min > Segment.Boundary[EIso::IsoV].Max))
-			{
-				return false;
-			}
-
-			if ((Boundary[EIso::IsoU].Max < Segment.Boundary[EIso::IsoU].Min) || (Boundary[EIso::IsoV].Max < Segment.Boundary[EIso::IsoV].Min))
-			{
-				return false;
-			}
-			return true;
+			return false;
 		}
-	};
+		return true;
+	}
+};
 
-	class FIntersectionSegmentTool
+class FIntersectionSegmentTool
+{
+	TArray<FSegment4IntersectionTools> Segments;
+	const FGrid& Grid;
+	bool bSegmentsAreSorted;
+
+public:
+	FIntersectionSegmentTool(const FGrid& InGrid)
+		: Grid(InGrid)
+		, bSegmentsAreSorted(false)
 	{
-		TArray<FSegment4IntersectionTools> Segments;
-		const FGrid& Grid;
-		bool bSegmentsAreSorted;
-
-	public:
-		FIntersectionSegmentTool(const FGrid& InGrid)
-			: Grid(InGrid)
-			, bSegmentsAreSorted(false)
-		{
-		}
+	}
 
-		void Empty(int32 InMaxNum)
-		{
-			Segments.Empty(InMaxNum);
-			bSegmentsAreSorted = false;
-		}
+	void Empty(int32 InMaxNum)
+	{
+		Segments.Empty(InMaxNum);
+		bSegmentsAreSorted = false;
+	}
 
-		void Reserve(int32 InMaxNum)
-		{
-			Segments.Reserve(InMaxNum);
-		}
+	void Reserve(int32 InMaxNum)
+	{
+		Segments.Reserve(InMaxNum);
+	}
 
-		void RemoveLast()
-		{
-			ensureCADKernel(!bSegmentsAreSorted);
-			Segments.RemoveAt(Segments.Num() - 1);
-			bSegmentsAreSorted = false;
-		}
+	void RemoveLast()
+	{
+		ensureCADKernel(!bSegmentsAreSorted);
+		Segments.RemoveAt(Segments.Num() - 1);
+		bSegmentsAreSorted = false;
+	}
 
-		void SetCount(int32 NewCount)
+	void SetCount(int32 NewCount)
+	{
+		ensureCADKernel(NewCount < Segments.Num() && !bSegmentsAreSorted);
+		while (NewCount != Segments.Num())
 		{
-			ensureCADKernel(NewCount < Segments.Num() && !bSegmentsAreSorted);
-			while(NewCount != Segments.Num())
-			{
-				Segments.RemoveAt(Segments.Num() - 1);
-			}
+			Segments.RemoveAt(Segments.Num() - 1);
 		}
+	}
 
-		/**
-		 * @return true if the segment is found and removed
-		 */
-		bool Remove(const FIsoSegment* Segment)
-		{
-			int32 SegmentIndex = Segments.IndexOfByPredicate([Segment](const FSegment4IntersectionTools& SegmentIter)
-				{
-					return SegmentIter.IsoSegment == Segment;
-				});
-
-			if (SegmentIndex != INDEX_NONE)
+	/**
+	 * @return true if the segment is found and removed
+	 */
+	bool Remove(const FIsoSegment* Segment)
+	{
+		int32 SegmentIndex = Segments.IndexOfByPredicate([Segment](const FSegment4IntersectionTools& SegmentIter)
 			{
-				Segments.RemoveAt(SegmentIndex);
-				return true;
-			}
-			return false;
-		}
+				return SegmentIter.IsoSegment == Segment;
+			});
 
-		int32 Count()
+		if (SegmentIndex != INDEX_NONE)
 		{
-			return Segments.Num();
+			Segments.RemoveAt(SegmentIndex);
+			return true;
 		}
+		return false;
+	}
 
-		bool Update(const FIsoSegment* Segment)
-		{
-			int32 SegmentIndex = Segments.IndexOfByPredicate([Segment](const FSegment4IntersectionTools& SegmentIter)
-				{
-					return SegmentIter.IsoSegment == Segment;
-				});
+	int32 Count()
+	{
+		return Segments.Num();
+	}
 
-			if (SegmentIndex != INDEX_NONE)
+	bool Update(const FIsoSegment* Segment)
+	{
+		int32 SegmentIndex = Segments.IndexOfByPredicate([Segment](const FSegment4IntersectionTools& SegmentIter)
 			{
-				Segments.RemoveAt(SegmentIndex);
-				Segments.EmplaceAt(SegmentIndex, Grid, *Segment);
-				bSegmentsAreSorted = false;
-				return true;
-			}
-			return false;
+				return SegmentIter.IsoSegment == Segment;
+			});
 
-		}
-
-		void AddSegments(FIsoSegment** InNewSegments, int32 Count)
+		if (SegmentIndex != INDEX_NONE)
 		{
+			Segments.RemoveAt(SegmentIndex);
+			Segments.EmplaceAt(SegmentIndex, Grid, *Segment);
 			bSegmentsAreSorted = false;
-			Segments.Reserve(Count + Segments.Num());
-			for (int32 Index = 0; Index < Count; ++Index)
-			{
-				FIsoSegment* NewSegment = InNewSegments[Index];
-				AddSegment(*NewSegment);
-			}
+			return true;
 		}
+		return false;
 
-		void AddSegments(const TArray<FIsoSegment*>& InNewSegments)
-		{
-			bSegmentsAreSorted = false;
-			Segments.Reserve(InNewSegments.Num() + Segments.Num());
-			for (FIsoSegment* Segment : InNewSegments)
-			{
-				AddSegment(*Segment);
-			}
-		}
+	}
 
-		void AddSegment(const FIsoSegment& Segment)
+	void AddSegments(FIsoSegment** InNewSegments, int32 Count)
+	{
+		bSegmentsAreSorted = false;
+		Segments.Reserve(Count + Segments.Num());
+		for (int32 Index = 0; Index < Count; ++Index)
 		{
-			bSegmentsAreSorted = false;
-			Segments.Emplace(Grid, Segment);
+			FIsoSegment* NewSegment = InNewSegments[Index];
+			AddSegment(*NewSegment);
 		}
+	}
 
-		void AddSegment(const FPoint2D& StartPoint, const FPoint2D& EndPoint)
+	void AddSegments(const TArray<FIsoSegment*>& InNewSegments)
+	{
+		bSegmentsAreSorted = false;
+		Segments.Reserve(InNewSegments.Num() + Segments.Num());
+		for (FIsoSegment* Segment : InNewSegments)
 		{
-			bSegmentsAreSorted = false;
-			Segments.Emplace(StartPoint, EndPoint);
+			AddSegment(*Segment);
 		}
+	}
 
-		const FIsoSegment* DoesIntersect(const FIsoSegment& Segment) const;
-		FIsoSegment* DoesIntersect(const FIsoSegment& Segment);
-
-		/**
-		 * WARNING StartPoint, EndPoint must be defined in EGridSpace::UniformScaled
-		 */
-		bool DoesIntersect(const FPoint2D& StartPoint, const FPoint2D& EndPoint) const;
-
-		/**
-		 * Allow StartNode to be connected to one segment
-		 * WARNING EndPoint must be defined in EGridSpace::UniformScaled
-		 */
-		bool DoesIntersect(const FIsoNode& StartNode, const FPoint2D& EndPoint) const;
-
-		/**
-		 * Allow StartNode and EndNode to be connected to one segment
-		 */
-		const FIsoSegment* DoesIntersect(const FIsoNode& StartNode, const FIsoNode& EndNode) const;
-
-		bool FindIntersections(const FIsoNode& StartNode, const FIsoNode& EndNode, TArray<const FIsoSegment*>& OutIntersections) const;
+	void AddSegment(const FIsoSegment& Segment)
+	{
+		bSegmentsAreSorted = false;
+		Segments.Emplace(Grid, Segment);
+	}
 
-		/**
-		 * segments are sorted by DMin increasing
-		 */
-		void Sort()
-		{
-			Algo::Sort(Segments, [](const FSegment4IntersectionTools& Segment1, const FSegment4IntersectionTools& Segment2) { return Segment1.AxisMin < Segment2.AxisMin; });
-			bSegmentsAreSorted = true;
-		}
+	void AddSegment(const FPoint2D& StartPoint, const FPoint2D& EndPoint)
+	{
+		bSegmentsAreSorted = false;
+		Segments.Emplace(StartPoint, EndPoint);
+	}
+
+	const FIsoSegment* DoesIntersect(const FIsoSegment& Segment) const;
+	FIsoSegment* DoesIntersect(const FIsoSegment& Segment);
+
+	/**
+	 * WARNING StartPoint, EndPoint must be defined in EGridSpace::UniformScaled
+	 */
+	bool DoesIntersect(const FPoint2D& StartPoint, const FPoint2D& EndPoint) const;
+
+	/**
+	 * Allow StartNode to be connected to one segment
+	 * WARNING EndPoint must be defined in EGridSpace::UniformScaled
+	 */
+	bool DoesIntersect(const FIsoNode& StartNode, const FPoint2D& EndPoint) const;
+
+	/**
+	 * Allow StartNode and EndNode to be connected to one segment
+	 */
+	const FIsoSegment* DoesIntersect(const FIsoNode& StartNode, const FIsoNode& EndNode) const;
+
+	bool FindIntersections(const FIsoNode& StartNode, const FIsoNode& EndNode, TArray<const FIsoSegment*>& OutIntersections) const;
+
+	/**
+	 * segments are sorted by DMin increasing
+	 */
+	void Sort()
+	{
+		Algo::Sort(Segments, [](const FSegment4IntersectionTools& Segment1, const FSegment4IntersectionTools& Segment2) { return Segment1.AxisMin < Segment2.AxisMin; });
+		bSegmentsAreSorted = true;
+	}
 
 #ifdef CADKERNEL_DEV
-		void Display(const TCHAR* Message, EVisuProperty Property = EVisuProperty::BlueCurve) const;
+	void Display(const TCHAR* Message, EVisuProperty Property = EVisuProperty::BlueCurve) const;
 #endif
 
-	};
-} // namespace CADKernel
+};
+} // namespace UE::CADKernel
 
diff --git a/Engine/Source/Runtime/Datasmith/CADKernel/Private/CADKernel/Mesh/Meshers/IsoTriangulator/IsoCell.h b/Engine/Source/Runtime/Datasmith/CADKernel/Private/CADKernel/Mesh/Meshers/IsoTriangulator/IsoCell.h
--- a/Engine/Source/Runtime/Datasmith/CADKernel/Private/CADKernel/Mesh/Meshers/IsoTriangulator/IsoCell.h
+++ b/Engine/Source/Runtime/Datasmith/CADKernel/Private/CADKernel/Mesh/Meshers/IsoTriangulator/IsoCell.h
@@ -8,162 +8,162 @@
 #include "CADKernel/Mesh/Structure/Grid.h"
 #include "CADKernel/Topo/TopologicalFace.h"
 
-namespace CADKernel
+namespace UE::CADKernel
 {
-	class FGrid;
+class FGrid;
 
-	struct FCell
+struct FCell
+{
+	const FGrid& Grid;
+	int32 Id;
+	int32 LoopNodeCount;
+	bool bHasOuterLoop = false;
+	int32 InnerLoopCount = 0;
+
+	TArray<int32> LoopIndexToIndex;
+	TArray<TArray<FLoopNode*>> SubLoops;
+
+	// the outer loop is subdivide into connected (in the cell) nodes
+	TArray<TArray<FLoopNode*>> OuterLoopSubdivision;
+
+	TArray<FIsoSegment*> CandidateSegments;
+	TArray<FIsoSegment*> FinalSegments;
+	FIntersectionSegmentTool IntersectionTool;
+
+	/** Loop that barycenter is on the border of the barycenter mesh */
+	TSet<int32> BorderLoopIndices;
+
+	FCell(const int32 InLoopIndex, const TArray<FLoopNode*>& InNodes, const FGrid& InGrid)
+		: Grid(InGrid)
+		, Id(InLoopIndex)
+		, LoopNodeCount(InNodes.Num())
+		, IntersectionTool(Grid)
 	{
-		const FGrid& Grid;
-		int32 Id;
-		int32 LoopNodeCount;
-		bool bHasOuterLoop = false;
-		int32 InnerLoopCount = 0;
-
-		TArray<int32> LoopIndexToIndex;
-		TArray<TArray<FLoopNode*>> SubLoops;
-
-		// the outer loop is subdivide into connected (in the cell) nodes
-		TArray<TArray<FLoopNode*>> OuterLoopSubdivision;
-
-		TArray<FIsoSegment*> CandidateSegments;
-		TArray<FIsoSegment*> FinalSegments;
-		FIntersectionSegmentTool IntersectionTool;
-
-		/** Loop that barycenter is on the border of the barycenter mesh */
-		TSet<int32> BorderLoopIndices;
-
-		FCell(const int32 InLoopIndex, const TArray<FLoopNode*>& InNodes, const FGrid& InGrid)
-			: Grid(InGrid)
-			, Id(InLoopIndex)
-			, LoopNodeCount(InNodes.Num())
-			, IntersectionTool(Grid)
-		{
-			ensureCADKernel(InNodes.Num() > 0);
+		ensureCADKernel(InNodes.Num() > 0);
 
-			int32 FaceLoopCount = Grid.GetFace().LoopCount();
-			int32 LoopCount = 0;
+		int32 FaceLoopCount = Grid.GetFace().LoopCount();
+		int32 LoopCount = 0;
+		{
+			TArray<int32> LoopVertexCount;
+			LoopVertexCount.Init(0, FaceLoopCount);
+			for (const FLoopNode* Node : InNodes)
 			{
-				TArray<int32> LoopVertexCount;
-				LoopVertexCount.Init(0, FaceLoopCount);
-				for (const FLoopNode* Node : InNodes)
-				{
-					LoopVertexCount[Node->GetLoopIndex()]++;
-				}
+				LoopVertexCount[Node->GetLoopIndex()]++;
+			}
 
-				for (const int32& VertexCount : LoopVertexCount)
+			for (const int32& VertexCount : LoopVertexCount)
+			{
+				if (VertexCount > 0)
 				{
-					if (VertexCount > 0)
-					{
-						LoopCount++;
-					}
+					LoopCount++;
 				}
+			}
 
-				SubLoops.SetNum(LoopCount);
-				int32 Index = 0;
-				for (int32 LoopIndex = 0; LoopIndex < FaceLoopCount; ++LoopIndex)
+			SubLoops.SetNum(LoopCount);
+			int32 Index = 0;
+			for (int32 LoopIndex = 0; LoopIndex < FaceLoopCount; ++LoopIndex)
+			{
+				if (LoopVertexCount[LoopIndex] > 0)
 				{
-					if (LoopVertexCount[LoopIndex] > 0)
-					{
-						SubLoops[Index].Reserve(LoopVertexCount[LoopIndex]);
-						LoopVertexCount[LoopIndex] = Index;
-						Index++;
-					}
+					SubLoops[Index].Reserve(LoopVertexCount[LoopIndex]);
+					LoopVertexCount[LoopIndex] = Index;
+					Index++;
 				}
-				Swap(LoopIndexToIndex, LoopVertexCount);
 			}
+			Swap(LoopIndexToIndex, LoopVertexCount);
+		}
 
-			for (FLoopNode* Node : InNodes)
-			{
-				int32 Index = LoopIndexToIndex[Node->GetLoopIndex()];
-				SubLoops[Index].Add(Node);
-			}
+		for (FLoopNode* Node : InNodes)
+		{
+			int32 Index = LoopIndexToIndex[Node->GetLoopIndex()];
+			SubLoops[Index].Add(Node);
+		}
 
-			bHasOuterLoop = (SubLoops[0][0]->GetLoopIndex() == 0);
-			InnerLoopCount = bHasOuterLoop ? (SubLoops.Num() - 1): SubLoops.Num();
+		bHasOuterLoop = (SubLoops[0][0]->GetLoopIndex() == 0);
+		InnerLoopCount = bHasOuterLoop ? (SubLoops.Num() - 1) : SubLoops.Num();
 
-			if (bHasOuterLoop)
+		if (bHasOuterLoop)
+		{
+			// Subdivide Loop0 in SubLoop
+			TArray<FLoopNode*>& Loop0 = SubLoops[0];
+			Algo::Sort(Loop0, [&](const FLoopNode* LoopNode1, const FLoopNode* LoopNode2)
+				{
+					return LoopNode1->GetIndex() < LoopNode2->GetIndex();
+				});
+
+			OuterLoopSubdivision.Reserve(Loop0.Num());
+			FLoopNode* PreviousNode = nullptr;
+			TArray<FLoopNode*>* SubLoop = nullptr;
+			for (FLoopNode* Node : Loop0)
 			{
-				// Subdivide Loop0 in SubLoop
-				TArray<FLoopNode*>& Loop0 = SubLoops[0];
-				Algo::Sort(Loop0, [&](const FLoopNode* LoopNode1, const FLoopNode* LoopNode2)
-					{
-						return LoopNode1->GetIndex() < LoopNode2->GetIndex();
-					});
-
-				OuterLoopSubdivision.Reserve(Loop0.Num());
-				FLoopNode* PreviousNode = nullptr;
-				TArray<FLoopNode*>* SubLoop = nullptr;
-				for (FLoopNode* Node : Loop0)
+				if (&Node->GetPreviousNode() != PreviousNode)
 				{
-					if (&Node->GetPreviousNode() != PreviousNode)
-					{
-						SubLoop = &OuterLoopSubdivision.Emplace_GetRef();
-					}
-					SubLoop->Add(Node);
-					PreviousNode = Node;
+					SubLoop = &OuterLoopSubdivision.Emplace_GetRef();
 				}
+				SubLoop->Add(Node);
+				PreviousNode = Node;
 			}
 		}
+	}
 
-		void SelectSegmentInCandidateSegments(TFactory<FIsoSegment>& SegmentFactory)
-		{
+	void SelectSegmentInCandidateSegments(TFactory<FIsoSegment>& SegmentFactory)
+	{
 #ifdef DEBUG_SELECT_SEGMENT
-			F3DDebugSession _(Grid.bDisplay, TEXT("SelectSegmentInCandidateSegments "));
-			IntersectionTool.Display(TEXT("Cell.IntersectionTool at SelectSegmentInCandidateSegments start"));
-			//Wait();
+		F3DDebugSession _(Grid.bDisplay, TEXT("SelectSegmentInCandidateSegments "));
+		IntersectionTool.Display(TEXT("Cell.IntersectionTool at SelectSegmentInCandidateSegments start"));
+		//Wait();
 #endif
 
-			Algo::Sort(CandidateSegments, [&](const FIsoSegment* Segment1, const FIsoSegment* Segment2)
-				{
-					return Segment1->Get2DLengthSquare(EGridSpace::UniformScaled, Grid) < Segment2->Get2DLengthSquare(EGridSpace::UniformScaled, Grid);
-				});
-
-			// Validate all candidate segments
-			for (FIsoSegment* Segment : CandidateSegments)
+		Algo::Sort(CandidateSegments, [&](const FIsoSegment* Segment1, const FIsoSegment* Segment2)
 			{
+				return Segment1->Get2DLengthSquare(EGridSpace::UniformScaled, Grid) < Segment2->Get2DLengthSquare(EGridSpace::UniformScaled, Grid);
+			});
+
+		// Validate all candidate segments
+		for (FIsoSegment* Segment : CandidateSegments)
+		{
 #ifdef DEBUG_SELECT_SEGMENT
-				F3DDebugSession _(Grid.bDisplay, TEXT("Segment"));
+			F3DDebugSession _(Grid.bDisplay, TEXT("Segment"));
 #endif
-				if (IntersectionTool.DoesIntersect(*Segment))
-				{
+			if (IntersectionTool.DoesIntersect(*Segment))
+			{
 #ifdef DEBUG_SELECT_SEGMENT
-					F3DDebugSession _(Grid.bDisplay, TEXT("SelectSegmentInCandidateSegments "));
-					Grid.DisplayIsoSegment(EGridSpace::UniformScaled, *Segment, YellowCurve);
+				F3DDebugSession _(Grid.bDisplay, TEXT("SelectSegmentInCandidateSegments "));
+				Grid.DisplayIsoSegment(EGridSpace::UniformScaled, *Segment, YellowCurve);
 #endif
-					SegmentFactory.DeleteEntity(Segment);
-					continue;
-				}
+				SegmentFactory.DeleteEntity(Segment);
+				continue;
+			}
 
-				if (FIsoSegment::IsItAlreadyDefined(&Segment->GetFirstNode(), &Segment->GetSecondNode()))
-				{
+			if (FIsoSegment::IsItAlreadyDefined(&Segment->GetFirstNode(), &Segment->GetSecondNode()))
+			{
 #ifdef DEBUG_SELECT_SEGMENT
-					Grid.DisplayIsoSegment(EGridSpace::UniformScaled, *Segment, GreenCurve);
+				Grid.DisplayIsoSegment(EGridSpace::UniformScaled, *Segment, GreenCurve);
 #endif
-					SegmentFactory.DeleteEntity(Segment);
-					continue;
-				}
+				SegmentFactory.DeleteEntity(Segment);
+				continue;
+			}
 
 #ifdef DEBUG_SELECT_SEGMENT
-				Grid.DisplayIsoSegment(EGridSpace::UniformScaled, *Segment, BlueCurve);
+			Grid.DisplayIsoSegment(EGridSpace::UniformScaled, *Segment, BlueCurve);
 #endif
 
-				FinalSegments.Add(Segment);
-				IntersectionTool.AddSegment(*Segment);
-				Segment->SetSelected();
-				Segment->ConnectToNode();
-			}
-			CandidateSegments.Empty();
+			FinalSegments.Add(Segment);
+			IntersectionTool.AddSegment(*Segment);
+			Segment->SetSelected();
+			Segment->ConnectToNode();
 		}
+		CandidateSegments.Empty();
+	}
 
-		bool Contains(FLoopNode* NodeToFind)
-		{
-			int32 LoopIndex = NodeToFind->GetLoopIndex();
-			int32 Index = LoopIndexToIndex[LoopIndex];
-			return SubLoops[Index].Find(NodeToFind) != INDEX_NONE;
-		}
+	bool Contains(FLoopNode* NodeToFind)
+	{
+		int32 LoopIndex = NodeToFind->GetLoopIndex();
+		int32 Index = LoopIndexToIndex[LoopIndex];
+		return SubLoops[Index].Find(NodeToFind) != INDEX_NONE;
+	}
 
-	};
+};
 
-} // namespace CADKernel
+} // namespace UE::CADKernel
 
diff --git a/Engine/Source/Runtime/Datasmith/CADKernel/Private/CADKernel/Mesh/Meshers/IsoTriangulator/IsoNode.h b/Engine/Source/Runtime/Datasmith/CADKernel/Private/CADKernel/Mesh/Meshers/IsoTriangulator/IsoNode.h
--- a/Engine/Source/Runtime/Datasmith/CADKernel/Private/CADKernel/Mesh/Meshers/IsoTriangulator/IsoNode.h
+++ b/Engine/Source/Runtime/Datasmith/CADKernel/Private/CADKernel/Mesh/Meshers/IsoTriangulator/IsoNode.h
@@ -5,7 +5,7 @@
 #include "CADKernel/Mesh/MeshEnum.h"
 #include "CADKernel/Mesh/Structure/Grid.h"
 
-namespace CADKernel
+namespace UE::CADKernel
 {
 class FGrid;
 class FIsoSegment;
@@ -86,7 +86,7 @@ public:
 		States = EIsoNodeStates::Delete;
 	}
 
-	bool IsDelete() const 
+	bool IsDelete() const
 	{
 		return (States & EIsoNodeStates::Delete) == EIsoNodeStates::Delete;
 	}
@@ -529,5 +529,5 @@ public:
 };
 
 
-} // namespace CADKernel
+} // namespace UE::CADKernel
 
diff --git a/Engine/Source/Runtime/Datasmith/CADKernel/Private/CADKernel/Mesh/Meshers/MesherTools.cpp b/Engine/Source/Runtime/Datasmith/CADKernel/Private/CADKernel/Mesh/Meshers/MesherTools.cpp
--- a/Engine/Source/Runtime/Datasmith/CADKernel/Private/CADKernel/Mesh/Meshers/MesherTools.cpp
+++ b/Engine/Source/Runtime/Datasmith/CADKernel/Private/CADKernel/Mesh/Meshers/MesherTools.cpp
@@ -5,297 +5,297 @@
 #include "CADKernel/Math/Boundary.h"
 #include "CADKernel/Topo/TopologicalEdge.h"
 
-namespace CADKernel
+namespace UE::CADKernel
 {
-	
-	void FMesherTools::ComputeFinalCuttingPointsWithPreferredCuttingPoints(const TArray<double>& CrossingUs, TArray<double> DeltaUs, const TArray<FCuttingPoint>& PreferredCuttingPoints, const FLinearBoundary& Boundary, TArray<double>& OutCuttingPoints)
+
+void FMesherTools::ComputeFinalCuttingPointsWithPreferredCuttingPoints(const TArray<double>& CrossingUs, TArray<double> DeltaUs, const TArray<FCuttingPoint>& PreferredCuttingPoints, const FLinearBoundary& Boundary, TArray<double>& OutCuttingPoints)
+{
+	// Evaluate the number of Cutting points
 	{
-		// Evaluate the number of Cutting points
+		double DeltaUMin = HUGE_VALUE;
+		for (const double& Delta : DeltaUs)
 		{
-			double DeltaUMin = HUGE_VALUE;
-			for (const double& Delta : DeltaUs)
+			if (Delta < DeltaUMin)
 			{
-				if (Delta < DeltaUMin)
-				{
-					DeltaUMin = Delta;
-				}
+				DeltaUMin = Delta;
 			}
-			int32 NodeNum = 5 + (int32)(Boundary.Length() / DeltaUMin);
-			OutCuttingPoints.Reserve(NodeNum);
 		}
+		int32 NodeNum = 5 + (int32)(Boundary.Length() / DeltaUMin);
+		OutCuttingPoints.Reserve(NodeNum);
+	}
 
-		int32 StartCrossingIndex = 0;
-		while (CrossingUs[StartCrossingIndex + 1] <= Boundary.Min)
-		{
-			++StartCrossingIndex;
-		};
+	int32 StartCrossingIndex = 0;
+	while (CrossingUs[StartCrossingIndex + 1] <= Boundary.Min)
+	{
+		++StartCrossingIndex;
+	};
 
-		// Step 1: Compute the number of mesh nodes i.e.compute the time of the travel between Boundary.Min and Boundary.Max with the initial speeds DeltaUs
-		int32 CrossingIndex = StartCrossingIndex;
-		double LastStep = Boundary.Min;
-		double TravelTime = 0;
-		do
-		{
-			double NextStep = FMath::Min(CrossingUs[CrossingIndex+1], Boundary.Max);
-			TravelTime += (NextStep - LastStep) / DeltaUs[CrossingIndex];
-			CrossingIndex++;
-			LastStep = NextStep;
-		} while (LastStep + DOUBLE_SMALL_NUMBER < Boundary.Max);
+	// Step 1: Compute the number of mesh nodes i.e.compute the time of the travel between Boundary.Min and Boundary.Max with the initial speeds DeltaUs
+	int32 CrossingIndex = StartCrossingIndex;
+	double LastStep = Boundary.Min;
+	double TravelTime = 0;
+	do
+	{
+		double NextStep = FMath::Min(CrossingUs[CrossingIndex + 1], Boundary.Max);
+		TravelTime += (NextStep - LastStep) / DeltaUs[CrossingIndex];
+		CrossingIndex++;
+		LastStep = NextStep;
+	} while (LastStep + DOUBLE_SMALL_NUMBER < Boundary.Max);
 
 
-		// Step 2: Define the number of step(each step lasts one) as it must be an integer bigger than 1
-		double CorrectedFinalTravelTime = FMath::Max(((double)((int32)(TravelTime + 0.5))), 1.);
+	// Step 2: Define the number of step(each step lasts one) as it must be an integer bigger than 1
+	double CorrectedFinalTravelTime = FMath::Max(((double)((int32)(TravelTime + 0.5))), 1.);
 
-		// Step 3: Adjust the speeds than the travel respect the count of step
+	// Step 3: Adjust the speeds than the travel respect the count of step
+	{
+		double CorrectiveFactor = TravelTime / CorrectedFinalTravelTime;
+		for (double& DeltaU : DeltaUs)
 		{
-			double CorrectiveFactor = TravelTime / CorrectedFinalTravelTime;
-			for (double& DeltaU : DeltaUs)
-			{
-				DeltaU *= CorrectiveFactor;
-			}
+			DeltaU *= CorrectiveFactor;
 		}
+	}
 
-		TArray<double> TravelTimeOfNeighborsU;
-		TravelTimeOfNeighborsU.Reserve(PreferredCuttingPoints.Num() + 1);
-		TravelTimeOfNeighborsU.SetNum(PreferredCuttingPoints.Num());
+	TArray<double> TravelTimeOfNeighborsU;
+	TravelTimeOfNeighborsU.Reserve(PreferredCuttingPoints.Num() + 1);
+	TravelTimeOfNeighborsU.SetNum(PreferredCuttingPoints.Num());
 
-		// Defined the travel time to each Preferred Cutting Points
+	// Defined the travel time to each Preferred Cutting Points
+	{
+		LastStep = Boundary.Min;
+		TravelTime = 0;
+		int32 UIndex = StartCrossingIndex;
+		for (int32 NeighbourIndex = 0; NeighbourIndex < PreferredCuttingPoints.Num(); ++NeighbourIndex)
 		{
-			LastStep = Boundary.Min;
-			TravelTime = 0;
-			int32 UIndex = StartCrossingIndex;
-			for (int32 NeighbourIndex = 0; NeighbourIndex < PreferredCuttingPoints.Num(); ++NeighbourIndex)
+			for (; UIndex < DeltaUs.Num(); ++UIndex)
 			{
-				for (; UIndex < DeltaUs.Num(); ++UIndex)
+				double DeltaU = DeltaUs[UIndex];
+				if (PreferredCuttingPoints[NeighbourIndex].Coordinate - DOUBLE_SMALL_NUMBER > CrossingUs[UIndex + 1])
 				{
-					double DeltaU = DeltaUs[UIndex];
-					if (PreferredCuttingPoints[NeighbourIndex].Coordinate - DOUBLE_SMALL_NUMBER > CrossingUs[UIndex + 1])
-					{
-						double DeltaTravelTime = (CrossingUs[UIndex + 1] - LastStep) / DeltaU;
-						LastStep = CrossingUs[UIndex + 1];
-						TravelTime += DeltaTravelTime;
-						continue;
-					}
-
-					if (CrossingUs[UIndex] - DOUBLE_SMALL_NUMBER < PreferredCuttingPoints[NeighbourIndex].Coordinate && PreferredCuttingPoints[NeighbourIndex].Coordinate < CrossingUs[UIndex + 1] + DOUBLE_SMALL_NUMBER)
-					{
-						double DeltaTravelTime = (PreferredCuttingPoints[NeighbourIndex].Coordinate - LastStep) / DeltaU;
-						double TravelTimeToNeighbourU = TravelTime + DeltaTravelTime;
-						TravelTimeOfNeighborsU[NeighbourIndex] = TravelTimeToNeighbourU;
-					}
-					break;
+					double DeltaTravelTime = (CrossingUs[UIndex + 1] - LastStep) / DeltaU;
+					LastStep = CrossingUs[UIndex + 1];
+					TravelTime += DeltaTravelTime;
+					continue;
 				}
-			}
 
-			// if some NeighborsU are near the end (less than 0.45x the expected element size), they are ignored 
-			int32 NeighbourIndex = 0;
-			for (; NeighbourIndex < TravelTimeOfNeighborsU.Num(); ++NeighbourIndex)
-			{
-				if (TravelTimeOfNeighborsU[NeighbourIndex] > CorrectedFinalTravelTime - 0.45)
+				if (CrossingUs[UIndex] - DOUBLE_SMALL_NUMBER < PreferredCuttingPoints[NeighbourIndex].Coordinate && PreferredCuttingPoints[NeighbourIndex].Coordinate < CrossingUs[UIndex + 1] + DOUBLE_SMALL_NUMBER)
 				{
-					break;
+					double DeltaTravelTime = (PreferredCuttingPoints[NeighbourIndex].Coordinate - LastStep) / DeltaU;
+					double TravelTimeToNeighbourU = TravelTime + DeltaTravelTime;
+					TravelTimeOfNeighborsU[NeighbourIndex] = TravelTimeToNeighbourU;
 				}
+				break;
 			}
+		}
 
-			if (NeighbourIndex == TravelTimeOfNeighborsU.Num())
-			{
-				TravelTimeOfNeighborsU.Add(CorrectedFinalTravelTime);
-			}
-			else
+		// if some NeighborsU are near the end (less than 0.45x the expected element size), they are ignored 
+		int32 NeighbourIndex = 0;
+		for (; NeighbourIndex < TravelTimeOfNeighborsU.Num(); ++NeighbourIndex)
+		{
+			if (TravelTimeOfNeighborsU[NeighbourIndex] > CorrectedFinalTravelTime - 0.45)
 			{
-				TravelTimeOfNeighborsU[NeighbourIndex] = CorrectedFinalTravelTime;
-				TravelTimeOfNeighborsU.SetNum(NeighbourIndex + 1);
+				break;
 			}
 		}
 
+		if (NeighbourIndex == TravelTimeOfNeighborsU.Num())
+		{
+			TravelTimeOfNeighborsU.Add(CorrectedFinalTravelTime);
+		}
+		else
+		{
+			TravelTimeOfNeighborsU[NeighbourIndex] = CorrectedFinalTravelTime;
+			TravelTimeOfNeighborsU.SetNum(NeighbourIndex + 1);
+		}
+	}
+
 
-		int32 NeighbourIndex = 0;
+	int32 NeighbourIndex = 0;
 
-		double TargetTime = 0.;
-		double TimeDelta = CorrectedFinalTravelTime;
-		// Define the next Target time step i.e. a Preferred Cutting Points or a round step   
-		TFunction<void(double&, double&)> IncreaseTargetTimeAccordingToNeighbors = [&](double& TargetTime, double& TimeDelta)
+	double TargetTime = 0.;
+	double TimeDelta = CorrectedFinalTravelTime;
+	// Define the next Target time step i.e. a Preferred Cutting Points or a round step   
+	TFunction<void(double&, double&)> IncreaseTargetTimeAccordingToNeighbors = [&](double& TargetTime, double& TimeDelta)
+	{
+		if (NeighbourIndex < TravelTimeOfNeighborsU.Num() && TargetTime + TimeDelta + DOUBLE_SMALL_NUMBER > TravelTimeOfNeighborsU[NeighbourIndex])
 		{
-			if (NeighbourIndex < TravelTimeOfNeighborsU.Num() && TargetTime + TimeDelta + DOUBLE_SMALL_NUMBER > TravelTimeOfNeighborsU[NeighbourIndex])
+			TimeDelta = CorrectedFinalTravelTime - TargetTime;
+			for (; NeighbourIndex < TravelTimeOfNeighborsU.Num(); ++NeighbourIndex)
 			{
-				TimeDelta = CorrectedFinalTravelTime - TargetTime;
-				for (; NeighbourIndex < TravelTimeOfNeighborsU.Num(); ++NeighbourIndex)
+				double BestGap = TravelTimeOfNeighborsU[NeighbourIndex] - TargetTime;
+				if (BestGap < 0.45)
 				{
-					double BestGap = TravelTimeOfNeighborsU[NeighbourIndex] - TargetTime;
-					if (BestGap < 0.45)
-					{
-						continue;
-					}
+					continue;
+				}
 
-					int32 BestIndex = NeighbourIndex;
-					if (BestGap < 1)
+				int32 BestIndex = NeighbourIndex;
+				if (BestGap < 1)
+				{
+					BestGap = FMath::Abs(1 - BestGap);
+					for (int32 Index = NeighbourIndex + 1; Index < TravelTimeOfNeighborsU.Num(); ++Index)
 					{
-						BestGap = FMath::Abs(1 - BestGap);
-						for (int32 Index = NeighbourIndex + 1; Index < TravelTimeOfNeighborsU.Num(); ++Index)
+						double Gap = FMath::Abs(1 - TravelTimeOfNeighborsU[Index] + TargetTime);
+						if (Gap < BestGap)
+						{
+							BestGap = Gap;
+							BestIndex = Index;
+						}
+						else
 						{
-							double Gap = FMath::Abs(1 - TravelTimeOfNeighborsU[Index] + TargetTime);
-							if (Gap < BestGap)
-							{
-								BestGap = Gap;
-								BestIndex = Index;
-							}
-							else
-							{
-								break;
-							}
+							break;
 						}
 					}
-					NeighbourIndex = BestIndex;
-					TimeDelta = TravelTimeOfNeighborsU[NeighbourIndex] - TargetTime;
-					int32 StepNum = FMath::Max((int32)(TimeDelta + 0.5), 1);
-					TimeDelta /= StepNum;
-					break;
 				}
+				NeighbourIndex = BestIndex;
+				TimeDelta = TravelTimeOfNeighborsU[NeighbourIndex] - TargetTime;
+				int32 StepNum = FMath::Max((int32)(TimeDelta + 0.5), 1);
+				TimeDelta /= StepNum;
+				break;
 			}
+		}
 
-			TargetTime += TimeDelta;
-		};
+		TargetTime += TimeDelta;
+	};
 
-		// Step 4: Define the steps (if a preferred step is near the next step, the preferred step is used)
+	// Step 4: Define the steps (if a preferred step is near the next step, the preferred step is used)
 
-		LastStep = Boundary.Min;
-		TravelTime = 0;
-		OutCuttingPoints.Add(Boundary.Min);
+	LastStep = Boundary.Min;
+	TravelTime = 0;
+	OutCuttingPoints.Add(Boundary.Min);
 
-		IncreaseTargetTimeAccordingToNeighbors(TargetTime, TimeDelta);
-		for (CrossingIndex = StartCrossingIndex; CrossingIndex < DeltaUs.Num();)
+	IncreaseTargetTimeAccordingToNeighbors(TargetTime, TimeDelta);
+	for (CrossingIndex = StartCrossingIndex; CrossingIndex < DeltaUs.Num();)
+	{
+		double& DeltaU = DeltaUs[CrossingIndex];
+		double DeltaTravelTime = (CrossingUs[CrossingIndex + 1] - LastStep) / DeltaU;
+		double MaxTravelTime = TravelTime + DeltaTravelTime;
+		while (TimeDelta > DOUBLE_SMALL_NUMBER && MaxTravelTime > TargetTime - DOUBLE_SMALL_NUMBER)
 		{
-			double& DeltaU = DeltaUs[CrossingIndex];
-			double DeltaTravelTime = (CrossingUs[CrossingIndex + 1] - LastStep) / DeltaU;
-			double MaxTravelTime = TravelTime + DeltaTravelTime;
-			while (TimeDelta > DOUBLE_SMALL_NUMBER && MaxTravelTime > TargetTime - DOUBLE_SMALL_NUMBER)
-			{
-				DeltaTravelTime = TargetTime - TravelTime;
-				LastStep += DeltaU * DeltaTravelTime;
-				OutCuttingPoints.Add(LastStep);
-				TravelTime = TargetTime;
+			DeltaTravelTime = TargetTime - TravelTime;
+			LastStep += DeltaU * DeltaTravelTime;
+			OutCuttingPoints.Add(LastStep);
+			TravelTime = TargetTime;
 
-				IncreaseTargetTimeAccordingToNeighbors(TargetTime, TimeDelta);
-			}
-			
-			if (OutCuttingPoints.Last() + DOUBLE_SMALL_NUMBER > Boundary.GetMax())
-			{
-				break;
-			}
+			IncreaseTargetTimeAccordingToNeighbors(TargetTime, TimeDelta);
+		}
 
-			CrossingIndex++;
-			LastStep = CrossingUs[CrossingIndex];
-			TravelTime = MaxTravelTime;
+		if (OutCuttingPoints.Last() + DOUBLE_SMALL_NUMBER > Boundary.GetMax())
+		{
+			break;
 		}
 
-		// Set the last to insure the good value
-		OutCuttingPoints.Last() = Boundary.GetMax();
-	};
+		CrossingIndex++;
+		LastStep = CrossingUs[CrossingIndex];
+		TravelTime = MaxTravelTime;
+	}
+
+	// Set the last to insure the good value
+	OutCuttingPoints.Last() = Boundary.GetMax();
+};
 
-	void FMesherTools::ComputeFinalCuttingPointsWithImposedCuttingPoints(const TArray<double>& CrossingUs, const TArray<double>& DeltaUs, const TArray<FCuttingPoint>& ImposedCuttingPoints, TArray<double>& OutCuttingPointUs)
+void FMesherTools::ComputeFinalCuttingPointsWithImposedCuttingPoints(const TArray<double>& CrossingUs, const TArray<double>& DeltaUs, const TArray<FCuttingPoint>& ImposedCuttingPoints, TArray<double>& OutCuttingPointUs)
+{
+	// Evaluate the number of Cutting points
+	TArray<double> TravelTimeByStep;
 	{
-		// Evaluate the number of Cutting points
-		TArray<double> TravelTimeByStep;
+		double DeltaUMin = HUGE_VALUE;
+		for (const double& Delta : DeltaUs)
 		{
-			double DeltaUMin = HUGE_VALUE;
-			for (const double& Delta : DeltaUs)
+			if (Delta < DeltaUMin)
 			{
-				if (Delta < DeltaUMin)
-				{
-					DeltaUMin = Delta;
-				}
+				DeltaUMin = Delta;
 			}
-			int32 NodeNum = 5 + (int32)((ImposedCuttingPoints.Last().Coordinate - ImposedCuttingPoints[0].Coordinate) / DeltaUMin);
-			OutCuttingPointUs.Reserve(NodeNum);
-			TravelTimeByStep.Reserve(NodeNum);
 		}
+		int32 NodeNum = 5 + (int32)((ImposedCuttingPoints.Last().Coordinate - ImposedCuttingPoints[0].Coordinate) / DeltaUMin);
+		OutCuttingPointUs.Reserve(NodeNum);
+		TravelTimeByStep.Reserve(NodeNum);
+	}
 
-		int32 LastCrossingIndex = 0;
-		int32 CrossingIndex = 0;
+	int32 LastCrossingIndex = 0;
+	int32 CrossingIndex = 0;
 
-		// Laps time at crossing points between UMin et UMax
-		TFunction<void(const double, const double, double&, double&)> LapTimesUntilNextPoint = [&](const double StartU, const double EndU, double& StepCount, double& CorrectiveFactor)
-		{
-			CrossingIndex = LastCrossingIndex;
-			double LastStep = StartU;
-			double TravelTime = 0;
+	// Laps time at crossing points between UMin et UMax
+	TFunction<void(const double, const double, double&, double&)> LapTimesUntilNextPoint = [&](const double StartU, const double EndU, double& StepCount, double& CorrectiveFactor)
+	{
+		CrossingIndex = LastCrossingIndex;
+		double LastStep = StartU;
+		double TravelTime = 0;
 
-			TravelTimeByStep.Empty();
-			for (;;)
+		TravelTimeByStep.Empty();
+		for (;;)
+		{
+			double Step = CrossingUs[CrossingIndex + 1];
+			Step = FMath::Min(Step, EndU);
+			double StepTime = (Step - LastStep) / DeltaUs[CrossingIndex];
+			TravelTimeByStep.Add(StepTime);
+			TravelTime += StepTime;
+			LastStep = Step;
+			if (LastStep + DOUBLE_SMALL_NUMBER > EndU)
 			{
-				double Step = CrossingUs[CrossingIndex + 1];
-				Step = FMath::Min(Step, EndU);
-				double StepTime = (Step - LastStep) / DeltaUs[CrossingIndex];
-				TravelTimeByStep.Add(StepTime);
-				TravelTime += StepTime;
-				LastStep = Step;
-				if (LastStep + DOUBLE_SMALL_NUMBER > EndU)
-				{
-					break;
-				}
-				CrossingIndex++;
+				break;
 			}
+			CrossingIndex++;
+		}
 
-			StepCount = FMath::Max(1., ((double)((int32)(TravelTime + 0.5))));
-			CorrectiveFactor = TravelTime / StepCount;
-		};
+		StepCount = FMath::Max(1., ((double)((int32)(TravelTime + 0.5))));
+		CorrectiveFactor = TravelTime / StepCount;
+	};
 
-		TFunction<void(const double, const double, const double)> ComputeCuttingPointUntilNextImposedPoint = [&](const double EndU, const double NbStep, const double CorrectiveFactor)
+	TFunction<void(const double, const double, const double)> ComputeCuttingPointUntilNextImposedPoint = [&](const double EndU, const double NbStep, const double CorrectiveFactor)
+	{
+		int32 CrossingIndex = LastCrossingIndex;
+		double TravelTime = 0;
+		double TravelLength = OutCuttingPointUs.Last();
+		double Step = 1;
+		for (double StepTime : TravelTimeByStep)
 		{
-			int32 CrossingIndex = LastCrossingIndex;
-			double TravelTime = 0;
-			double TravelLength = OutCuttingPointUs.Last();
-			double Step = 1;
-			for (double StepTime : TravelTimeByStep)
+			double CorrectedStepTime = StepTime / CorrectiveFactor;
+			double CorrectedDeltaU = DeltaUs[CrossingIndex] * CorrectiveFactor;
+			while (TravelTime + CorrectedStepTime + DOUBLE_SMALL_NUMBER > Step)
 			{
-				double CorrectedStepTime = StepTime / CorrectiveFactor;
-				double CorrectedDeltaU = DeltaUs[CrossingIndex] * CorrectiveFactor;
-				while (TravelTime + CorrectedStepTime + DOUBLE_SMALL_NUMBER > Step)
+				double Time = Step - TravelTime;
+				OutCuttingPointUs.Add(TravelLength + CorrectedDeltaU * Time);
+				Step++;
+				if (Step + DOUBLE_SMALL_NUMBER > NbStep)
 				{
-					double Time = Step - TravelTime;
-					OutCuttingPointUs.Add(TravelLength + CorrectedDeltaU * Time);
-					Step++;
-					if (Step + DOUBLE_SMALL_NUMBER > NbStep)
-					{
-						OutCuttingPointUs.Add(EndU);
-						return;
-					}
+					OutCuttingPointUs.Add(EndU);
+					return;
 				}
-				TravelLength += CorrectedDeltaU * CorrectedStepTime;
-				TravelTime += CorrectedStepTime;
-				CrossingIndex++;
 			}
-		};
+			TravelLength += CorrectedDeltaU * CorrectedStepTime;
+			TravelTime += CorrectedStepTime;
+			CrossingIndex++;
+		}
+	};
 
-		OutCuttingPointUs.Add(ImposedCuttingPoints[0].Coordinate);
+	OutCuttingPointUs.Add(ImposedCuttingPoints[0].Coordinate);
 
-		while (CrossingUs[LastCrossingIndex + 1] <= ImposedCuttingPoints[0].Coordinate)
-		{
-			++LastCrossingIndex;
-		};
+	while (CrossingUs[LastCrossingIndex + 1] <= ImposedCuttingPoints[0].Coordinate)
+	{
+		++LastCrossingIndex;
+	};
 
-		// Between each imposed cutting points
-		double NextStep = ImposedCuttingPoints[0].Coordinate;
-		for (int32 Index = 1; Index < ImposedCuttingPoints.Num(); ++Index)
-		{
-			double StartStep = NextStep;
-			double StepCount = 0;
-			double CorrectiveFactor = 0;
-			NextStep = ImposedCuttingPoints[Index].Coordinate;
+	// Between each imposed cutting points
+	double NextStep = ImposedCuttingPoints[0].Coordinate;
+	for (int32 Index = 1; Index < ImposedCuttingPoints.Num(); ++Index)
+	{
+		double StartStep = NextStep;
+		double StepCount = 0;
+		double CorrectiveFactor = 0;
+		NextStep = ImposedCuttingPoints[Index].Coordinate;
 
-			// Step 1: Define the number of step (each step lasts one) as it must be an integer bigger than 1
-			// Step 2: Adjust the speeds than the travel respect the count of step
-			LapTimesUntilNextPoint(StartStep, NextStep, StepCount, CorrectiveFactor);
+		// Step 1: Define the number of step (each step lasts one) as it must be an integer bigger than 1
+		// Step 2: Adjust the speeds than the travel respect the count of step
+		LapTimesUntilNextPoint(StartStep, NextStep, StepCount, CorrectiveFactor);
 
-			if (StepCount == 1)
-			{
-				OutCuttingPointUs.Add(NextStep);
-			}
-			else
-			{
-				ComputeCuttingPointUntilNextImposedPoint(NextStep, StepCount, CorrectiveFactor);
-			}
-			LastCrossingIndex = CrossingIndex;
+		if (StepCount == 1)
+		{
+			OutCuttingPointUs.Add(NextStep);
 		}
-	};
+		else
+		{
+			ComputeCuttingPointUntilNextImposedPoint(NextStep, StepCount, CorrectiveFactor);
+		}
+		LastCrossingIndex = CrossingIndex;
+	}
+};
 }
\ No newline at end of file
