Rename CADKernel namespace into UE::CADKernel

diff --git a/Engine/Source/Runtime/Datasmith/CADKernel/Private/CADKernel/Mesh/Meshers/MesherTools.h b/Engine/Source/Runtime/Datasmith/CADKernel/Private/CADKernel/Mesh/Meshers/MesherTools.h
--- a/Engine/Source/Runtime/Datasmith/CADKernel/Private/CADKernel/Mesh/Meshers/MesherTools.h
+++ b/Engine/Source/Runtime/Datasmith/CADKernel/Private/CADKernel/Mesh/Meshers/MesherTools.h
@@ -4,65 +4,65 @@
 #include "CADKernel/Core/Types.h"
 #include "CADKernel/Math/Boundary.h"
 
-namespace CADKernel
+namespace UE::CADKernel
 {
-	struct FCuttingPoint;
-	struct FLinearBoundary;
+struct FCuttingPoint;
+struct FLinearBoundary;
 
-	class FMesherTools
-	{
-	public:
+class FMesherTools
+{
+public:
 
-		/**
-		 * The purpose of this method is to compute the cutting of a face or an edge by preferring as soon as possible the suggested cuttings.
-		 * In input a set of coordinate (CrossingUs) and the DeltaU between each coordinate (DeltaUs) which makes it possible to respect the mesh criteria.
-		 *
-		 * By analogy with a travel:
-		 * - the target time between each cutting points is one
-		 * - the speed of the travel at a point p(u) is DeltaUs(u)
-		 *
-		 * Step 1: Compute the number of mesh nodes i.e. compute the time of the travel between Boundary.Min and Boundary.Max with the initial speeds DeltaUs
-		 * Step 2: Define the number of step (each step lasts one) as it must be an integer bigger than 1
-		 * Step 3: Adjust the speeds than the travel respect the count of step
-		 * Step 4: Define the steps (if a preferred step is near the next step, the preferred step is used)
-		 *
-		 * Mandatory: PreferredCuttingPoints.Last() == UMax,
-		 */
-		static void ComputeFinalCuttingPointsWithPreferredCuttingPoints(const TArray<double>& CrossingUs, TArray<double> DeltaUs, const TArray<FCuttingPoint>& PreferredCuttingPoints, const FLinearBoundary& Boundary, TArray<double>& OutCuttingPoints);
+	/**
+	 * The purpose of this method is to compute the cutting of a face or an edge by preferring as soon as possible the suggested cuttings.
+	 * In input a set of coordinate (CrossingUs) and the DeltaU between each coordinate (DeltaUs) which makes it possible to respect the mesh criteria.
+	 *
+	 * By analogy with a travel:
+	 * - the target time between each cutting points is one
+	 * - the speed of the travel at a point p(u) is DeltaUs(u)
+	 *
+	 * Step 1: Compute the number of mesh nodes i.e. compute the time of the travel between Boundary.Min and Boundary.Max with the initial speeds DeltaUs
+	 * Step 2: Define the number of step (each step lasts one) as it must be an integer bigger than 1
+	 * Step 3: Adjust the speeds than the travel respect the count of step
+	 * Step 4: Define the steps (if a preferred step is near the next step, the preferred step is used)
+	 *
+	 * Mandatory: PreferredCuttingPoints.Last() == UMax,
+	 */
+	static void ComputeFinalCuttingPointsWithPreferredCuttingPoints(const TArray<double>& CrossingUs, TArray<double> DeltaUs, const TArray<FCuttingPoint>& PreferredCuttingPoints, const FLinearBoundary& Boundary, TArray<double>& OutCuttingPoints);
 
-		/**
-		 * The purpose of this method is to compute the cutting of a face or an edge by using the imposed cuttings.
-		 * In input a set of coordinate (CrossingUs) and the DeltaU between each coordinate (DeltaUs) which makes it possible to respect the mesh criteria.
-		 *
-		 * By analogy with a travel:
-		 * - the target time between each cutting points is one
-		 * - the speed of the travel at a point p(u) is DeltaUs(u)
-		 *
-		 * Between each imposed cutting points 
-		 * Step 1: Define the number of step (each step lasts one) as it must be an integer bigger than 1
-		 * Step 2: Adjust the speeds than the travel respect the count of step
-		 * Step 3: Compute the steps until the next imposed point
-		 *
-		 * Mandatory: ImposedCuttingPoints[0] == UMin && ImposedCuttingPoints.Last() == UMax,
-		 */
-		static void ComputeFinalCuttingPointsWithImposedCuttingPoints(const TArray<double>& CrossingUs, const TArray<double>& DeltaUs, const TArray<FCuttingPoint>& ImposedCuttingPoints, TArray<double>& OutCuttingPoints);
+	/**
+	 * The purpose of this method is to compute the cutting of a face or an edge by using the imposed cuttings.
+	 * In input a set of coordinate (CrossingUs) and the DeltaU between each coordinate (DeltaUs) which makes it possible to respect the mesh criteria.
+	 *
+	 * By analogy with a travel:
+	 * - the target time between each cutting points is one
+	 * - the speed of the travel at a point p(u) is DeltaUs(u)
+	 *
+	 * Between each imposed cutting points
+	 * Step 1: Define the number of step (each step lasts one) as it must be an integer bigger than 1
+	 * Step 2: Adjust the speeds than the travel respect the count of step
+	 * Step 3: Compute the steps until the next imposed point
+	 *
+	 * Mandatory: ImposedCuttingPoints[0] == UMin && ImposedCuttingPoints.Last() == UMax,
+	 */
+	static void ComputeFinalCuttingPointsWithImposedCuttingPoints(const TArray<double>& CrossingUs, const TArray<double>& DeltaUs, const TArray<FCuttingPoint>& ImposedCuttingPoints, TArray<double>& OutCuttingPoints);
 
-		/**
-		 * Fill an array of CuttingPointCoordinates with "PointCount" coordinates evenly divided between the extremities of the boundary 
-		 */
-		static void FillCuttingPointCoordinates(const FLinearBoundary& Boundary, int32 PointCount, TArray<double>& OutCuttingPointCoordinates)
+	/**
+	 * Fill an array of CuttingPointCoordinates with "PointCount" coordinates evenly divided between the extremities of the boundary
+	 */
+	static void FillCuttingPointCoordinates(const FLinearBoundary& Boundary, int32 PointCount, TArray<double>& OutCuttingPointCoordinates)
+	{
+		PointCount--;
+		ensureCADKernel(PointCount > 0);
+		double Delta = Boundary.Length() / PointCount;
+		OutCuttingPointCoordinates.Empty(PointCount);
+		OutCuttingPointCoordinates.Add(Boundary.GetMin());
+		for (int32 Index = 0; Index < PointCount; ++Index)
 		{
-			PointCount--;
-			ensureCADKernel(PointCount > 0);
-			double Delta = Boundary.Length() / PointCount;
-			OutCuttingPointCoordinates.Empty(PointCount);
-			OutCuttingPointCoordinates.Add(Boundary.GetMin());
-			for (int32 Index = 0; Index < PointCount; ++Index)
-			{
-				OutCuttingPointCoordinates.Add(OutCuttingPointCoordinates.Last() + Delta);
-			}
+			OutCuttingPointCoordinates.Add(OutCuttingPointCoordinates.Last() + Delta);
 		}
-	};
+	}
+};
 
-} // namespace CADKernel
+} // namespace UE::CADKernel
 
diff --git a/Engine/Source/Runtime/Datasmith/CADKernel/Private/CADKernel/Mesh/Meshers/ParametricMesher.cpp b/Engine/Source/Runtime/Datasmith/CADKernel/Private/CADKernel/Mesh/Meshers/ParametricMesher.cpp
--- a/Engine/Source/Runtime/Datasmith/CADKernel/Private/CADKernel/Mesh/Meshers/ParametricMesher.cpp
+++ b/Engine/Source/Runtime/Datasmith/CADKernel/Private/CADKernel/Mesh/Meshers/ParametricMesher.cpp
@@ -22,7 +22,7 @@
 
 //#define DEBUG_MESH_EDGE
 //#define DEBUG_GETPREFERREDUVCOORDINATESFROMNEIGHBOURS
-namespace CADKernel
+namespace UE::CADKernel
 {
 
 FParametricMesher::FParametricMesher(FModelMesh& InMeshModel)
@@ -467,7 +467,7 @@ void FParametricMesher::Mesh(FTopologicalEdge& InEdge, const FTopologicalFace& F
 	DebugIntersectEdgeIsos(Face, SurfaceTabV, EIso::IsoV);
 	{
 		F3DDebugSession _(FString::Printf(TEXT("Edge 2D %d"), InEdge.GetId()));
-		CADKernel::Display2D(InEdge);
+		UE::CADKernel::Display2D(InEdge);
 	}
 #endif
 
@@ -702,21 +702,21 @@ void FParametricMesher::Mesh(FTopologicalEdge& InEdge, const FTopologicalFace& F
 				F3DDebugSession G(TEXT("U From Iso"));
 				for (const FCuttingPoint& CuttingU : ImposedIsoCuttingPoints)
 				{
-					CADKernel::DisplayPoint(FPoint2D(CuttingU.Coordinate, 0.0));
+					UE::CADKernel::DisplayPoint(FPoint2D(CuttingU.Coordinate, 0.0));
 				}
 			}
 			{
 				F3DDebugSession G(TEXT("U From Criteria"));
 				for (double CuttingU : CuttingPoints2)
 				{
-					CADKernel::DisplayPoint(FPoint2D(CuttingU, 0.05), EVisuProperty::NonManifoldEdge);
+					UE::CADKernel::DisplayPoint(FPoint2D(CuttingU, 0.05), EVisuProperty::NonManifoldEdge);
 				}
 			}
 			{
 				F3DDebugSession G(TEXT("U Final (Criteria & Iso)"));
 				for (double CuttingU : CuttingPoints)
 				{
-					CADKernel::DisplayPoint(FPoint2D(CuttingU, 0.1), EVisuProperty::PurplePoint);
+					UE::CADKernel::DisplayPoint(FPoint2D(CuttingU, 0.1), EVisuProperty::PurplePoint);
 				}
 			}
 			//Wait();
diff --git a/Engine/Source/Runtime/Datasmith/CADKernel/Private/CADKernel/Mesh/Structure/EdgeSegment.h b/Engine/Source/Runtime/Datasmith/CADKernel/Private/CADKernel/Mesh/Structure/EdgeSegment.h
--- a/Engine/Source/Runtime/Datasmith/CADKernel/Private/CADKernel/Mesh/Structure/EdgeSegment.h
+++ b/Engine/Source/Runtime/Datasmith/CADKernel/Private/CADKernel/Mesh/Structure/EdgeSegment.h
@@ -9,252 +9,252 @@
 #include "CADKernel/Math/Point.h"
 #include "CADKernel/Math/SlopeUtils.h"
 
-namespace CADKernel
+namespace UE::CADKernel
 {
-	class FTopologicalLoop;
-	class FTopologicalEdge;
-	class FGrid;
-	class FThinZone;
-	class FThinZone2DFinder;
-	class FTopologicalVertex;
-
-	class FEdgeSegment : public FHaveStates
-	{
-	private:
-		TSharedPtr<FTopologicalEdge> Edge;
-		double Coordinates[2];
-		FPoint2D Points[2];
+class FTopologicalLoop;
+class FTopologicalEdge;
+class FGrid;
+class FThinZone;
+class FThinZone2DFinder;
+class FTopologicalVertex;
+
+class FEdgeSegment : public FHaveStates
+{
+private:
+	TSharedPtr<FTopologicalEdge> Edge;
+	double Coordinates[2];
+	FPoint2D Points[2];
 
-		FEdgeSegment* NextSegment;
-		FEdgeSegment* PreviousSegment;
+	FEdgeSegment* NextSegment;
+	FEdgeSegment* PreviousSegment;
 
-		FEdgeSegment* ClosedSegment[2] = { nullptr, nullptr };
+	FEdgeSegment* ClosedSegment[2] = { nullptr, nullptr };
 
-		FSurfacicBoundary Boundary;
-		double AxisMin;
+	FSurfacicBoundary Boundary;
+	double AxisMin;
 
-		double SquareDistanceToClosedSegment;
-		double Length;
+	double SquareDistanceToClosedSegment;
+	double Length;
 
-		FIdent ChainIndex;
+	FIdent ChainIndex;
 
-		FIdent Id;
-		static FIdent LastId;
+	FIdent Id;
+	static FIdent LastId;
 
-	public:
-		FEdgeSegment()
-			: NextSegment(nullptr)
-			, PreviousSegment(nullptr)
-			, AxisMin(0.)
-			, SquareDistanceToClosedSegment(HUGE_VALUE)
-			, Length(-1.)
-			, ChainIndex(0)
-			, Id(0)
-		{
-		};
+public:
+	FEdgeSegment()
+		: NextSegment(nullptr)
+		, PreviousSegment(nullptr)
+		, AxisMin(0.)
+		, SquareDistanceToClosedSegment(HUGE_VALUE)
+		, Length(-1.)
+		, ChainIndex(0)
+		, Id(0)
+	{
+	};
 
-		virtual ~FEdgeSegment() = default;
-		
-		void SetBoundarySegment(bool bInIsInnerLoop, const TSharedPtr<FTopologicalEdge>& InEdge, double InStartU, double InEndU, const FPoint2D& InStartPoint, const FPoint2D& InEndPoint)
+	virtual ~FEdgeSegment() = default;
+
+	void SetBoundarySegment(bool bInIsInnerLoop, const TSharedPtr<FTopologicalEdge>& InEdge, double InStartU, double InEndU, const FPoint2D& InStartPoint, const FPoint2D& InEndPoint)
+	{
+		if (bInIsInnerLoop)
 		{
-			if (bInIsInnerLoop)
-			{
-				SetInner();
-			}
+			SetInner();
+		}
 
-			Edge = InEdge;
-			Coordinates[ELimit::Start] = InStartU;
-			Coordinates[ELimit::End] = InEndU;
-			Points[ELimit::Start] = InStartPoint;
-			Points[ELimit::End]   = InEndPoint;
-			NextSegment = nullptr;
-			PreviousSegment = nullptr;
-			ClosedSegment[0] = nullptr;
-			ClosedSegment[1] = nullptr;
+		Edge = InEdge;
+		Coordinates[ELimit::Start] = InStartU;
+		Coordinates[ELimit::End] = InEndU;
+		Points[ELimit::Start] = InStartPoint;
+		Points[ELimit::End] = InEndPoint;
+		NextSegment = nullptr;
+		PreviousSegment = nullptr;
+		ClosedSegment[0] = nullptr;
+		ClosedSegment[1] = nullptr;
 
-			SquareDistanceToClosedSegment = HUGE_VAL;
-			Length = Points[ELimit::Start].Distance(Points[ELimit::End]);
+		SquareDistanceToClosedSegment = HUGE_VAL;
+		Length = Points[ELimit::Start].Distance(Points[ELimit::End]);
 
-			Id = ++LastId;
-			ChainIndex = 0;
+		Id = ++LastId;
+		ChainIndex = 0;
 
-			Boundary.Set(Points[ELimit::Start], Points[ELimit::End]);
+		Boundary.Set(Points[ELimit::Start], Points[ELimit::End]);
 
-			AxisMin = Boundary[EIso::IsoU].Min + Boundary[EIso::IsoV].Min;
-		};
+		AxisMin = Boundary[EIso::IsoU].Min + Boundary[EIso::IsoV].Min;
+	};
 
-		double GetAxeMin() const
-		{
-			return AxisMin;
-		}
+	double GetAxeMin() const
+	{
+		return AxisMin;
+	}
 
-		FIdent GetChainIndex() const
-		{
-			return ChainIndex;
-		}
+	FIdent GetChainIndex() const
+	{
+		return ChainIndex;
+	}
 
-		void SetChainIndex(FIdent index)
-		{
-			ChainIndex = index;
-		}
+	void SetChainIndex(FIdent index)
+	{
+		ChainIndex = index;
+	}
 
-		bool IsInner() const
-		{
-			return ((States & EHaveStates::IsInner) == EHaveStates::IsInner);
-		}
+	bool IsInner() const
+	{
+		return ((States & EHaveStates::IsInner) == EHaveStates::IsInner);
+	}
 
-		void SetInner()
-		{
-			States |= EHaveStates::IsInner;
-		}
+	void SetInner()
+	{
+		States |= EHaveStates::IsInner;
+	}
 
-		FIdent GetId()
-		{
-			return Id;
-		}
+	FIdent GetId()
+	{
+		return Id;
+	}
 
-		const TSharedPtr<FTopologicalEdge> GetEdge() const
-		{
-			return Edge;
-		}
+	const TSharedPtr<FTopologicalEdge> GetEdge() const
+	{
+		return Edge;
+	}
 
-		//const TSharedPtr<FTopologicalLoop>& GetLoop() const
-		//{
-		//	return Edge->GetLoop();
-		//}
+	//const TSharedPtr<FTopologicalLoop>& GetLoop() const
+	//{
+	//	return Edge->GetLoop();
+	//}
 
-		double GetLength() const
-		{
-			return Length;
-		}
+	double GetLength() const
+	{
+		return Length;
+	}
 
-		FPoint GetCenter() const
-		{
-			return (Points[ELimit::Start] + Points[ELimit::End]) * 0.5;
-		}
+	FPoint GetCenter() const
+	{
+		return (Points[ELimit::Start] + Points[ELimit::End]) * 0.5;
+	}
 
-		FPoint ComputeEdgePoint(double EdgeParamU) const
-		{
-			double SegmentParamS = (EdgeParamU - Coordinates[ELimit::Start]) / (Coordinates[ELimit::End] - Coordinates[ELimit::Start]);
-			return Points[ELimit::Start] + (Points[ELimit::End] - Points[ELimit::Start]) * SegmentParamS;
-		}
+	FPoint ComputeEdgePoint(double EdgeParamU) const
+	{
+		double SegmentParamS = (EdgeParamU - Coordinates[ELimit::Start]) / (Coordinates[ELimit::End] - Coordinates[ELimit::Start]);
+		return Points[ELimit::Start] + (Points[ELimit::End] - Points[ELimit::Start]) * SegmentParamS;
+	}
 
-		constexpr const FPoint2D& GetExtemity(const ELimit Limit) const
+	constexpr const FPoint2D& GetExtemity(const ELimit Limit) const
+	{
+		if (Limit == Start)
 		{
-			if (Limit == Start)
-			{
-				return Points[ELimit::Start];
-			}
-			else
-			{
-				return Points[ELimit::End];
-			}
+			return Points[ELimit::Start];
 		}
-
-		constexpr double GetCoordinate(const ELimit Limit) const
+		else
 		{
-			if(Limit == Start)
-			{
-				return Coordinates[ELimit::Start];
-			}
-			else
-			{
-				return Coordinates[ELimit::End];
-			}
+			return Points[ELimit::End];
 		}
+	}
 
-		/**
-		 * Compute the slope of the input Segment according to this.
-		 */
-		double ComputeUnorientedSlopeOf(FEdgeSegment* Segment)
+	constexpr double GetCoordinate(const ELimit Limit) const
+	{
+		if (Limit == Start)
 		{
-			double ReferenceSlope = ComputeSlope(Points[ELimit::Start], Points[ELimit::End]);
-			return ComputeUnorientedSlope(Segment->Points[ELimit::Start], Segment->Points[ELimit::End], ReferenceSlope);
+			return Coordinates[ELimit::Start];
 		}
-
-		/**
-		 * Compute the slope of the Segment defined by the two input points according to this.
-		 */
-		double ComputeUnorientedSlopeOf(FPoint2D& Middle, FPoint2D& Projection)
+		else
 		{
-			double ReferenceSlope = ComputeSlope(Points[ELimit::Start], Points[ELimit::End]);
-			return ComputeUnorientedSlope(Projection, Middle, ReferenceSlope);
+			return Coordinates[ELimit::End];
 		}
+	}
 
-		FEdgeSegment* GetNext() const
-		{
-			return NextSegment;
-		}
+	/**
+	 * Compute the slope of the input Segment according to this.
+	 */
+	double ComputeUnorientedSlopeOf(FEdgeSegment* Segment)
+	{
+		double ReferenceSlope = ComputeSlope(Points[ELimit::Start], Points[ELimit::End]);
+		return ComputeUnorientedSlope(Segment->Points[ELimit::Start], Segment->Points[ELimit::End], ReferenceSlope);
+	}
+
+	/**
+	 * Compute the slope of the Segment defined by the two input points according to this.
+	 */
+	double ComputeUnorientedSlopeOf(FPoint2D& Middle, FPoint2D& Projection)
+	{
+		double ReferenceSlope = ComputeSlope(Points[ELimit::Start], Points[ELimit::End]);
+		return ComputeUnorientedSlope(Projection, Middle, ReferenceSlope);
+	}
 
-		FEdgeSegment* GetPrevious() const
-		{
-			return PreviousSegment;
-		}
+	FEdgeSegment* GetNext() const
+	{
+		return NextSegment;
+	}
 
-		FEdgeSegment* GetClosedSegment() const
-		{
-			return ClosedSegment[0];
-		}
+	FEdgeSegment* GetPrevious() const
+	{
+		return PreviousSegment;
+	}
 
-		void ResetClosedData()
-		{
-			if (ClosedSegment[0]->ClosedSegment[0] == this)
-			{
-				ClosedSegment[0]->ClosedSegment[0] = ClosedSegment[0]->ClosedSegment[1];
-			}
-			else
-			{
-				ClosedSegment[0]->ClosedSegment[1] = nullptr;
-			}
-
-			ClosedSegment[0] = nullptr;
-			SquareDistanceToClosedSegment = HUGE_VAL;
-		}
+	FEdgeSegment* GetClosedSegment() const
+	{
+		return ClosedSegment[0];
+	}
 
-		void SetClosedSegment(FEdgeSegment* InSegmentA, FEdgeSegment* InSegmentB, double InDistance, bool bFirstChoice)
+	void ResetClosedData()
+	{
+		if (ClosedSegment[0]->ClosedSegment[0] == this)
 		{
-			ClosedSegment[0] = InSegmentA;
-			SquareDistanceToClosedSegment = InDistance;
-
-			if (bFirstChoice && InDistance < InSegmentA->SquareDistanceToClosedSegment)
-			{
-				InSegmentA->ClosedSegment[0] = this;
-				InSegmentA->SquareDistanceToClosedSegment = InDistance;
-			}
-			if (InSegmentB && InDistance < InSegmentB->SquareDistanceToClosedSegment)
-			{
-				InSegmentB->ClosedSegment[0] = this;
-				InSegmentB->SquareDistanceToClosedSegment = InDistance;
-			}
+			ClosedSegment[0]->ClosedSegment[0] = ClosedSegment[0]->ClosedSegment[1];
 		}
-
-		double GetClosedSquareDistance() const
+		else
 		{
-			return SquareDistanceToClosedSegment;
+			ClosedSegment[0]->ClosedSegment[1] = nullptr;
 		}
 
-		void SetNext(FEdgeSegment* Segment)
-		{
-			NextSegment = Segment;
-			Segment->SetPrevious(this);
-		}
+		ClosedSegment[0] = nullptr;
+		SquareDistanceToClosedSegment = HUGE_VAL;
+	}
+
+	void SetClosedSegment(FEdgeSegment* InSegmentA, FEdgeSegment* InSegmentB, double InDistance, bool bFirstChoice)
+	{
+		ClosedSegment[0] = InSegmentA;
+		SquareDistanceToClosedSegment = InDistance;
 
-		double ComputeEdgeCoordinate(const double SegmentU) const
+		if (bFirstChoice && InDistance < InSegmentA->SquareDistanceToClosedSegment)
 		{
-			return Coordinates[ELimit::Start] + (Coordinates[ELimit::End] - Coordinates[ELimit::Start]) * SegmentU;
+			InSegmentA->ClosedSegment[0] = this;
+			InSegmentA->SquareDistanceToClosedSegment = InDistance;
 		}
-
-		FPoint2D ProjectPoint(const FPoint2D& PointToProject, double& SegmentU) const
+		if (InSegmentB && InDistance < InSegmentB->SquareDistanceToClosedSegment)
 		{
-			return ProjectPointOnSegment(PointToProject, Points[ELimit::Start], Points[ELimit::End], SegmentU, true);
+			InSegmentB->ClosedSegment[0] = this;
+			InSegmentB->SquareDistanceToClosedSegment = InDistance;
 		}
+	}
 
-	private:
-		void SetPrevious(FEdgeSegment* Segment)
-		{
+	double GetClosedSquareDistance() const
+	{
+		return SquareDistanceToClosedSegment;
+	}
 
-			PreviousSegment = Segment;
-		}
-	};
+	void SetNext(FEdgeSegment* Segment)
+	{
+		NextSegment = Segment;
+		Segment->SetPrevious(this);
+	}
+
+	double ComputeEdgeCoordinate(const double SegmentU) const
+	{
+		return Coordinates[ELimit::Start] + (Coordinates[ELimit::End] - Coordinates[ELimit::Start]) * SegmentU;
+	}
+
+	FPoint2D ProjectPoint(const FPoint2D& PointToProject, double& SegmentU) const
+	{
+		return ProjectPointOnSegment(PointToProject, Points[ELimit::Start], Points[ELimit::End], SegmentU, true);
+	}
+
+private:
+	void SetPrevious(FEdgeSegment* Segment)
+	{
+
+		PreviousSegment = Segment;
+	}
+};
 }
 
diff --git a/Engine/Source/Runtime/Datasmith/CADKernel/Private/CADKernel/Mesh/Structure/Grid.cpp b/Engine/Source/Runtime/Datasmith/CADKernel/Private/CADKernel/Mesh/Structure/Grid.cpp
--- a/Engine/Source/Runtime/Datasmith/CADKernel/Private/CADKernel/Mesh/Structure/Grid.cpp
+++ b/Engine/Source/Runtime/Datasmith/CADKernel/Private/CADKernel/Mesh/Structure/Grid.cpp
@@ -13,8 +13,9 @@
 
 //#define DEBUG_GRID
 //#define DEBUG_GETPREFERREDUVCOORDINATESFROMNEIGHBOURS
-namespace CADKernel
+namespace UE::CADKernel
 {
+
 FGrid::FGrid(FTopologicalFace& InFace, FModelMesh& InMeshModel)
 	: Face(InFace)
 	, FaceTolerance(InFace.GetIsoTolerances())
@@ -198,7 +199,7 @@ void FGrid::GetPreferredUVCuttingParametersFromLoops(FCuttingGrid& CuttingParame
 	{
 		{
 			F3DDebugSession _(TEXT("Surface 2D"));
-			CADKernel::Display2D(*Face.GetCarrierSurface());
+			UE::CADKernel::Display2D(*Face.GetCarrierSurface());
 		}
 	}
 #endif
diff --git a/Engine/Source/Runtime/Datasmith/CADKernel/Private/CADKernel/Mesh/Structure/ModelMeshAnalyzer.cpp b/Engine/Source/Runtime/Datasmith/CADKernel/Private/CADKernel/Mesh/Structure/ModelMeshAnalyzer.cpp
--- a/Engine/Source/Runtime/Datasmith/CADKernel/Private/CADKernel/Mesh/Structure/ModelMeshAnalyzer.cpp
+++ b/Engine/Source/Runtime/Datasmith/CADKernel/Private/CADKernel/Mesh/Structure/ModelMeshAnalyzer.cpp
@@ -13,7 +13,7 @@
 
 #include "Templates/UnrealTemplate.h"
 
-namespace CADKernel
+namespace UE::CADKernel
 {
 
 void FModelMeshAnalyzer::BuildMesh()
@@ -212,8 +212,8 @@ bool FModelMeshAnalyzer::CheckOrientation() const
 
 FPoint Analyzer::FTriangle::ComputeNormal(const TArray<FPoint>& NodeCoordinates) const
 {
-	return CADKernel::FTriangle(NodeCoordinates[Vertices[0]], NodeCoordinates[Vertices[1]], NodeCoordinates[Vertices[2]]).ComputeNormal();
+	return UE::CADKernel::FTriangle(NodeCoordinates[Vertices[0]], NodeCoordinates[Vertices[1]], NodeCoordinates[Vertices[2]]).ComputeNormal();
 }
 
-} // namespace CADKernel
+} // namespace UE::CADKernel
 
diff --git a/Engine/Source/Runtime/Datasmith/CADKernel/Private/CADKernel/Topo/Shell.cpp b/Engine/Source/Runtime/Datasmith/CADKernel/Private/CADKernel/Topo/Shell.cpp
--- a/Engine/Source/Runtime/Datasmith/CADKernel/Private/CADKernel/Topo/Shell.cpp
+++ b/Engine/Source/Runtime/Datasmith/CADKernel/Private/CADKernel/Topo/Shell.cpp
@@ -9,7 +9,7 @@
 #include "CADKernel/Topo/TopologyReport.h"
 #endif
 
-namespace CADKernel
+namespace UE::CADKernel
 {
 
 FShell::FShell(const TArray<TSharedPtr<FTopologicalFace>>& InTopologicalFaces, bool bIsInnerShell)
@@ -396,14 +396,14 @@ int32 FShell::Orient()
 			F3DDebugSession _(TEXT("BBox Face"));
 
 			FAABB AABB(BBox.Min, BBox.Max);
-			CADKernel::DisplayAABB(AABB);
+			UE::CADKernel::DisplayAABB(AABB);
 
 			for (int32 Index = 0; Index < 3; ++Index)
 			{
-				CADKernel::DisplayPoint(BBox.MaxPoints[Index], EVisuProperty::YellowPoint);
-				CADKernel::DisplayPoint(BBox.MinPoints[Index], EVisuProperty::YellowPoint);
-				CADKernel::DisplaySegment(BBox.MaxPoints[Index], BBox.MaxPoints[Index] + BBox.MaxPointNormals[Index], 0, EVisuProperty::YellowCurve);
-				CADKernel::DisplaySegment(BBox.MinPoints[Index], BBox.MinPoints[Index] + BBox.MinPointNormals[Index], 0, EVisuProperty::YellowCurve);
+				UE::CADKernel::DisplayPoint(BBox.MaxPoints[Index], EVisuProperty::YellowPoint);
+				UE::CADKernel::DisplayPoint(BBox.MinPoints[Index], EVisuProperty::YellowPoint);
+				UE::CADKernel::DisplaySegment(BBox.MaxPoints[Index], BBox.MaxPoints[Index] + BBox.MaxPointNormals[Index], 0, EVisuProperty::YellowCurve);
+				UE::CADKernel::DisplaySegment(BBox.MinPoints[Index], BBox.MinPoints[Index] + BBox.MinPointNormals[Index], 0, EVisuProperty::YellowCurve);
 			}
 			Wait();
 		}
diff --git a/Engine/Source/Runtime/Datasmith/CADKernel/Private/CADKernel/Topo/TopologicalFace.cpp b/Engine/Source/Runtime/Datasmith/CADKernel/Private/CADKernel/Topo/TopologicalFace.cpp
--- a/Engine/Source/Runtime/Datasmith/CADKernel/Private/CADKernel/Topo/TopologicalFace.cpp
+++ b/Engine/Source/Runtime/Datasmith/CADKernel/Private/CADKernel/Topo/TopologicalFace.cpp
@@ -19,7 +19,7 @@
 #include "CADKernel/Topo/TopologyReport.h"
 #endif
 
-namespace CADKernel
+namespace UE::CADKernel
 {
 
 void FTopologicalFace::ComputeBoundary() const
@@ -131,8 +131,8 @@ void FTopologicalFace::UpdateBBox(int32 IsoCount, const double ApproximationFact
 #ifdef DEBUG_GET_BBOX2
 			for (int32 Index = 0; Index < 3; ++Index)
 			{
-				CADKernel::DisplayPoint(IsoBBox.MaxPoints[Index], EVisuProperty::YellowPoint);
-				CADKernel::DisplayPoint(IsoBBox.MinPoints[Index], EVisuProperty::YellowPoint);
+				UE::CADKernel::DisplayPoint(IsoBBox.MaxPoints[Index], EVisuProperty::YellowPoint);
+				UE::CADKernel::DisplayPoint(IsoBBox.MinPoints[Index], EVisuProperty::YellowPoint);
 			}
 #endif
 
@@ -150,14 +150,14 @@ void FTopologicalFace::UpdateBBox(int32 IsoCount, const double ApproximationFact
 		F3DDebugSession _(TEXT("BBox Face"));
 
 		FAABB AABB(BBox.Min, BBox.Max);
-		CADKernel::DisplayAABB(AABB);
+		UE::CADKernel::DisplayAABB(AABB);
 
 		for (int32 Index = 0; Index < 3; ++Index)
 		{
-			CADKernel::DisplayPoint(BBox.MaxPoints[Index], EVisuProperty::YellowPoint);
-			CADKernel::DisplayPoint(BBox.MinPoints[Index], EVisuProperty::YellowPoint);
-			CADKernel::DisplaySegment(BBox.MaxPoints[Index], BBox.MaxPoints[Index] + BBox.MaxPointNormals[Index], 0, EVisuProperty::YellowCurve);
-			CADKernel::DisplaySegment(BBox.MinPoints[Index], BBox.MinPoints[Index] + BBox.MinPointNormals[Index], 0, EVisuProperty::YellowCurve);
+			UE::CADKernel::DisplayPoint(BBox.MaxPoints[Index], EVisuProperty::YellowPoint);
+			UE::CADKernel::DisplayPoint(BBox.MinPoints[Index], EVisuProperty::YellowPoint);
+			UE::CADKernel::DisplaySegment(BBox.MaxPoints[Index], BBox.MaxPoints[Index] + BBox.MaxPointNormals[Index], 0, EVisuProperty::YellowCurve);
+			UE::CADKernel::DisplaySegment(BBox.MinPoints[Index], BBox.MinPoints[Index] + BBox.MinPointNormals[Index], 0, EVisuProperty::YellowCurve);
 		}
 		Wait();
 	}
diff --git a/Engine/Source/Runtime/Datasmith/CADKernel/Private/CADKernel/UI/Visu.cpp b/Engine/Source/Runtime/Datasmith/CADKernel/Private/CADKernel/UI/Visu.cpp
--- a/Engine/Source/Runtime/Datasmith/CADKernel/Private/CADKernel/UI/Visu.cpp
+++ b/Engine/Source/Runtime/Datasmith/CADKernel/Private/CADKernel/UI/Visu.cpp
@@ -2,39 +2,39 @@
 
 #include "CADKernel/UI/Visu.h"
 
-namespace CADKernel
+namespace UE::CADKernel
 {
-	const TCHAR* VisuPropertyNames[] = 
-	{
-		TEXT("Bounding Box"),                   
-		TEXT("Points"),
-		TEXT("Elements"),
-		TEXT("Curves"),
-		TEXT("Iso"),
-		TEXT("Control Lines"),
-		TEXT("Edges"),
-		TEXT("Nodes"),
-		TEXT("Border Edges"),
-		TEXT("Multiple Edges"),
-		TEXT("Geometrical Borders"),
-		TEXT("Control Points"),
+const TCHAR* VisuPropertyNames[] =
+{
+	TEXT("Bounding Box"),
+	TEXT("Points"),
+	TEXT("Elements"),
+	TEXT("Curves"),
+	TEXT("Iso"),
+	TEXT("Control Lines"),
+	TEXT("Edges"),
+	TEXT("Nodes"),
+	TEXT("Border Edges"),
+	TEXT("Multiple Edges"),
+	TEXT("Geometrical Borders"),
+	TEXT("Control Points"),
 
-		TEXT("YellowPoint"),
-		TEXT("YellowCurve"),
-		TEXT("BluePoint"),
-		TEXT("BlueCurve"),
-		TEXT("RedPoint"),
-		TEXT("RedCurve"),
-		TEXT("PurplePoint"),
-		TEXT("PurpleCurve"),
-		TEXT("GreenPoint"),
-		TEXT("GreenCurve"),
-		TEXT("PinkPoint"),
-		TEXT("PinkCurve"),
-		TEXT("OrangePoint"),
-		TEXT("OrangeCurve"),
-		nullptr
-	};
+	TEXT("YellowPoint"),
+	TEXT("YellowCurve"),
+	TEXT("BluePoint"),
+	TEXT("BlueCurve"),
+	TEXT("RedPoint"),
+	TEXT("RedCurve"),
+	TEXT("PurplePoint"),
+	TEXT("PurpleCurve"),
+	TEXT("GreenPoint"),
+	TEXT("GreenCurve"),
+	TEXT("PinkPoint"),
+	TEXT("PinkCurve"),
+	TEXT("OrangePoint"),
+	TEXT("OrangeCurve"),
+	nullptr
+};
 
 
 }
diff --git a/Engine/Source/Runtime/Datasmith/CADKernel/Private/CADKernel/Utils/ArrayUtils.h b/Engine/Source/Runtime/Datasmith/CADKernel/Private/CADKernel/Utils/ArrayUtils.h
--- a/Engine/Source/Runtime/Datasmith/CADKernel/Private/CADKernel/Utils/ArrayUtils.h
+++ b/Engine/Source/Runtime/Datasmith/CADKernel/Private/CADKernel/Utils/ArrayUtils.h
@@ -6,227 +6,227 @@
 #include "CADKernel/Math/Boundary.h"
 #include "CADKernel/Utils/IndexOfCoordinateFinder.h"
 
-namespace CADKernel
+namespace UE::CADKernel
 {
-	namespace ArrayUtils
+namespace ArrayUtils
+{
+/**
+ * Merge two array of sorted values and keep the resulting array sorted.
+ * The value to insert is inserted if and only if it is not nearly equal to an existent value
+ * and if it is inside the boundary of the InitialArray ([InitialArray[0], InitialArray.Last()])
+ * The values of the existing array are always kept.
+ */
+inline void InsertInside(TArray<double>& InitialArray, const TArray<double>& ValuesToInsert, double Epsilon)
+{
+	ensureCADKernel(InitialArray.Num() >= 2);
+
+	if (ValuesToInsert.Num() == 0)
 	{
-		/**
-		 * Merge two array of sorted values and keep the resulting array sorted.
-		 * The value to insert is inserted if and only if it is not nearly equal to an existent value 
-		 * and if it is inside the boundary of the InitialArray ([InitialArray[0], InitialArray.Last()])
-		 * The values of the existing array are always kept.
-		 */
-		inline void InsertInside(TArray<double>& InitialArray, const TArray<double>& ValuesToInsert, double Epsilon)
+		return;
+	}
+
+	InitialArray.Reserve(InitialArray.Num() + ValuesToInsert.Num());
+
+	FDichotomyFinder Finder(InitialArray);
+
+	int32 Index = 0;
+	for (double Value : ValuesToInsert)
+	{
+		if (Value < InitialArray[0])
 		{
-			ensureCADKernel(InitialArray.Num() >= 2);
+			continue;
+		}
 
-			if (ValuesToInsert.Num() == 0)
-			{
-				return;
-			}
+		if (Value > InitialArray.Last())
+		{
+			break;
+		}
 
-			InitialArray.Reserve(InitialArray.Num() + ValuesToInsert.Num());
+		Index = Finder.Find(Value);
 
-			FDichotomyFinder Finder(InitialArray);
+		if (FMath::IsNearlyEqual(Value, InitialArray[Index], Epsilon))
+		{
+			continue;
+		}
 
-			int32 Index = 0;
-			for (double Value : ValuesToInsert)
-			{
-				if (Value < InitialArray[0])
-				{
-					continue;
-				}
+		if (FMath::IsNearlyEqual(Value, InitialArray[Index + 1], Epsilon))
+		{
+			continue;
+		}
 
-				if (Value > InitialArray.Last())
-				{
-					break;
-				}
+		InitialArray.EmplaceAt(Index + 1, Value);
+		Finder.StartUpper++;
+	}
+}
 
-				Index = Finder.Find(Value);
+/**
+ * Merge two array of sorted values and keep the resulting array sorted.
+ * The value to insert is inserted if and only if it is not nearly equal to an existent value except for the first and last value of ValuesToInsert if they are outside initial array boundary
+ */
+inline void Complete(TArray<double>& InitialArray, TArray<double>& ValuesToInsert, double Epsilon)
+{
+	if (ValuesToInsert.Num() == 0)
+	{
+		return;
+	}
+
+	if (InitialArray.Num() == 0)
+	{
+		Swap(InitialArray, ValuesToInsert);
+		return;
+	}
 
-				if (FMath::IsNearlyEqual(Value, InitialArray[Index], Epsilon))
-				{
-					continue;
-				}
+	InitialArray.Reserve(InitialArray.Num() + ValuesToInsert.Num());
 
-				if (FMath::IsNearlyEqual(Value, InitialArray[Index + 1], Epsilon))
-				{
-					continue;
-				}
+	FDichotomyFinder Finder(InitialArray);
 
-				InitialArray.EmplaceAt(Index + 1, Value);
+	int32 Index = 0;
+	for (double Value : ValuesToInsert)
+	{
+		if (Value < InitialArray[0])
+		{
+			if (FMath::IsNearlyEqual(Value, InitialArray[Index], Epsilon))
+			{
+				InitialArray[Index] = Value;
+			}
+			else
+			{
+				InitialArray.EmplaceAt(Index++, Value);
 				Finder.StartUpper++;
 			}
+			continue;
 		}
 
-		/**
-		 * Merge two array of sorted values and keep the resulting array sorted.
-		 * The value to insert is inserted if and only if it is not nearly equal to an existent value except for the first and last value of ValuesToInsert if they are outside initial array boundary
-		 */
-		inline void Complete(TArray<double>& InitialArray, TArray<double>& ValuesToInsert, double Epsilon)
+		if (Value > InitialArray.Last())
 		{
-			if (ValuesToInsert.Num() == 0)
+			if (FMath::IsNearlyEqual(Value, InitialArray.Last(), Epsilon))
 			{
-				return;
+				InitialArray.Last() = Value;
 			}
-
-			if (InitialArray.Num() == 0)
+			else
 			{
-				Swap(InitialArray, ValuesToInsert);
-				return;
+				InitialArray.Emplace(Value);;
 			}
+			continue;
+		}
 
-			InitialArray.Reserve(InitialArray.Num() + ValuesToInsert.Num());
+		Index = Finder.Find(Value);
 
-			FDichotomyFinder Finder(InitialArray);
+		if (FMath::IsNearlyEqual(Value, InitialArray[Index], Epsilon))
+		{
+			continue;
+		}
 
-			int32 Index = 0;
-			for (double Value : ValuesToInsert)
-			{
-				if (Value < InitialArray[0])
-				{
-					if (FMath::IsNearlyEqual(Value, InitialArray[Index], Epsilon))
-					{
-						InitialArray[Index] = Value;
-					}
-					else
-					{
-						InitialArray.EmplaceAt(Index++, Value);
-						Finder.StartUpper++;
-					}
-					continue;
-				}
-
-				if (Value > InitialArray.Last())
-				{
-					if (FMath::IsNearlyEqual(Value, InitialArray.Last(), Epsilon))
-					{
-						InitialArray.Last() = Value;
-					}
-					else
-					{
-						InitialArray.Emplace(Value);;
-					}
-					continue;
-				}
-
-				Index = Finder.Find(Value);
-
-				if (FMath::IsNearlyEqual(Value, InitialArray[Index], Epsilon))
-				{
-					continue;
-				}
-
-				if (FMath::IsNearlyEqual(Value, InitialArray[Index + 1], Epsilon))
-				{
-					continue;
-				}
-
-				InitialArray.EmplaceAt(Index + 1, Value);
-				Finder.StartUpper++;
-			}
+		if (FMath::IsNearlyEqual(Value, InitialArray[Index + 1], Epsilon))
+		{
+			continue;
 		}
 
-		inline void RemoveDuplicates(TArray<double>& ArrayToProcess, double Epsilon)
+		InitialArray.EmplaceAt(Index + 1, Value);
+		Finder.StartUpper++;
+	}
+}
+
+inline void RemoveDuplicates(TArray<double>& ArrayToProcess, double Epsilon)
+{
+	if (ArrayToProcess.Num() < 2)
+	{
+		return;
+	}
+
+	for (int32 Index = ArrayToProcess.Num() - 1; Index > 1; --Index)
+	{
+		if (FMath::IsNearlyEqual(ArrayToProcess[Index], ArrayToProcess[Index - 1], Epsilon))
 		{
-			if (ArrayToProcess.Num() < 2)
-			{
-				return;
-			}
+			ArrayToProcess.RemoveAt(Index - 1);
+		}
+	}
 
-			for (int32 Index = ArrayToProcess.Num() - 1; Index > 1; --Index)
-			{
-				if (FMath::IsNearlyEqual(ArrayToProcess[Index], ArrayToProcess[Index - 1], Epsilon))
-				{
-					ArrayToProcess.RemoveAt(Index - 1);
-				}
-			}
+	if (ArrayToProcess.Num() > 1)
+	{
+		if (FMath::IsNearlyEqual(ArrayToProcess[0], ArrayToProcess[1], Epsilon))
+		{
+			ArrayToProcess.RemoveAt(1);
+		}
+	}
+}
 
-			if (ArrayToProcess.Num() > 1)
-			{
-				if (FMath::IsNearlyEqual(ArrayToProcess[0], ArrayToProcess[1], Epsilon))
-				{
-					ArrayToProcess.RemoveAt(1);
-				}
-			}
+inline void SubArrayWithoutBoundary(const TArray<double>& InitialArray, const FLinearBoundary& Boundary, double Tolerance, TArray<double>& OutArray)
+{
+	int32 InitialSize = InitialArray.Num();
+	OutArray.Empty(InitialSize + 2);
+
+	if (InitialSize > 1)
+	{
+		FDichotomyFinder Finder(InitialArray);
+		int32 MinIndex = Finder.Find(Boundary.GetMin());
+		int32 MaxIndex = Finder.Find(Boundary.GetMax());
+
+		if (InitialArray[MinIndex] > Boundary.GetMax())
+		{
+			return;
 		}
 
-		inline void SubArrayWithoutBoundary(const TArray<double>& InitialArray, const FLinearBoundary& Boundary, double Tolerance, TArray<double>& OutArray)
+		MaxIndex++;
+		if (InitialArray[MaxIndex] < Boundary.GetMin())
 		{
-			int32 InitialSize = InitialArray.Num();
-			OutArray.Empty(InitialSize + 2);
+			return;
+		}
 
-			if (InitialSize > 1)
-			{
-				FDichotomyFinder Finder(InitialArray);
-				int32 MinIndex = Finder.Find(Boundary.GetMin());
-				int32 MaxIndex = Finder.Find(Boundary.GetMax());
-
-				if (InitialArray[MinIndex] > Boundary.GetMax())
-				{
-					return;
-				}
-
-				MaxIndex++;
-				if (InitialArray[MaxIndex] < Boundary.GetMin())
-				{
-					return;
-				}
-
-				if (Boundary.GetMin() + 5 * Tolerance > InitialArray[MinIndex])
-				{
-					MinIndex++;
-				}
-
-				if (Boundary.GetMax() - 5 * Tolerance > InitialArray[MaxIndex])
-				{
-					MaxIndex++;
-				}
-
-				if(MaxIndex - MinIndex)
-				{
-					OutArray.Append(InitialArray.GetData() + MinIndex, MaxIndex - MinIndex);
-				}
-			}
-			else if (InitialSize == 1)
-			{
-				double Value = InitialArray[0];
-				if ((Boundary.GetMin() + 5 * Tolerance > Value) || (Boundary.GetMax() - 5 * Tolerance < Value))
-				{
-					return;
-				}
-				OutArray.Add(Value);
-			}
+		if (Boundary.GetMin() + 5 * Tolerance > InitialArray[MinIndex])
+		{
+			MinIndex++;
 		}
 
-		/**
-		 * Find the index of the segment containing the coordinate i.e. Coordinate is in [InCoordinates[OutIndex], InCoordinates[OutIndex+1]]
-		 * This method is very fast if initial value of Index is near the solution.
-		 * This is very useful to process linked set of coordinate like loop point coordinates.
-		 */
-		inline void FindCoordinateIndex(const TArray<double>& InCoordinates, double Coordinate, int32& OutIndex)
+		if (Boundary.GetMax() - 5 * Tolerance > InitialArray[MaxIndex])
 		{
-			if (OutIndex >= InCoordinates.Num())
-			{
-				OutIndex = InCoordinates.Num() - 1;
-			}
+			MaxIndex++;
+		}
 
-			if (OutIndex < 0)
-			{
-				OutIndex = 0;
-			}
+		if (MaxIndex - MinIndex)
+		{
+			OutArray.Append(InitialArray.GetData() + MinIndex, MaxIndex - MinIndex);
+		}
+	}
+	else if (InitialSize == 1)
+	{
+		double Value = InitialArray[0];
+		if ((Boundary.GetMin() + 5 * Tolerance > Value) || (Boundary.GetMax() - 5 * Tolerance < Value))
+		{
+			return;
+		}
+		OutArray.Add(Value);
+	}
+}
 
-			while (OutIndex > 0 && Coordinate < InCoordinates[OutIndex])
-			{
-				OutIndex--;
-			}
+/**
+ * Find the index of the segment containing the coordinate i.e. Coordinate is in [InCoordinates[OutIndex], InCoordinates[OutIndex+1]]
+ * This method is very fast if initial value of Index is near the solution.
+ * This is very useful to process linked set of coordinate like loop point coordinates.
+ */
+inline void FindCoordinateIndex(const TArray<double>& InCoordinates, double Coordinate, int32& OutIndex)
+{
+	if (OutIndex >= InCoordinates.Num())
+	{
+		OutIndex = InCoordinates.Num() - 1;
+	}
 
-			for (; OutIndex + 2 < InCoordinates.Num() && Coordinate > InCoordinates[OutIndex + 1]; ++OutIndex)
-			{
-			}
+	if (OutIndex < 0)
+	{
+		OutIndex = 0;
+	}
+
+	while (OutIndex > 0 && Coordinate < InCoordinates[OutIndex])
+	{
+		OutIndex--;
+	}
+
+	for (; OutIndex + 2 < InCoordinates.Num() && Coordinate > InCoordinates[OutIndex + 1]; ++OutIndex)
+	{
+	}
 
-			ensureCADKernel(InCoordinates.IsValidIndex(OutIndex));
-		};
+	ensureCADKernel(InCoordinates.IsValidIndex(OutIndex));
+};
 
-	};
+};
 }
diff --git a/Engine/Source/Runtime/Datasmith/CADKernel/Private/CADKernel/Utils/Util.h b/Engine/Source/Runtime/Datasmith/CADKernel/Private/CADKernel/Utils/Util.h
--- a/Engine/Source/Runtime/Datasmith/CADKernel/Private/CADKernel/Utils/Util.h
+++ b/Engine/Source/Runtime/Datasmith/CADKernel/Private/CADKernel/Utils/Util.h
@@ -6,34 +6,34 @@
 #include "CADKernel/Math/Boundary.h"
 #include "CADKernel/Utils/IndexOfCoordinateFinder.h"
 
-namespace CADKernel
+namespace UE::CADKernel
 {
-	namespace Utils
+namespace Utils
+{
+template<typename SRC, typename DST>
+inline void TSharedPtrArrayConversion(const TArray<TSharedPtr<SRC>>& Source, TArray<TSharedPtr<DST>>& Destination)
+{
+	Destination.Reserve(Destination.Num() + Source.Num());
+	for (const TSharedPtr<SRC>& SourceItem : Source)
 	{
-		template<typename SRC, typename DST>
-		inline void TSharedPtrArrayConversion(const TArray<TSharedPtr<SRC>>& Source, TArray<TSharedPtr<DST>>& Destination)
-		{
-			Destination.Reserve(Destination.Num() + Source.Num());
-			for (const TSharedPtr<SRC>& SourceItem : Source)
-			{
-				Destination.Add(StaticCastSharedPtr<DST>(SourceItem));
-			}
-		}
+		Destination.Add(StaticCastSharedPtr<DST>(SourceItem));
+	}
+}
 
-		template<typename SRC, typename DST>
-		inline void TSharedPtrArrayConversionChecked(const TArray<TSharedPtr<SRC>>& Source, TArray<TSharedPtr<DST>>& Destination)
+template<typename SRC, typename DST>
+inline void TSharedPtrArrayConversionChecked(const TArray<TSharedPtr<SRC>>& Source, TArray<TSharedPtr<DST>>& Destination)
+{
+	Destination.Reserve(Destination.Num() + Source.Num());
+	for (const TSharedPtr<SRC>& SourceItem : Source)
+	{
+		TSharedPtr<DST> ConvertedPtr = StaticCastSharedPtr<DST>(SourceItem);
+		if (!ConvertedPtr.IsValid())
 		{
-			Destination.Reserve(Destination.Num() + Source.Num());
-			for (const TSharedPtr<SRC>& SourceItem : Source)
-			{
-				TSharedPtr<DST> ConvertedPtr = StaticCastSharedPtr<DST>(SourceItem);
-				if (!ConvertedPtr.IsValid())
-				{
-					continue;
-				}
-				Destination.Add(ConvertedPtr);
-			}
-
+			continue;
 		}
-	};
+		Destination.Add(ConvertedPtr);
+	}
+
+}
+};
 }
diff --git a/Engine/Source/Runtime/Datasmith/CADKernel/Public/CADKernel/Core/CADKernelArchive.h b/Engine/Source/Runtime/Datasmith/CADKernel/Public/CADKernel/Core/CADKernelArchive.h
--- a/Engine/Source/Runtime/Datasmith/CADKernel/Public/CADKernel/Core/CADKernelArchive.h
+++ b/Engine/Source/Runtime/Datasmith/CADKernel/Public/CADKernel/Core/CADKernelArchive.h
@@ -8,129 +8,129 @@
 #include "Serialization/Archive.h"
 #include "Serialization/MemoryReader.h"
 
-namespace CADKernel
+namespace UE::CADKernel
 {
-	class FEntity;
-	class FSession;
-	class FModel;
+class FEntity;
+class FSession;
+class FModel;
 
-	class FCADKernelArchive
+class FCADKernelArchive
+{
+public:
+	TUniquePtr<FMemoryReader> MemoryReader;
+	FArchive& Archive;
+	FSession& Session;
+	FModel* ArchiveModel = nullptr;
+
+	FCADKernelArchive(FSession& InSession, FArchive& InArchive)
+		: Archive(InArchive)
+		, Session(InSession)
 	{
-	public:
-		TUniquePtr<FMemoryReader> MemoryReader;
-		FArchive& Archive;
-		FSession& Session;
-		FModel* ArchiveModel = nullptr;
-
-		FCADKernelArchive(FSession& InSession, FArchive& InArchive)
-			: Archive(InArchive)
-			, Session(InSession)
-		{
-		}
-
-		/**
-		 * Archive reader constructor
-		 */
-		FCADKernelArchive(FSession& InSession, const TArray<uint8>& InRawData)
-			: MemoryReader(MakeUnique<FMemoryReader>(InRawData, true))
-			, Archive(*MemoryReader)
-			, Session(InSession)
-		{
-		}
+	}
+
+	/**
+	 * Archive reader constructor
+	 */
+	FCADKernelArchive(FSession& InSession, const TArray<uint8>& InRawData)
+		: MemoryReader(MakeUnique<FMemoryReader>(InRawData, true))
+		, Archive(*MemoryReader)
+		, Session(InSession)
+	{
+	}
 
-		static TSharedPtr<FCADKernelArchive> CreateArchiveWriter(FSession& InSession, const TCHAR* FileName)
+	static TSharedPtr<FCADKernelArchive> CreateArchiveWriter(FSession& InSession, const TCHAR* FileName)
+	{
+		FArchive* Ar = IFileManager::Get().CreateFileWriter(FileName);
+		if (Ar == nullptr)
 		{
-			FArchive* Ar = IFileManager::Get().CreateFileWriter(FileName);
-			if (Ar == nullptr)
-			{
-				FMessage::Printf(Log, TEXT("The archive file %s is corrupted\n"), FileName);
-				return TSharedPtr<FCADKernelArchive>();
-			}
-			return MakeShared<FCADKernelArchive>(InSession, *Ar);
+			FMessage::Printf(Log, TEXT("The archive file %s is corrupted\n"), FileName);
+			return TSharedPtr<FCADKernelArchive>();
 		}
+		return MakeShared<FCADKernelArchive>(InSession, *Ar);
+	}
 
-		static TSharedPtr<FCADKernelArchive> CreateArchiveReader(FSession& InSession, const TCHAR* FileName)
+	static TSharedPtr<FCADKernelArchive> CreateArchiveReader(FSession& InSession, const TCHAR* FileName)
+	{
+		FArchive* Ar = IFileManager::Get().CreateFileReader(FileName);
+		if (Ar == nullptr)
 		{
-			FArchive* Ar = IFileManager::Get().CreateFileReader(FileName);
-			if (Ar == nullptr)
-			{
-				FMessage::Printf(Log, TEXT("The archive file %s is corrupted\n"), FileName);
-				return TSharedPtr<FCADKernelArchive>();
-			}
-			return MakeShared<FCADKernelArchive>(InSession, *Ar);
+			FMessage::Printf(Log, TEXT("The archive file %s is corrupted\n"), FileName);
+			return TSharedPtr<FCADKernelArchive>();
 		}
+		return MakeShared<FCADKernelArchive>(InSession, *Ar);
+	}
 
-		template<typename EntityType>
-		void operator<<(EntityType& Entity)
-		{
-			Archive << Entity;
-		}
+	template<typename EntityType>
+	void operator<<(EntityType& Entity)
+	{
+		Archive << Entity;
+	}
 
-		bool IsLoading() const
-		{
-			return Archive.IsLoading();
-		}
+	bool IsLoading() const
+	{
+		return Archive.IsLoading();
+	}
 
-		bool IsSaving() const
-		{
-			return Archive.IsSaving();
-		}
+	bool IsSaving() const
+	{
+		return Archive.IsSaving();
+	}
 
-		void Serialize(void* Value, int64 Length)
-		{ 
-			Archive.Serialize(Value, Length);
-		}
+	void Serialize(void* Value, int64 Length)
+	{
+		Archive.Serialize(Value, Length);
+	}
 
-		template<typename EntityType>
-		void Serialize(TArray<EntityType>& EntityArray)
+	template<typename EntityType>
+	void Serialize(TArray<EntityType>& EntityArray)
+	{
+		if (Archive.IsLoading())
 		{
-			if(Archive.IsLoading())
-			{
-				int32 ElementCount = 0;
-				Archive << ElementCount;
-				EntityArray.SetNum(ElementCount);
-			}
-			else
-			{
-				int32 ElementCount = EntityArray.Num();
-				Archive << ElementCount;
-			}
-			Archive.Serialize((void*) EntityArray.GetData(), EntityArray.Num()*sizeof(EntityType));
+			int32 ElementCount = 0;
+			Archive << ElementCount;
+			EntityArray.SetNum(ElementCount);
 		}
-
-		template<typename EntityType>
-		void Serialize(EntityType& EntityArray)
+		else
 		{
-			Archive.Serialize((void*) &EntityArray, sizeof(EntityType));
+			int32 ElementCount = EntityArray.Num();
+			Archive << ElementCount;
 		}
+		Archive.Serialize((void*)EntityArray.GetData(), EntityArray.Num() * sizeof(EntityType));
+	}
 
-		void SetReferencedEntityOrAddToWaitingList(FIdent ArchiveId, FEntity** Entity);
-		void SetReferencedEntityOrAddToWaitingList(FIdent ArchiveId, TWeakPtr<FEntity>& Entity);
-		void SetReferencedEntityOrAddToWaitingList(FIdent ArchiveId, TSharedPtr<FEntity>& Entity);
+	template<typename EntityType>
+	void Serialize(EntityType& EntityArray)
+	{
+		Archive.Serialize((void*)&EntityArray, sizeof(EntityType));
+	}
 
-		void AddEntityToSave(FIdent Id);
-		void AddEntityFromArchive(TSharedPtr<FEntity>& Entity);
+	void SetReferencedEntityOrAddToWaitingList(FIdent ArchiveId, FEntity** Entity);
+	void SetReferencedEntityOrAddToWaitingList(FIdent ArchiveId, TWeakPtr<FEntity>& Entity);
+	void SetReferencedEntityOrAddToWaitingList(FIdent ArchiveId, TSharedPtr<FEntity>& Entity);
 
-		template<typename EntityType>
-		void AddEntityFromArchive(TSharedPtr<EntityType>& Entity)
-		{
-			AddEntityFromArchive((TSharedPtr<FEntity>&) Entity);
-		}
+	void AddEntityToSave(FIdent Id);
+	void AddEntityFromArchive(TSharedPtr<FEntity>& Entity);
 
-		int64 TotalSize()
-		{
-			return Archive.TotalSize();
-		}
+	template<typename EntityType>
+	void AddEntityFromArchive(TSharedPtr<EntityType>& Entity)
+	{
+		AddEntityFromArchive((TSharedPtr<FEntity>&) Entity);
+	}
 
-		int64 Tell()
-		{
-			return Archive.Tell();
-		}
+	int64 TotalSize()
+	{
+		return Archive.TotalSize();
+	}
 
-		void Close()
-		{
-			Archive.Close();
-		}
-	};
+	int64 Tell()
+	{
+		return Archive.Tell();
+	}
+
+	void Close()
+	{
+		Archive.Close();
+	}
+};
 }
 
diff --git a/Engine/Source/Runtime/Datasmith/CADKernel/Public/CADKernel/Core/Chrono.h b/Engine/Source/Runtime/Datasmith/CADKernel/Public/CADKernel/Core/Chrono.h
--- a/Engine/Source/Runtime/Datasmith/CADKernel/Public/CADKernel/Core/Chrono.h
+++ b/Engine/Source/Runtime/Datasmith/CADKernel/Public/CADKernel/Core/Chrono.h
@@ -9,65 +9,65 @@
 #include <utility>
 #endif
 
-namespace CADKernel
+namespace UE::CADKernel
 {
-	enum class ETimeUnit : uint8
-	{
-		NanoSeconds = 0,
-		MicroSeconds = 1,
-		MilliSeconds,
-		Seconds
-	};
+enum class ETimeUnit : uint8
+{
+	NanoSeconds = 0,
+	MicroSeconds = 1,
+	MilliSeconds,
+	Seconds
+};
 
 #ifdef CADKERNEL_DEV
-	typedef std::chrono::high_resolution_clock::time_point FTimePoint;
-	typedef std::chrono::high_resolution_clock::duration FDuration;
+typedef std::chrono::high_resolution_clock::time_point FTimePoint;
+typedef std::chrono::high_resolution_clock::duration FDuration;
 #else
-	typedef uint64 FTimePoint;
-	typedef uint64 FDuration;
+typedef uint64 FTimePoint;
+typedef uint64 FDuration;
 #endif
 
-	class CADKERNEL_API FChrono
+class CADKERNEL_API FChrono
+{
+	ETimeUnit DefaultUnit = ETimeUnit::NanoSeconds;
+public:
+	static const FTimePoint Now()
 	{
-		ETimeUnit DefaultUnit = ETimeUnit::NanoSeconds;
-	public:
-		static const FTimePoint Now()
-		{
 #ifdef CADKERNEL_DEV
-			return std::chrono::high_resolution_clock::now();
+		return std::chrono::high_resolution_clock::now();
 #else
-			return 0;
+		return 0;
 #endif
-		}
+	}
 
-		static const FDuration Elapse(FTimePoint StartTime)
-		{
+	static const FDuration Elapse(FTimePoint StartTime)
+	{
 #ifdef CADKERNEL_DEV
-			return (std::chrono::high_resolution_clock::now() - StartTime);
+		return (std::chrono::high_resolution_clock::now() - StartTime);
 #else
-			return 0;
+		return 0;
 #endif
-		}
+	}
 
-		static const FDuration Init()
-		{
+	static const FDuration Init()
+	{
 #ifdef CADKERNEL_DEV
-			return FDuration();
+		return FDuration();
 #else
-			return 0;
+		return 0;
 #endif
-		}
+	}
 
-		template<typename Unit>
-		static int64 ConvertInto(FDuration Duration)
-		{
+	template<typename Unit>
+	static int64 ConvertInto(FDuration Duration)
+	{
 #ifdef CADKERNEL_DEV
-			return std::chrono::duration_cast<Unit>(Duration).count();
+		return std::chrono::duration_cast<Unit>(Duration).count();
 #else
-			return 0;
+		return 0;
 #endif
-		}
+	}
 
-		static void PrintClockElapse(EVerboseLevel Level, const TCHAR* Indent, const TCHAR* Process, FDuration Duration, ETimeUnit Unit = ETimeUnit::MicroSeconds);
-	};
+	static void PrintClockElapse(EVerboseLevel Level, const TCHAR* Indent, const TCHAR* Process, FDuration Duration, ETimeUnit Unit = ETimeUnit::MicroSeconds);
+};
 }
\ No newline at end of file
diff --git a/Engine/Source/Runtime/Datasmith/CADKernel/Public/CADKernel/Core/Database.h b/Engine/Source/Runtime/Datasmith/CADKernel/Public/CADKernel/Core/Database.h
--- a/Engine/Source/Runtime/Datasmith/CADKernel/Public/CADKernel/Core/Database.h
+++ b/Engine/Source/Runtime/Datasmith/CADKernel/Public/CADKernel/Core/Database.h
@@ -6,308 +6,308 @@
 #include "CADKernel/Core/Types.h"
 #include "Containers/Queue.h"
 
-namespace CADKernel
+namespace UE::CADKernel
 {
-	class FCADKernelArchive;
-	class FModel;
-	class FTopologicalEntity;
-	class FTopologicalShapeEntity;
+class FCADKernelArchive;
+class FModel;
+class FTopologicalEntity;
+class FTopologicalShapeEntity;
 
-	class CADKERNEL_API FDatabase
+class CADKERNEL_API FDatabase
+{
+	friend class FCADKernelArchive;
+	friend class FEntity;
+	friend class FSession;
+
+protected:
+
+	/**
+	 * Root Geometry Entity node
+	 */
+	TSharedPtr<FModel> Model;
+
+	/**
+	 * Array (use as a Map<FIdent, TSharedPtr<FEntity>>) of all the entities create during the session
+	 * Each Entity.Id is the index of the entity in the array
+	 */
+	TArray<TSharedPtr<FEntity>> DatabaseEntities;
+
+	/**
+	 * Recycling of the Id to avoid fragmented idents
+	 */
+	TArray<FIdent> AvailableIdents;
+
+	/**
+	 * Array (use as a Map<FIdent, TSharedPtr<FEntity>>) of all the entities of an archive
+	 * Each Entity is at the index of AchiveId in the array
+	 */
+	TArray<TSharedPtr<FEntity>> ArchiveEntities;
+
+	/**
+	 * Array (use as a Map<FIdent, TArray<TSharedPtr<FEntity>*>>)
+	 * Some Entity reference an other Entities (with a TSharedPtr).
+	 * If at the deserialization of the Entity, the referenced entity is not yet created,
+	 * the pointer of the TSharedPtr of the referenced entity is added in the array associated to the ArchiveId of the referenced entity.
+	 *
+	 * E.g.
+	 * Edge
+	 *   |- ArchiveId = 5
+	 *   |- Vertex1 (ArchiveId = 6)
+	 *
+	 * IdToWaitingPointers[ArchiveId].Add(&Edge.Vertex)
+	 */
+	TArray<TArray<TSharedPtr<FEntity>*>> ArchiveIdToWaitingSharedPointers;
+	TArray<TArray<TWeakPtr<FEntity>*>> ArchiveIdToWaitingWeakPointers;
+	TArray<TArray<FEntity**>> ArchiveIdToWaitingPointers;
+
+	/**
+	 * Dedicated to SerializeSelection. Set of entity Ids needed to be saved
+	 */
+	TSet<FIdent> EntitiesToBeSerialized;
+
+	/**
+	 * Dedicated to SerializeSelection. Queue of entity Ids waiting to be saved
+	 */
+	TQueue<FIdent> NotYetSerialized;
+
+	mutable bool bForceSpawning = false;
+	bool bIsRecursiveSerialization = false;
+	uint32 EntityCount;
+
+public:
+
+	FDatabase();
+
+	FModel& GetModel();
+
+	/**
+	 * Remove from database
+	 */
+	void RemoveEntity(FIdent idEntity);
+	void RemoveEntity(FEntity& Entity);
+
+
+	/**
+	 * Recursive serialization of a selection. The selection and all the dependencies are serialized.
+	 * All entities that need to be saved (added in EntitiesToBeSerialized set) for the completeness of the archive are added to the queue of entities to Serialized (NotYetSerialized).
+	 */
+	void SerializeSelection(FCADKernelArchive& Ar, const TArray<FIdent>& SelectionIds);
+
+	/**
+	 * All entities of the session are saved in the archive
+	 */
+	void Serialize(FCADKernelArchive& Ar);
+
+	void Deserialize(FCADKernelArchive& Ar);
+
+	/**
+	 * 	Method to clean entity descriptions i.e. removing invalid references (like unsaved linked edges of an edge because the neighbor topological face is not saved)
+	 */
+	void CleanArchiveEntities();
+
+	/**
+	 * Delete all the entities of the session
+	 */
+	void Empty();
+
+	/**
+	 * All Entities are saved in the session map that allows to retrieve them according to their Id
+	 * This process is automatically done during the entity creation (FEntity::MakeShared)
+	 */
+	void AddEntity(TSharedPtr<FEntity> Entity)
 	{
-		friend class FCADKernelArchive;
-		friend class FEntity;
-		friend class FSession;
-
-	protected:
-
-		/**
-		 * Root Geometry Entity node
-		 */
-		TSharedPtr<FModel> Model;
-
-		/**
-		 * Array (use as a Map<FIdent, TSharedPtr<FEntity>>) of all the entities create during the session
-		 * Each Entity.Id is the index of the entity in the array 
-		 */
-		TArray<TSharedPtr<FEntity>> DatabaseEntities;
-
-		/**
-		 * Recycling of the Id to avoid fragmented idents
-		 */
-		TArray<FIdent> AvailableIdents;
-
-		/**
-		 * Array (use as a Map<FIdent, TSharedPtr<FEntity>>) of all the entities of an archive
-		 * Each Entity is at the index of AchiveId in the array
-		 */
-		TArray<TSharedPtr<FEntity>> ArchiveEntities;
-
-		/**
-		 * Array (use as a Map<FIdent, TArray<TSharedPtr<FEntity>*>>) 
-		 * Some Entity reference an other Entities (with a TSharedPtr).
-		 * If at the deserialization of the Entity, the referenced entity is not yet created, 
-		 * the pointer of the TSharedPtr of the referenced entity is added in the array associated to the ArchiveId of the referenced entity.
-		 *
-		 * E.g.
-		 * Edge
-		 *   |- ArchiveId = 5
-		 *   |- Vertex1 (ArchiveId = 6)
-		 * 
-		 * IdToWaitingPointers[ArchiveId].Add(&Edge.Vertex) 
-		 */
-		TArray<TArray<TSharedPtr<FEntity>*>> ArchiveIdToWaitingSharedPointers;
-		TArray<TArray<TWeakPtr<FEntity>*>> ArchiveIdToWaitingWeakPointers;
-		TArray<TArray<FEntity**>> ArchiveIdToWaitingPointers;
-
-		/**
-		 * Dedicated to SerializeSelection. Set of entity Ids needed to be saved
-		 */
-		TSet<FIdent> EntitiesToBeSerialized;
-
-		/**
-		 * Dedicated to SerializeSelection. Queue of entity Ids waiting to be saved
-		 */
-		TQueue<FIdent> NotYetSerialized;
-
-		mutable bool bForceSpawning = false;
-		bool bIsRecursiveSerialization = false;
-		uint32 EntityCount;
-
-	public:
-
-		FDatabase();
-
-		FModel& GetModel();
-
-		/**
-		 * Remove from database
-		 */
-		void RemoveEntity(FIdent idEntity);
-		void RemoveEntity(FEntity& Entity);
-
-
-		/**
-		 * Recursive serialization of a selection. The selection and all the dependencies are serialized. 
-		 * All entities that need to be saved (added in EntitiesToBeSerialized set) for the completeness of the archive are added to the queue of entities to Serialized (NotYetSerialized).
-		 */
-		void SerializeSelection(FCADKernelArchive& Ar, const TArray<FIdent>& SelectionIds);
-
-		/**
-		 * All entities of the session are saved in the archive
-		 */
-		void Serialize(FCADKernelArchive& Ar);
-
-		void Deserialize(FCADKernelArchive& Ar);
-
-		/**
-		 * 	Method to clean entity descriptions i.e. removing invalid references (like unsaved linked edges of an edge because the neighbor topological face is not saved)
-		 */
-		void CleanArchiveEntities();
-
-		/**
-		 * Delete all the entities of the session
-		 */
-		void Empty();
-
-		/**
-		 * All Entities are saved in the session map that allows to retrieve them according to their Id
-		 * This process is automatically done during the entity creation (FEntity::MakeShared)
-		 */
-		void AddEntity(TSharedPtr<FEntity> Entity)
+		if (!Entity.IsValid())
 		{
-			if (!Entity.IsValid())
-			{
-				return;
-			}
-
-			if (Entity->GetId() > 0)
-			{
-				return;
-			}
-			FIdent Id = CreateId();
-			DatabaseEntities[Id] = Entity;
-			Entity->Id = Id;
+			return;
 		}
 
-		/**
-		 * During the load of an FArchive: 
-		 * All Entities are saved in the session map (like Entity created in runtime) that allows to retrieve them according to their Id.
-		 * Archive's entities are also saved in a temporary map for the load purpose (c.f. FEntity::Serialize)
-		 * This process is automatically done during the entity creation (FEntity::MakeShared)
-		 */
-		void AddEntityFromArchive(TSharedPtr<FEntity>Entity)
+		if (Entity->GetId() > 0)
 		{
-			FIdent ArchiveId = Entity->GetId();
-			ensureCADKernel(ArchiveEntities.Num() > (int32)ArchiveId);
-			// Add the entity to ArchiveEntities map.
-			ArchiveEntities[ArchiveId] = Entity;
+			return;
+		}
+		FIdent Id = CreateId();
+		DatabaseEntities[Id] = Entity;
+		Entity->Id = Id;
+	}
+
+	/**
+	 * During the load of an FArchive:
+	 * All Entities are saved in the session map (like Entity created in runtime) that allows to retrieve them according to their Id.
+	 * Archive's entities are also saved in a temporary map for the load purpose (c.f. FEntity::Serialize)
+	 * This process is automatically done during the entity creation (FEntity::MakeShared)
+	 */
+	void AddEntityFromArchive(TSharedPtr<FEntity>Entity)
+	{
+		FIdent ArchiveId = Entity->GetId();
+		ensureCADKernel(ArchiveEntities.Num() > (int32)ArchiveId);
+		// Add the entity to ArchiveEntities map.
+		ArchiveEntities[ArchiveId] = Entity;
 
-			// Check if the entity is not referenced by other entities if the archive(ArchiveIdToWaitingPointers).
-			// If yes(IdToWaitingPointers[ArchiveId].Num() > 0), set all TSharedPtr with the new entity
+		// Check if the entity is not referenced by other entities if the archive(ArchiveIdToWaitingPointers).
+		// If yes(IdToWaitingPointers[ArchiveId].Num() > 0), set all TSharedPtr with the new entity
+		{
+			TArray<TSharedPtr<FEntity>*>& WaitingList = ArchiveIdToWaitingSharedPointers[ArchiveId];
+			if (WaitingList.Num() > 0)
 			{
-				TArray<TSharedPtr<FEntity>*>& WaitingList = ArchiveIdToWaitingSharedPointers[ArchiveId];
-				if (WaitingList.Num() > 0)
+				for (TSharedPtr<FEntity>* SharedPtr : WaitingList)
 				{
-					for (TSharedPtr<FEntity>* SharedPtr : WaitingList)
-					{
-						*SharedPtr = Entity;
-					}
-					WaitingList.Empty();
+					*SharedPtr = Entity;
 				}
+				WaitingList.Empty();
 			}
+		}
 
+		{
+			TArray<TWeakPtr<FEntity>*>& WaitingList = ArchiveIdToWaitingWeakPointers[ArchiveId];
+			if (WaitingList.Num() > 0)
 			{
-				TArray<TWeakPtr<FEntity>*>& WaitingList = ArchiveIdToWaitingWeakPointers[ArchiveId];
-				if (WaitingList.Num() > 0)
+				for (TWeakPtr<FEntity>* WeakPtr : WaitingList)
 				{
-					for (TWeakPtr<FEntity>* WeakPtr : WaitingList)
-					{
-						*WeakPtr = Entity;
-					}
-					WaitingList.Empty();
+					*WeakPtr = Entity;
 				}
+				WaitingList.Empty();
 			}
+		}
 
+		{
+			TArray<FEntity**>& WaitingList = ArchiveIdToWaitingPointers[ArchiveId];
+			if (WaitingList.Num() > 0)
 			{
-				TArray<FEntity**>& WaitingList = ArchiveIdToWaitingPointers[ArchiveId];
-				if (WaitingList.Num() > 0)
+				for (FEntity** Ptr : WaitingList)
 				{
-					for (FEntity** Ptr : WaitingList)
-					{
-						*Ptr = &*Entity;
-					}
-					WaitingList.Empty();
+					*Ptr = &*Entity;
 				}
+				WaitingList.Empty();
 			}
-
-			// Reset Id to add it in the DatabaseEntities map.
-			Entity->Id = 0;
-			// Add the entity to DatabaseEntities map.
-			AddEntity(Entity);
 		}
 
-		/**
-		 * Find TSharedPtr<FEntity> associate to an Archive Id
-		 * If the Entity is not yet created, add the entity in the waiting list
-		 */
-		void SetReferencedEntityOrAddToWaitingList(FIdent ArchiveId, TSharedPtr<FEntity>& Entity)
+		// Reset Id to add it in the DatabaseEntities map.
+		Entity->Id = 0;
+		// Add the entity to DatabaseEntities map.
+		AddEntity(Entity);
+	}
+
+	/**
+	 * Find TSharedPtr<FEntity> associate to an Archive Id
+	 * If the Entity is not yet created, add the entity in the waiting list
+	 */
+	void SetReferencedEntityOrAddToWaitingList(FIdent ArchiveId, TSharedPtr<FEntity>& Entity)
+	{
+		if (ArchiveId == 0)
 		{
-			if (ArchiveId == 0)
-			{
-				Entity = TSharedPtr<FEntity>();
-				return;
-			}
+			Entity = TSharedPtr<FEntity>();
+			return;
+		}
 
-			TSharedPtr<FEntity>& EntityPtr = ArchiveEntities[ArchiveId];
-			if (EntityPtr.IsValid())
-			{
-				Entity = EntityPtr;
-			}
-			else
-			{
-				ArchiveIdToWaitingSharedPointers[ArchiveId].Add(&Entity);
-			}
+		TSharedPtr<FEntity>& EntityPtr = ArchiveEntities[ArchiveId];
+		if (EntityPtr.IsValid())
+		{
+			Entity = EntityPtr;
+		}
+		else
+		{
+			ArchiveIdToWaitingSharedPointers[ArchiveId].Add(&Entity);
 		}
+	}
 
-		/**
-		 * Find TWeakPtr<FEntity> associate to an Archive Id
-		 * If the Entity is not yet created, add the entity in the waiting list
-		 */
-		void SetReferencedEntityOrAddToWaitingList(FIdent ArchiveId, TWeakPtr<FEntity>& Entity)
+	/**
+	 * Find TWeakPtr<FEntity> associate to an Archive Id
+	 * If the Entity is not yet created, add the entity in the waiting list
+	 */
+	void SetReferencedEntityOrAddToWaitingList(FIdent ArchiveId, TWeakPtr<FEntity>& Entity)
+	{
+		if (ArchiveId == 0)
 		{
-			if (ArchiveId == 0)
-			{
-				Entity = TWeakPtr<FEntity>();
-				return;
-			}
+			Entity = TWeakPtr<FEntity>();
+			return;
+		}
 
-			TSharedPtr<FEntity>& EntityPtr = ArchiveEntities[ArchiveId];
-			if (EntityPtr.IsValid())
-			{
-				Entity = EntityPtr;
-			}
-			else
-			{
-				ArchiveIdToWaitingWeakPointers[ArchiveId].Add(&Entity);
-			}
+		TSharedPtr<FEntity>& EntityPtr = ArchiveEntities[ArchiveId];
+		if (EntityPtr.IsValid())
+		{
+			Entity = EntityPtr;
+		}
+		else
+		{
+			ArchiveIdToWaitingWeakPointers[ArchiveId].Add(&Entity);
 		}
+	}
 
-		/**
-		 * Find FEntity* associate to an Archive Id
-		 * If the Entity is not yet created, add the entity in the waiting list
-		 */
-		void SetReferencedEntityOrAddToWaitingList(FIdent ArchiveId, FEntity** Entity)
+	/**
+	 * Find FEntity* associate to an Archive Id
+	 * If the Entity is not yet created, add the entity in the waiting list
+	 */
+	void SetReferencedEntityOrAddToWaitingList(FIdent ArchiveId, FEntity** Entity)
+	{
+		if (ArchiveId == 0)
 		{
-			if (ArchiveId == 0)
-			{
-				*Entity = nullptr;
-				return;
-			}
+			*Entity = nullptr;
+			return;
+		}
 
-			TSharedPtr<FEntity>& EntityPtr = ArchiveEntities[ArchiveId];
-			if (EntityPtr.IsValid())
-			{
-				*Entity = EntityPtr.Get();
-			}
-			else
-			{
-				ArchiveIdToWaitingPointers[ArchiveId].Add(Entity);
-			}
+		TSharedPtr<FEntity>& EntityPtr = ArchiveEntities[ArchiveId];
+		if (EntityPtr.IsValid())
+		{
+			*Entity = EntityPtr.Get();
+		}
+		else
+		{
+			ArchiveIdToWaitingPointers[ArchiveId].Add(Entity);
 		}
+	}
 
-		uint32 SpawnEntityIdents(const TArray<TSharedPtr<FEntity>>& SelectedEntities, bool bForceSpawning = false);
-		uint32 SpawnEntityIdents(const TArray<FEntity*>& SelectedEntities, bool bForceSpawning = false);
-		uint32 SpawnEntityIdent(FEntity& SelectedEntity, bool bForceSpawning = false);
+	uint32 SpawnEntityIdents(const TArray<TSharedPtr<FEntity>>& SelectedEntities, bool bForceSpawning = false);
+	uint32 SpawnEntityIdents(const TArray<FEntity*>& SelectedEntities, bool bForceSpawning = false);
+	uint32 SpawnEntityIdent(FEntity& SelectedEntity, bool bForceSpawning = false);
 
-	protected:
-		/**
-		 * Dedicated method for SerializeSelection
-		 */
-		void AddEntityToSave(FIdent EntityId)
+protected:
+	/**
+	 * Dedicated method for SerializeSelection
+	 */
+	void AddEntityToSave(FIdent EntityId)
+	{
+		if (bIsRecursiveSerialization && !EntitiesToBeSerialized.Find(EntityId))
 		{
-			if (bIsRecursiveSerialization && !EntitiesToBeSerialized.Find(EntityId))
-			{
-				EntitiesToBeSerialized.Add(EntityId);
-				NotYetSerialized.Enqueue(EntityId);
-			}
+			EntitiesToBeSerialized.Add(EntityId);
+			NotYetSerialized.Enqueue(EntityId);
 		}
+	}
 
-	private:
-		FIdent CreateId();
+private:
+	FIdent CreateId();
 
-		// =========================================================================================================================================================================================================
-		// =========================================================================================================================================================================================================
-		// =========================================================================================================================================================================================================
-		//
-		//
-		//                                                                            NOT YET REVIEWED
-		//
-		//
-		// =========================================================================================================================================================================================================
-		// =========================================================================================================================================================================================================
-		// =========================================================================================================================================================================================================
+	// =========================================================================================================================================================================================================
+	// =========================================================================================================================================================================================================
+	// =========================================================================================================================================================================================================
+	//
+	//
+	//                                                                            NOT YET REVIEWED
+	//
+	//
+	// =========================================================================================================================================================================================================
+	// =========================================================================================================================================================================================================
+	// =========================================================================================================================================================================================================
 
-	public:
+public:
 
-		TSharedPtr<FEntity> GetEntity(FIdent id) const;
-		void GetEntities(const TArray<FIdent>& ids, TArray<TSharedPtr<FEntity>>& Entities) const;
-		void GetTopologicalEntities(const TArray<FIdent>& ids, TArray<FTopologicalEntity*>& Entities) const;
-		void GetTopologicalShapeEntities(const TArray<FIdent>& ids, TArray<FTopologicalShapeEntity*>& Entities) const;
+	TSharedPtr<FEntity> GetEntity(FIdent id) const;
+	void GetEntities(const TArray<FIdent>& ids, TArray<TSharedPtr<FEntity>>& Entities) const;
+	void GetTopologicalEntities(const TArray<FIdent>& ids, TArray<FTopologicalEntity*>& Entities) const;
+	void GetTopologicalShapeEntities(const TArray<FIdent>& ids, TArray<FTopologicalShapeEntity*>& Entities) const;
 
-		void GetEntitiesOfType(EEntity Type, TArray<TSharedPtr<FEntity>>& OutEntities) const;
-		void GetEntitiesOfTypes(const TSet<EEntity>& Types, TArray<TSharedPtr<FEntity>>& OutEntities) const;
+	void GetEntitiesOfType(EEntity Type, TArray<TSharedPtr<FEntity>>& OutEntities) const;
+	void GetEntitiesOfTypes(const TSet<EEntity>& Types, TArray<TSharedPtr<FEntity>>& OutEntities) const;
 
-		void ExpandSelection(TSharedPtr<FEntity> Entities, const TSet<EEntity>& Filter, TSet<TSharedPtr<FEntity>>& OutSelection) const;
-		void ExpandSelection(const TArray<TSharedPtr<FEntity>>& Entities, const TSet<EEntity>& Filter, TSet<TSharedPtr<FEntity>>& OutSelection) const;
-		void ExpandSelection(const TArray<TSharedPtr<FEntity>>& Entities, const TSet<EEntity>& Filter, TArray<TSharedPtr<FEntity>>& OutSelection) const;
+	void ExpandSelection(TSharedPtr<FEntity> Entities, const TSet<EEntity>& Filter, TSet<TSharedPtr<FEntity>>& OutSelection) const;
+	void ExpandSelection(const TArray<TSharedPtr<FEntity>>& Entities, const TSet<EEntity>& Filter, TSet<TSharedPtr<FEntity>>& OutSelection) const;
+	void ExpandSelection(const TArray<TSharedPtr<FEntity>>& Entities, const TSet<EEntity>& Filter, TArray<TSharedPtr<FEntity>>& OutSelection) const;
 
-		void TopologicalEntitiesSelection(const TArray<TSharedPtr<FEntity>>& Entities, TArray<TSharedPtr<FTopologicalEntity>>& OutTopoEntities, bool bFindSurfaces, bool bFindEdges, bool bFindVertices) const;
-		void CadEntitiesSelection(const TArray<TSharedPtr<FEntity>>& Entities, TArray<TSharedPtr<FEntity>>& OutCadEntities) const;
-		void MeshEntitiesSelection(const TArray<TSharedPtr<FEntity>>& Entities, TArray<TSharedPtr<FEntity>>& OutMeshEntities) const;
+	void TopologicalEntitiesSelection(const TArray<TSharedPtr<FEntity>>& Entities, TArray<TSharedPtr<FTopologicalEntity>>& OutTopoEntities, bool bFindSurfaces, bool bFindEdges, bool bFindVertices) const;
+	void CadEntitiesSelection(const TArray<TSharedPtr<FEntity>>& Entities, TArray<TSharedPtr<FEntity>>& OutCadEntities) const;
+	void MeshEntitiesSelection(const TArray<TSharedPtr<FEntity>>& Entities, TArray<TSharedPtr<FEntity>>& OutMeshEntities) const;
 
-		TSharedPtr<FEntity> GetFirstEntityOfType(EEntity Type) const;
-	};
+	TSharedPtr<FEntity> GetFirstEntityOfType(EEntity Type) const;
+};
 
-} // namespace CADKernel
+} // namespace UE::CADKernel
 
diff --git a/Engine/Source/Runtime/Datasmith/CADKernel/Public/CADKernel/Core/Entity.h b/Engine/Source/Runtime/Datasmith/CADKernel/Public/CADKernel/Core/Entity.h
--- a/Engine/Source/Runtime/Datasmith/CADKernel/Public/CADKernel/Core/Entity.h
+++ b/Engine/Source/Runtime/Datasmith/CADKernel/Public/CADKernel/Core/Entity.h
@@ -12,249 +12,249 @@
 #include "Toolkit/Core/InfoEntity.h"
 #endif
 
-namespace CADKernel
+namespace UE::CADKernel
 {
-	enum class EEntity : uint8
-	{
-		NullEntity = 0,
-
-		Curve,
-		Surface,
-
-		EdgeLink,
-		VertexLink,
-		TopologicalEdge,
-		TopologicalFace,
-		TopologicalLink,
-		TopologicalLoop,
-		TopologicalVertex,
-		Shell,
-		Body,
-		Model,
-
-		MeshModel,
-		Mesh,
-
-		Group,
-		Criterion,
-		Property,
-		
-		EntityTypeEnd
-	};
-
-	class FDatabase;
-	class FEntity;
-	class FInfoEntity;
-	class FSession;
-	class FSystem;
-
-	class CADKERNEL_API FEntity : public TSharedFromThis<FEntity>, public FHaveStates
-	{
-		friend class FDatabase;
+enum class EEntity : uint8
+{
+	NullEntity = 0,
+
+	Curve,
+	Surface,
+
+	EdgeLink,
+	VertexLink,
+	TopologicalEdge,
+	TopologicalFace,
+	TopologicalLink,
+	TopologicalLoop,
+	TopologicalVertex,
+	Shell,
+	Body,
+	Model,
+
+	MeshModel,
+	Mesh,
+
+	Group,
+	Criterion,
+	Property,
+
+	EntityTypeEnd
+};
+
+class FDatabase;
+class FEntity;
+class FInfoEntity;
+class FSession;
+class FSystem;
+
+class CADKERNEL_API FEntity : public TSharedFromThis<FEntity>, public FHaveStates
+{
+	friend class FDatabase;
 
-	protected:
-		static const TCHAR* TypesNames[];
-		FIdent Id = 0;
+protected:
+	static const TCHAR* TypesNames[];
+	FIdent Id = 0;
 
-	public:
-		virtual ~FEntity();
+public:
+	virtual ~FEntity();
 
-		virtual void Delete()
-		{
-			SetDeleted();
-		}
+	virtual void Delete()
+	{
+		SetDeleted();
+	}
 
-		template<typename OtherEntity, typename... InArgTypes>
-		static TSharedRef<OtherEntity> MakeShared(InArgTypes&&... Args)
-		{
-			OtherEntity* Entity = new OtherEntity(Forward<InArgTypes>(Args)...);
-			TSharedRef<OtherEntity> NewShared = MakeShareable<OtherEntity>(Entity);
+	template<typename OtherEntity, typename... InArgTypes>
+	static TSharedRef<OtherEntity> MakeShared(InArgTypes&&... Args)
+	{
+		OtherEntity* Entity = new OtherEntity(Forward<InArgTypes>(Args)...);
+		TSharedRef<OtherEntity> NewShared = MakeShareable<OtherEntity>(Entity);
 
 #ifdef CADKERNEL_DEV
-			AddEntityInDatabase(NewShared);
+		AddEntityInDatabase(NewShared);
 #endif
-			return NewShared;
-		}
+		return NewShared;
+	}
 
-		template<typename OtherEntity>
-		static TSharedRef<OtherEntity> MakeShared(FCADKernelArchive& Archive)
-		{
-			OtherEntity* Entity = new OtherEntity();
-			TSharedRef<OtherEntity> NewShared = MakeShareable<OtherEntity>(Entity);
-			return NewShared;
-		}
-
-		/**
-		 * To Serialize a TSharedPtr<FEntity>, the Id of the FEntity is saved in the archive (ArchiveId). 
-		 * At the deserialization of the archive, a map is build linking ArchiveId to the TSharedPtr of the deserialized entities.
-		 * At the deserialization of a TSharedPtr<FEntity>, the ArchiveId of the TSharedPtr<FEntity> is get. The map return the associated TSharedPtr.
-		 * If the Entity is not yet deserialized, the TSharedPtr is add to a waiting list to be set as soon as the associated Entity is deserialized
-		 * See FDatabase for more details 
-		 */
-		static void SerializeIdent(FCADKernelArchive& Ar, TSharedPtr<FEntity>& Entity, bool bSaveSelection = true);
-		static void SerializeIdent(FCADKernelArchive& Ar, TWeakPtr<FEntity>& Entity, bool bSaveSelection = true);
-		static void SerializeIdent(FCADKernelArchive& Ar, FEntity** Entity, bool bSaveSelection = true);
-
-		template<typename EntityType>
-		static void SerializeIdent(FCADKernelArchive& Ar, EntityType** Entity, bool bSaveSelection = true)
-		{
-			SerializeIdent(Ar, (FEntity**) Entity, bSaveSelection);
-		}
-
-		template<typename EntityType>
-		static void SerializeIdent(FCADKernelArchive& Ar, TSharedPtr<EntityType>& Entity, bool bSaveSelection = true)
-		{
-			SerializeIdent(Ar, (TSharedPtr<FEntity>&) Entity, bSaveSelection);
-		}
-
-		template<typename EntityType>
-		static void SerializeIdent(FCADKernelArchive& Ar, TWeakPtr<EntityType>& Entity, bool bSaveSelection = true)
-		{
-			SerializeIdent(Ar, (TWeakPtr<FEntity>&) Entity, bSaveSelection);
-		}
+	template<typename OtherEntity>
+	static TSharedRef<OtherEntity> MakeShared(FCADKernelArchive& Archive)
+	{
+		OtherEntity* Entity = new OtherEntity();
+		TSharedRef<OtherEntity> NewShared = MakeShareable<OtherEntity>(Entity);
+		return NewShared;
+	}
+
+	/**
+	 * To Serialize a TSharedPtr<FEntity>, the Id of the FEntity is saved in the archive (ArchiveId).
+	 * At the deserialization of the archive, a map is build linking ArchiveId to the TSharedPtr of the deserialized entities.
+	 * At the deserialization of a TSharedPtr<FEntity>, the ArchiveId of the TSharedPtr<FEntity> is get. The map return the associated TSharedPtr.
+	 * If the Entity is not yet deserialized, the TSharedPtr is add to a waiting list to be set as soon as the associated Entity is deserialized
+	 * See FDatabase for more details
+	 */
+	static void SerializeIdent(FCADKernelArchive& Ar, TSharedPtr<FEntity>& Entity, bool bSaveSelection = true);
+	static void SerializeIdent(FCADKernelArchive& Ar, TWeakPtr<FEntity>& Entity, bool bSaveSelection = true);
+	static void SerializeIdent(FCADKernelArchive& Ar, FEntity** Entity, bool bSaveSelection = true);
+
+	template<typename EntityType>
+	static void SerializeIdent(FCADKernelArchive& Ar, EntityType** Entity, bool bSaveSelection = true)
+	{
+		SerializeIdent(Ar, (FEntity**)Entity, bSaveSelection);
+	}
 
-		/**
-		 * SerializeIdent of each TSharedPtr<FEntity> of the array
-		 */
-		static void SerializeIdents(FCADKernelArchive& Ar, TArray<FEntity*>& Array, bool bSaveSelection = true);
-		static void SerializeIdents(FCADKernelArchive& Ar, TArray<TWeakPtr<FEntity>>& Array, bool bSaveSelection = true);
-		static void SerializeIdents(FCADKernelArchive& Ar, TArray<TSharedPtr<FEntity>>& Array, bool bSaveSelection = true);
-		static void SerializeIdents(FCADKernelArchive& Ar, TArray<TOrientedEntity<FEntity>>& Array);
+	template<typename EntityType>
+	static void SerializeIdent(FCADKernelArchive& Ar, TSharedPtr<EntityType>& Entity, bool bSaveSelection = true)
+	{
+		SerializeIdent(Ar, (TSharedPtr<FEntity>&) Entity, bSaveSelection);
+	}
 
-		template<typename EntityType>
-		static void SerializeIdents(FCADKernelArchive& Ar, TArray<EntityType*>& Array, bool bSaveSelection = true)
-		{
-			SerializeIdents(Ar, (TArray<FEntity*>&) Array, bSaveSelection);
-		}
+	template<typename EntityType>
+	static void SerializeIdent(FCADKernelArchive& Ar, TWeakPtr<EntityType>& Entity, bool bSaveSelection = true)
+	{
+		SerializeIdent(Ar, (TWeakPtr<FEntity>&) Entity, bSaveSelection);
+	}
+
+	/**
+	 * SerializeIdent of each TSharedPtr<FEntity> of the array
+	 */
+	static void SerializeIdents(FCADKernelArchive& Ar, TArray<FEntity*>& Array, bool bSaveSelection = true);
+	static void SerializeIdents(FCADKernelArchive& Ar, TArray<TWeakPtr<FEntity>>& Array, bool bSaveSelection = true);
+	static void SerializeIdents(FCADKernelArchive& Ar, TArray<TSharedPtr<FEntity>>& Array, bool bSaveSelection = true);
+	static void SerializeIdents(FCADKernelArchive& Ar, TArray<TOrientedEntity<FEntity>>& Array);
+
+	template<typename EntityType>
+	static void SerializeIdents(FCADKernelArchive& Ar, TArray<EntityType*>& Array, bool bSaveSelection = true)
+	{
+		SerializeIdents(Ar, (TArray<FEntity*>&) Array, bSaveSelection);
+	}
 
-		template<typename EntityType>
-		static void SerializeIdents(FCADKernelArchive& Ar, TArray<TWeakPtr<EntityType>>& Array, bool bSaveSelection = true)
-		{
-			SerializeIdents(Ar, (TArray<TWeakPtr<FEntity>>&) Array, bSaveSelection);
-		}
+	template<typename EntityType>
+	static void SerializeIdents(FCADKernelArchive& Ar, TArray<TWeakPtr<EntityType>>& Array, bool bSaveSelection = true)
+	{
+		SerializeIdents(Ar, (TArray<TWeakPtr<FEntity>>&) Array, bSaveSelection);
+	}
 
-		template<typename EntityType>
-		static void SerializeIdents(FCADKernelArchive& Ar, TArray<TSharedPtr<EntityType>>& Array, bool bSaveSelection = true)
-		{
-			SerializeIdents(Ar, (TArray<TSharedPtr<FEntity>>&) Array, bSaveSelection);
-		}
+	template<typename EntityType>
+	static void SerializeIdents(FCADKernelArchive& Ar, TArray<TSharedPtr<EntityType>>& Array, bool bSaveSelection = true)
+	{
+		SerializeIdents(Ar, (TArray<TSharedPtr<FEntity>>&) Array, bSaveSelection);
+	}
 
-		/**
-		 * Spawn Ident on each TSharedPtr<FEntity> of the array
-		 */
-		static void SpawnIdentOnEntities(TArray<TSharedPtr<FEntity>>& Array, FDatabase& Database)
+	/**
+	 * Spawn Ident on each TSharedPtr<FEntity> of the array
+	 */
+	static void SpawnIdentOnEntities(TArray<TSharedPtr<FEntity>>& Array, FDatabase& Database)
+	{
+		for (TSharedPtr<FEntity>& Entity : Array)
 		{
-			for (TSharedPtr<FEntity>& Entity : Array)
-			{
-				Entity->SpawnIdent(Database);
-			}
+			Entity->SpawnIdent(Database);
 		}
+	}
 
-		static void SpawnIdentOnEntities(TArray<TOrientedEntity<FEntity>>& Array, FDatabase& Database)
-		{
-			for (TOrientedEntity<FEntity>& Entity : Array)
-			{
-				Entity.Entity->SpawnIdent(Database);
-			}
-		}
-			
-		template<typename EntityType>
-		static void SpawnIdentOnEntities(TArray<TSharedPtr<EntityType>>& Array, FDatabase& Database)
+	static void SpawnIdentOnEntities(TArray<TOrientedEntity<FEntity>>& Array, FDatabase& Database)
+	{
+		for (TOrientedEntity<FEntity>& Entity : Array)
 		{
-			SpawnIdentOnEntities((TArray<TSharedPtr<FEntity>>&) Array, Database);
+			Entity.Entity->SpawnIdent(Database);
 		}
+	}
 
-		/**
-		 * Reset Processed flags Recursively on each TSharedPtr<FEntity> of the array
-		 */
-		static void ResetMarkersRecursivelyOnEntities(TArray<TWeakPtr<FEntity>>& Array)
-		{
-			for (TWeakPtr<FEntity>& Entity : Array)
-			{
-				Entity.Pin()->ResetMarkersRecursively();
-			}
-		}
+	template<typename EntityType>
+	static void SpawnIdentOnEntities(TArray<TSharedPtr<EntityType>>& Array, FDatabase& Database)
+	{
+		SpawnIdentOnEntities((TArray<TSharedPtr<FEntity>>&) Array, Database);
+	}
 
-		static void ResetMarkersRecursivelyOnEntities(TArray<TSharedPtr<FEntity>>& Array)
+	/**
+	 * Reset Processed flags Recursively on each TSharedPtr<FEntity> of the array
+	 */
+	static void ResetMarkersRecursivelyOnEntities(TArray<TWeakPtr<FEntity>>& Array)
+	{
+		for (TWeakPtr<FEntity>& Entity : Array)
 		{
-			for (TSharedPtr<FEntity>& Entity : Array)
-			{
-				Entity->ResetMarkersRecursively();
-			}
+			Entity.Pin()->ResetMarkersRecursively();
 		}
+	}
 
-		static void ResetMarkersRecursivelyOnEntities(TArray<TOrientedEntity<FEntity>>& Array)
+	static void ResetMarkersRecursivelyOnEntities(TArray<TSharedPtr<FEntity>>& Array)
+	{
+		for (TSharedPtr<FEntity>& Entity : Array)
 		{
-			for (TOrientedEntity<FEntity>& OrientedEntity : Array)
-			{
-				OrientedEntity.Entity->ResetMarkersRecursively();
-			}
+			Entity->ResetMarkersRecursively();
 		}
+	}
 
-		template<typename EntityType>
-		static void ResetMarkersRecursivelyOnEntities(TArray<TSharedPtr<EntityType>>& Array)
+	static void ResetMarkersRecursivelyOnEntities(TArray<TOrientedEntity<FEntity>>& Array)
+	{
+		for (TOrientedEntity<FEntity>& OrientedEntity : Array)
 		{
-			ResetMarkersRecursivelyOnEntities((TArray<TSharedPtr<FEntity>>&) Array);
+			OrientedEntity.Entity->ResetMarkersRecursively();
 		}
+	}
 
-		/**
-		 * Serialization of a FEntity. Each class derived from FEntity has to override this method (and call the direct base class override method first)
-		 * E.g.
-		 * 	class FEntityXXX : public FEntityXX
-		 *  {
-		 *		virtual void Serialize(FCADKernelArchive& Ar) override
-		 * 		{
-		 *			FEntityXX::Serialize(Ar);
-		 *			...
-		 * 		}
-		 *	}
-		 */
-		virtual void Serialize(FCADKernelArchive& Ar)
-		{
-			Ar << Id;
-			FHaveStates::Serialize(Ar);
-		}
+	template<typename EntityType>
+	static void ResetMarkersRecursivelyOnEntities(TArray<TSharedPtr<EntityType>>& Array)
+	{
+		ResetMarkersRecursivelyOnEntities((TArray<TSharedPtr<FEntity>>&) Array);
+	}
+
+	/**
+	 * Serialization of a FEntity. Each class derived from FEntity has to override this method (and call the direct base class override method first)
+	 * E.g.
+	 * 	class FEntityXXX : public FEntityXX
+	 *  {
+	 *		virtual void Serialize(FCADKernelArchive& Ar) override
+	 * 		{
+	 *			FEntityXX::Serialize(Ar);
+	 *			...
+	 * 		}
+	 *	}
+	 */
+	virtual void Serialize(FCADKernelArchive& Ar)
+	{
+		Ar << Id;
+		FHaveStates::Serialize(Ar);
+	}
 
-		/**
-		 * Deserialize the next entity in the archive.
-		 */
-		static TSharedPtr<FEntity> Deserialize(FCADKernelArchive& Ar);
+	/**
+	 * Deserialize the next entity in the archive.
+	 */
+	static TSharedPtr<FEntity> Deserialize(FCADKernelArchive& Ar);
 
-		static const TCHAR* GetTypeName(EEntity Type);
+	static const TCHAR* GetTypeName(EEntity Type);
 
-		virtual EEntity GetEntityType() const = 0;
-		const TCHAR* GetTypeName() const
-		{
-			return GetTypeName(GetEntityType());
-		}
+	virtual EEntity GetEntityType() const = 0;
+	const TCHAR* GetTypeName() const
+	{
+		return GetTypeName(GetEntityType());
+	}
 
-		const FIdent& GetId() const
-		{
-			return Id;
-		}
+	const FIdent& GetId() const
+	{
+		return Id;
+	}
 
 #ifdef CADKERNEL_DEV
-		void InfoEntity() const;
-		virtual FInfoEntity& GetInfo(FInfoEntity& EntityInfo) const;
-		static void AddEntityInDatabase(TSharedRef<FEntity> Entity);
+	void InfoEntity() const;
+	virtual FInfoEntity& GetInfo(FInfoEntity& EntityInfo) const;
+	static void AddEntityInDatabase(TSharedRef<FEntity> Entity);
 #endif
 
-		virtual void SpawnIdent(FDatabase& Database)
+	virtual void SpawnIdent(FDatabase& Database)
+	{
+		if (!FEntity::SetId(Database))
 		{
-			if (!FEntity::SetId(Database))
-			{
-				return;
-			}
+			return;
 		}
+	}
 
-		virtual void ResetMarkersRecursively()
-		{
-			ResetMarkers();
-		}
+	virtual void ResetMarkersRecursively()
+	{
+		ResetMarkers();
+	}
 
-	protected:
+protected:
 
-		bool SetId(FDatabase& Database);
-	};
+	bool SetId(FDatabase& Database);
+};
 }
 
diff --git a/Engine/Source/Runtime/Datasmith/CADKernel/Public/CADKernel/Core/HaveStates.h b/Engine/Source/Runtime/Datasmith/CADKernel/Public/CADKernel/Core/HaveStates.h
--- a/Engine/Source/Runtime/Datasmith/CADKernel/Public/CADKernel/Core/HaveStates.h
+++ b/Engine/Source/Runtime/Datasmith/CADKernel/Public/CADKernel/Core/HaveStates.h
@@ -6,141 +6,141 @@
 
 enum class EHaveStates : uint16
 {
-	None                = 0x0000u,  // No flags.
-								    
-	IsApplyCriteria     = 0x0001u,  // used for FTopologicalFace, FEdge
-	IsRemoved           = 0x0001u,  // used for FThinZone2D
-								    
-	IsMeshed            = 0x0002u,  // used for FTopologicalFace, FEdge
-	IsInner             = 0x0002u,  // used for FEdgeSegment, FShell, 
-	IsFirstSide         = 0x0002u,  // used for FThinZoneSide
-	FirstSideClosed     = 0x0002u,  // used for FThinZone2D
-								    
-	IsBackOriented      = 0x0004u,  // used for FTopologicalFace
-	ThinPeak            = 0x0004u,  // used for FEdge 
-	SecondSideClosed    = 0x0004u,  // used for FThinZone2D
-								    
-	ThinZone            = 0x0008u,  // used for FTopologicalFace, FEdge, FEdgeSegment  
-								    
-	Degenerated         = 0x0010u,  // used for FEdge, FGrid
-								    
-	IsDeleted           = 0x0020u,  // used for all class
-	IsVirtuallyMeshed   = 0x0040u,  // used for FEdge
-
-	HasMarker1          = 0x1000u,  // used for all class
-	HasMarker2          = 0x2000u,  // used for all class
-	HasMarker3          = 0x4000u,  // used for all class
-	IsProcess4			= 0x8000u,  // used for all class
-	AllMarkers          = 0xF000u,  // used for all class
-
-	All                 = 0xFFu
+	None = 0x0000u,  // No flags.
+
+	IsApplyCriteria = 0x0001u,  // used for FTopologicalFace, FEdge
+	IsRemoved = 0x0001u,  // used for FThinZone2D
+
+	IsMeshed = 0x0002u,  // used for FTopologicalFace, FEdge
+	IsInner = 0x0002u,  // used for FEdgeSegment, FShell, 
+	IsFirstSide = 0x0002u,  // used for FThinZoneSide
+	FirstSideClosed = 0x0002u,  // used for FThinZone2D
+
+	IsBackOriented = 0x0004u,  // used for FTopologicalFace
+	ThinPeak = 0x0004u,  // used for FEdge 
+	SecondSideClosed = 0x0004u,  // used for FThinZone2D
+
+	ThinZone = 0x0008u,  // used for FTopologicalFace, FEdge, FEdgeSegment  
+
+	Degenerated = 0x0010u,  // used for FEdge, FGrid
+
+	IsDeleted = 0x0020u,  // used for all class
+	IsVirtuallyMeshed = 0x0040u,  // used for FEdge
+
+	HasMarker1 = 0x1000u,  // used for all class
+	HasMarker2 = 0x2000u,  // used for all class
+	HasMarker3 = 0x4000u,  // used for all class
+	IsProcess4 = 0x8000u,  // used for all class
+	AllMarkers = 0xF000u,  // used for all class
+
+	All = 0xFFu
 };
 
 ENUM_CLASS_FLAGS(EHaveStates);
 
-namespace CADKernel
+namespace UE::CADKernel
+{
+class CADKERNEL_API FHaveStates
 {
-	class CADKERNEL_API FHaveStates
-	{
-	protected:
-		mutable EHaveStates States;
-
-	public:
-
-		FHaveStates()
-			: States(EHaveStates::None)
-		{};
-
-		void Serialize(FCADKernelArchive& Ar)
-		{
-			Ar << States;
-		}
-
-		void ResetElementStatus()
-		{
-			States = EHaveStates::None;
-		}
-
-		bool HasMarker1() const
-		{
-			return ((States & EHaveStates::HasMarker1) == EHaveStates::HasMarker1);
-		}
-
-		bool HasMarker2() const
-		{
-			return ((States & EHaveStates::HasMarker2) == EHaveStates::HasMarker2);
-		}
-
-		bool HasMarker3() const
-		{
-			return ((States & EHaveStates::HasMarker3) == EHaveStates::HasMarker3);
-		}
-
-		void SetMarker1() const
-		{
-			States |= EHaveStates::HasMarker1;
-		}
-
-		void SetMarker2() const
-		{
-			States |= EHaveStates::HasMarker2;
-		}
-
-		void SetMarker3() const
-		{
-			States |= EHaveStates::HasMarker3;
-		}
-
-		void ResetMarker1() const
-		{
-			States &= ~EHaveStates::HasMarker1;
-		}
-
-		void ResetMarker2() const
-		{
-			States &= ~EHaveStates::HasMarker2;
-		}
-
-		void ResetMarker3() const
-		{
-			States &= ~EHaveStates::HasMarker3;
-		}
-
-		void ResetMarkers() const
-		{
-			States &= ~EHaveStates::AllMarkers;
-		}
-
-		bool IsDeleted() const
-		{
-			return ((States & EHaveStates::IsDeleted) == EHaveStates::IsDeleted);
-		}
-
-		void SetDeleted() const
-		{
-			States |= EHaveStates::IsDeleted;
-		}
-
-		void ResetDeleted() const
-		{
-			States &= ~EHaveStates::IsDeleted;
-		}
-
-		virtual bool IsDegenerated() const
-		{
-			return ((States & EHaveStates::Degenerated) == EHaveStates::Degenerated);
-		}
-
-		virtual void SetAsDegenerated() const
-		{
-			States |= EHaveStates::Degenerated;
-		}
-
-		virtual void ResetDegenerated() const
-		{
-			States &= ~EHaveStates::Degenerated;
-		}
-
-	};
+protected:
+	mutable EHaveStates States;
+
+public:
+
+	FHaveStates()
+		: States(EHaveStates::None)
+	{};
+
+	void Serialize(FCADKernelArchive& Ar)
+	{
+		Ar << States;
+	}
+
+	void ResetElementStatus()
+	{
+		States = EHaveStates::None;
+	}
+
+	bool HasMarker1() const
+	{
+		return ((States & EHaveStates::HasMarker1) == EHaveStates::HasMarker1);
+	}
+
+	bool HasMarker2() const
+	{
+		return ((States & EHaveStates::HasMarker2) == EHaveStates::HasMarker2);
+	}
+
+	bool HasMarker3() const
+	{
+		return ((States & EHaveStates::HasMarker3) == EHaveStates::HasMarker3);
+	}
+
+	void SetMarker1() const
+	{
+		States |= EHaveStates::HasMarker1;
+	}
+
+	void SetMarker2() const
+	{
+		States |= EHaveStates::HasMarker2;
+	}
+
+	void SetMarker3() const
+	{
+		States |= EHaveStates::HasMarker3;
+	}
+
+	void ResetMarker1() const
+	{
+		States &= ~EHaveStates::HasMarker1;
+	}
+
+	void ResetMarker2() const
+	{
+		States &= ~EHaveStates::HasMarker2;
+	}
+
+	void ResetMarker3() const
+	{
+		States &= ~EHaveStates::HasMarker3;
+	}
+
+	void ResetMarkers() const
+	{
+		States &= ~EHaveStates::AllMarkers;
+	}
+
+	bool IsDeleted() const
+	{
+		return ((States & EHaveStates::IsDeleted) == EHaveStates::IsDeleted);
+	}
+
+	void SetDeleted() const
+	{
+		States |= EHaveStates::IsDeleted;
+	}
+
+	void ResetDeleted() const
+	{
+		States &= ~EHaveStates::IsDeleted;
+	}
+
+	virtual bool IsDegenerated() const
+	{
+		return ((States & EHaveStates::Degenerated) == EHaveStates::Degenerated);
+	}
+
+	virtual void SetAsDegenerated() const
+	{
+		States |= EHaveStates::Degenerated;
+	}
+
+	virtual void ResetDegenerated() const
+	{
+		States &= ~EHaveStates::Degenerated;
+	}
+
+};
 }
 
diff --git a/Engine/Source/Runtime/Datasmith/CADKernel/Public/CADKernel/Core/OrientedEntity.h b/Engine/Source/Runtime/Datasmith/CADKernel/Public/CADKernel/Core/OrientedEntity.h
--- a/Engine/Source/Runtime/Datasmith/CADKernel/Public/CADKernel/Core/OrientedEntity.h
+++ b/Engine/Source/Runtime/Datasmith/CADKernel/Public/CADKernel/Core/OrientedEntity.h
@@ -4,28 +4,28 @@
 #include "CADKernel/Core/Types.h"
 #include "CADKernel/Geo/GeoEnum.h"
 
-namespace CADKernel
+namespace UE::CADKernel
 {
-	template <typename KeyType>
-	class CADKERNEL_API TOrientedEntity
-	{
-	public:
-		TSharedPtr<KeyType> Entity;
-		EOrientation Direction;
+template <typename KeyType>
+class CADKERNEL_API TOrientedEntity
+{
+public:
+	TSharedPtr<KeyType> Entity;
+	EOrientation Direction;
 
-		TOrientedEntity(TSharedPtr<KeyType> InEntity, EOrientation InDirection)
-			: Entity(InEntity)
-			, Direction(InDirection)
-		{
-		}
+	TOrientedEntity(TSharedPtr<KeyType> InEntity, EOrientation InDirection)
+		: Entity(InEntity)
+		, Direction(InDirection)
+	{
+	}
 
-		TOrientedEntity(const TOrientedEntity<KeyType>& OrientiredEntity)
-			: Entity(OrientiredEntity.Entity)
-			, Direction(OrientiredEntity.Direction)
-		{
-		}
+	TOrientedEntity(const TOrientedEntity<KeyType>& OrientiredEntity)
+		: Entity(OrientiredEntity.Entity)
+		, Direction(OrientiredEntity.Direction)
+	{
+	}
 
-		TOrientedEntity() = default;
-	};
+	TOrientedEntity() = default;
+};
 }
 
