Rename CADKernel namespace into UE::CADKernel

diff --git a/Engine/Source/Runtime/Datasmith/CADKernel/Public/CADKernel/Core/Parameter.h b/Engine/Source/Runtime/Datasmith/CADKernel/Public/CADKernel/Core/Parameter.h
--- a/Engine/Source/Runtime/Datasmith/CADKernel/Public/CADKernel/Core/Parameter.h
+++ b/Engine/Source/Runtime/Datasmith/CADKernel/Public/CADKernel/Core/Parameter.h
@@ -5,303 +5,303 @@
 
 #include "Serialization/Archive.h"
 
-namespace CADKernel
+namespace UE::CADKernel
 {
-	class CADKERNEL_API FParameters;
-
-	class CADKERNEL_API FParameterValue
-	{
-	public:
-		union
-		{
-			double Real;
-			int32 Integer;
-			bool Boolean;
-		} Value;
-
-		FString StringValue;
-
-	protected:
-		EValue Type;
-
-	public:
-		FParameterValue(int32 InValue)
-			: Type(EValue::Integer)
-		{
-			*this = InValue;
-		}
-
-		FParameterValue(double InValue)
-			: Type(EValue::Double)
-		{
-			*this = InValue;
-		}
-
-		FParameterValue(const FString& InValue)
-			: Type(EValue::String)
-		{
-			*this = InValue;
-		}
-
-		FParameterValue(bool InValue)
-			: Type(EValue::Boolean)
-		{
-			*this = InValue;
-		}
-
-		FParameterValue(EValue InType = EValue::Integer)
-			: Type(InType)
-		{
-			*this = 0;
-		}
-
-		FParameterValue(const FParameterValue& Other)
-		{
-			Type = Other.Type;
-			Value = Other.Value;
-			StringValue = Other.StringValue;
-		}
-
-		FString ToString() const;
-
-		inline EValue GetType() const
-		{
-			return Type;
-		}
-
-		inline void SetType(EValue NewType)
-		{
-			Type = NewType;
-		}
-
-		inline FParameterValue& operator=(double InValue)
-		{
-			Value.Real = InValue;
-			return *this;
-		}
-
-		inline FParameterValue& operator=(int32 InValue)
-		{
-			Value.Integer = InValue;
-			return *this;
-		}
-
-		inline FParameterValue& operator=(const FString& InValue)
-		{
-			StringValue = InValue;
-			return *this;
-		}
-
-		inline FParameterValue& operator=(bool InValue)
-		{
-			Value.Boolean = InValue;
-			return *this;
-		}
-
-		inline operator double& ()
-		{
-			return Value.Real;
-		}
-
-		inline operator FString& ()
-		{
-			return StringValue;
-		}
-
-		inline operator int32& ()
-		{
-			return Value.Integer;
-		}
-
-		inline operator bool& ()
-		{
-			return Value.Boolean;
-		}
-
-		inline operator const double& () const
-		{
-			return Value.Real;
-		}
-
-		inline operator const FString& () const
-		{
-			return StringValue;
-		}
-
-		inline operator const int32& () const
-		{
-			return Value.Integer;
-		}
-
-		inline operator const bool& () const
-		{
-			return Value.Boolean;
-		}
-
-		bool operator==(const FParameterValue& Other) const;
-	};
+class CADKERNEL_API FParameters;
+
+class CADKERNEL_API FParameterValue
+{
+public:
+	union
+	{
+		double Real;
+		int32 Integer;
+		bool Boolean;
+	} Value;
+
+	FString StringValue;
+
+protected:
+	EValue Type;
+
+public:
+	FParameterValue(int32 InValue)
+		: Type(EValue::Integer)
+	{
+		*this = InValue;
+	}
+
+	FParameterValue(double InValue)
+		: Type(EValue::Double)
+	{
+		*this = InValue;
+	}
+
+	FParameterValue(const FString& InValue)
+		: Type(EValue::String)
+	{
+		*this = InValue;
+	}
+
+	FParameterValue(bool InValue)
+		: Type(EValue::Boolean)
+	{
+		*this = InValue;
+	}
+
+	FParameterValue(EValue InType = EValue::Integer)
+		: Type(InType)
+	{
+		*this = 0;
+	}
+
+	FParameterValue(const FParameterValue& Other)
+	{
+		Type = Other.Type;
+		Value = Other.Value;
+		StringValue = Other.StringValue;
+	}
+
+	FString ToString() const;
+
+	inline EValue GetType() const
+	{
+		return Type;
+	}
+
+	inline void SetType(EValue NewType)
+	{
+		Type = NewType;
+	}
+
+	inline FParameterValue& operator=(double InValue)
+	{
+		Value.Real = InValue;
+		return *this;
+	}
 
-	class CADKERNEL_API FParameter
-	{
-	private:
-		FString Name;
-
-		FParameterValue Value;
-		FParameterValue DefaultValue;
-
-		const TCHAR** EnumLabels;
-
-		void AddToParameterMap(FParameter& Parameter, FParameters& Parameters);
-
-	public:
-
-		template<typename ValueType>
-		FParameter(const FString& ParameterName, const ValueType& FirstValue, FParameters& Parameters)
-			: Name(ParameterName)
-			, Value(FirstValue)
-			, DefaultValue(FirstValue)
-			, EnumLabels(nullptr)
-		{
-			AddToParameterMap(*this, Parameters);
-		}
-
-		FParameter(const FString& ParameterName, const TCHAR* v, FParameters& Parameters)
-			: Name(ParameterName)
-			, Value(FString(v))
-			, DefaultValue(FString(v))
-			, EnumLabels(nullptr)
-		{
-			AddToParameterMap(*this, Parameters);
-		}
-
-		FParameter(const FParameter& Other, FParameters& Parameters)
-			: Name(Other.Name)
-			, Value(Other.Value)
-			, DefaultValue(Other.DefaultValue)
-			, EnumLabels(Other.EnumLabels)
-		{
-			AddToParameterMap(*this, Parameters);
-		}
-
-		FParameter()
-			: Name()
-			, Value()
-			, DefaultValue()
-			, EnumLabels()
-		{
-		}
-
-		virtual ~FParameter()
-		{
-		}
-
-		const FParameterValue& GetValue() const
-		{
-			return Value;
-		}
-
-		const FParameterValue& GetDefaultValue() const
-		{
-			return DefaultValue;
-		}
-
-		template<typename ValueType>
-		void SetValue(const ValueType NewValue)
-		{
-			Value = NewValue;
-		}
-
-		template<typename ValueType>
-		void SetDefaultValue(const ValueType NewValue)
-		{
-			DefaultValue = NewValue;
-		}
-
-		template<typename ValueType>
-		inline FParameter& operator= (const ValueType& NewValue)
-		{
-			SetValue(NewValue);
-			return *this;
-		}
-
-		inline FParameter& operator= (const FParameter& Parameter)
-		{
-			SetValue(Parameter.Value);
-			return *this;
-		}
-
-		inline operator int32 () const
-		{
-			return Value.Value.Integer;
-		}
-
-		inline operator double() const
-		{
-			return Value.Value.Real;
-		}
-
-		inline operator bool() const
-		{
-			return Value.Value.Boolean;
-		}
-
-		inline operator const FString& () const
-		{
-			return Value.StringValue;
-		}
-
-		inline bool IsDefault() const 
-		{
-			return Value == DefaultValue;
-		}
-
-		inline const FString& GetName() const
-		{
-			return Name;
-		}
-
-		inline EValue GetType() const
-		{
-			return Value.GetType();
-		}
-
-		FString GetTypeString() const;
-
-		void SetFromString(const FString& str);
-		FString ToString() const;
+	inline FParameterValue& operator=(int32 InValue)
+	{
+		Value.Integer = InValue;
+		return *this;
+	}
+
+	inline FParameterValue& operator=(const FString& InValue)
+	{
+		StringValue = InValue;
+		return *this;
+	}
+
+	inline FParameterValue& operator=(bool InValue)
+	{
+		Value.Boolean = InValue;
+		return *this;
+	}
+
+	inline operator double& ()
+	{
+		return Value.Real;
+	}
+
+	inline operator FString& ()
+	{
+		return StringValue;
+	}
+
+	inline operator int32& ()
+	{
+		return Value.Integer;
+	}
+
+	inline operator bool& ()
+	{
+		return Value.Boolean;
+	}
+
+	inline operator const double& () const
+	{
+		return Value.Real;
+	}
+
+	inline operator const FString& () const
+	{
+		return StringValue;
+	}
+
+	inline operator const int32& () const
+	{
+		return Value.Integer;
+	}
+
+	inline operator const bool& () const
+	{
+		return Value.Boolean;
+	}
+
+	bool operator==(const FParameterValue& Other) const;
+};
+
+class CADKERNEL_API FParameter
+{
+private:
+	FString Name;
+
+	FParameterValue Value;
+	FParameterValue DefaultValue;
+
+	const TCHAR** EnumLabels;
+
+	void AddToParameterMap(FParameter& Parameter, FParameters& Parameters);
+
+public:
+
+	template<typename ValueType>
+	FParameter(const FString& ParameterName, const ValueType& FirstValue, FParameters& Parameters)
+		: Name(ParameterName)
+		, Value(FirstValue)
+		, DefaultValue(FirstValue)
+		, EnumLabels(nullptr)
+	{
+		AddToParameterMap(*this, Parameters);
+	}
+
+	FParameter(const FString& ParameterName, const TCHAR* v, FParameters& Parameters)
+		: Name(ParameterName)
+		, Value(FString(v))
+		, DefaultValue(FString(v))
+		, EnumLabels(nullptr)
+	{
+		AddToParameterMap(*this, Parameters);
+	}
+
+	FParameter(const FParameter& Other, FParameters& Parameters)
+		: Name(Other.Name)
+		, Value(Other.Value)
+		, DefaultValue(Other.DefaultValue)
+		, EnumLabels(Other.EnumLabels)
+	{
+		AddToParameterMap(*this, Parameters);
+	}
+
+	FParameter()
+		: Name()
+		, Value()
+		, DefaultValue()
+		, EnumLabels()
+	{
+	}
+
+	virtual ~FParameter()
+	{
+	}
+
+	const FParameterValue& GetValue() const
+	{
+		return Value;
+	}
+
+	const FParameterValue& GetDefaultValue() const
+	{
+		return DefaultValue;
+	}
+
+	template<typename ValueType>
+	void SetValue(const ValueType NewValue)
+	{
+		Value = NewValue;
+	}
+
+	template<typename ValueType>
+	void SetDefaultValue(const ValueType NewValue)
+	{
+		DefaultValue = NewValue;
+	}
+
+	template<typename ValueType>
+	inline FParameter& operator= (const ValueType& NewValue)
+	{
+		SetValue(NewValue);
+		return *this;
+	}
+
+	inline FParameter& operator= (const FParameter& Parameter)
+	{
+		SetValue(Parameter.Value);
+		return *this;
+	}
+
+	inline operator int32 () const
+	{
+		return Value.Value.Integer;
+	}
+
+	inline operator double() const
+	{
+		return Value.Value.Real;
+	}
+
+	inline operator bool() const
+	{
+		return Value.Value.Boolean;
+	}
+
+	inline operator const FString& () const
+	{
+		return Value.StringValue;
+	}
+
+	inline bool IsDefault() const
+	{
+		return Value == DefaultValue;
+	}
+
+	inline const FString& GetName() const
+	{
+		return Name;
+	}
+
+	inline EValue GetType() const
+	{
+		return Value.GetType();
+	}
+
+	FString GetTypeString() const;
+
+	void SetFromString(const FString& str);
+	FString ToString() const;
+};
+
+class CADKERNEL_API FAngleParameter
+	: public FParameter
+{
+private:
+	double RadianValue;
+	double CosineValue;
+
+public:
+	FAngleParameter(const FString paramName, double AngleValue, FParameters& Parameters)
+		: FParameter(paramName, AngleValue, Parameters)
+	{
+		RadianValue = FMath::DegreesToRadians(AngleValue);
+		CosineValue = cos(RadianValue);
+	}
+
+	void SetValue(const FParameterValue& NewValue)
+	{
+		FParameter::SetValue(NewValue);
+		RadianValue = FMath::DegreesToRadians((double)NewValue);
+		CosineValue = cos(RadianValue);
+	}
+
+	const double GetAngle() const
+	{
+		return RadianValue;
 	};
 
-	class CADKERNEL_API FAngleParameter
-		: public FParameter
-	{
-	private:
-		double RadianValue;
-		double CosineValue;
-
-	public:
-		FAngleParameter(const FString paramName, double AngleValue, FParameters& Parameters)
-			: FParameter(paramName, AngleValue, Parameters)
-		{
-			RadianValue = FMath::DegreesToRadians(AngleValue);
-			CosineValue = cos(RadianValue);
-		}
-
-		void SetValue(const FParameterValue& NewValue)
-		{
-			FParameter::SetValue(NewValue);
-			RadianValue = FMath::DegreesToRadians((double) NewValue);
-			CosineValue = cos(RadianValue);
-		}
-
-		const double GetAngle() const
-		{
-			return RadianValue;
-		};
-
-		const double GetCosine() const
-		{
-			return CosineValue;
-		};
+	const double GetCosine() const
+	{
+		return CosineValue;
 	};
-} // namespace CADKernel
+};
+} // namespace UE::CADKernel
 
diff --git a/Engine/Source/Runtime/Datasmith/CADKernel/Public/CADKernel/Core/Parameters.h b/Engine/Source/Runtime/Datasmith/CADKernel/Public/CADKernel/Core/Parameters.h
--- a/Engine/Source/Runtime/Datasmith/CADKernel/Public/CADKernel/Core/Parameters.h
+++ b/Engine/Source/Runtime/Datasmith/CADKernel/Public/CADKernel/Core/Parameters.h
@@ -3,80 +3,77 @@
 
 #include "CADKernel/Core/Parameter.h"
 #include "CADKernel/Core/Types.h"
-//#include "Containers/Map.h"
-//#include "Containers/UnrealString.h"
-//#include "HAL/Platform.h"
 #include "Serialization/Archive.h"
 
 class FArchive;
 
-namespace CADKernel
+namespace UE::CADKernel
 {
-	class FDatabaseFile;
+class FDatabaseFile;
 
-	class CADKERNEL_API FParameters
-	{
-		friend FParameter; 
-	protected:
+class CADKERNEL_API FParameters
+{
+	friend FParameter;
+protected:
 
-		void Add(FParameter& NewParameter)
+	void Add(FParameter& NewParameter)
+	{
+		FParameter** Parameter = Map.Find(NewParameter.GetName());
+		if (Parameter)
 		{
-			FParameter** Parameter = Map.Find(NewParameter.GetName());
-			if (Parameter)
-			{
-				(*Parameter)->SetValue(NewParameter.GetValue());
-			}
-			else
-			{
-				Map.Add(NewParameter.GetName(), &NewParameter);
-			}
+			(*Parameter)->SetValue(NewParameter.GetValue());
 		}
-
-		TMap<FString, FParameter*> Map;
-
-	public:
-		FParameters(int32 ParametersNum = 3)
+		else
 		{
-			Map.Reserve(ParametersNum);
+			Map.Add(NewParameter.GetName(), &NewParameter);
 		}
+	}
 
-		FParameters& operator=(const FParameters& Other)
-		{
-			for (const TPair<FString, FParameter*>& Parameter : Other.Map)
-			{
-				Map.Add(Parameter.Key, Parameter.Value);
-			}
-			return *this;
-		}
+	TMap<FString, FParameter*> Map;
 
-		virtual ~FParameters() = default;
+public:
+	FParameters(int32 ParametersNum = 3)
+	{
+		Map.Reserve(ParametersNum);
+	}
 
-		FParameter* GetByName(const FString& Name)
+	FParameters& operator=(const FParameters& Other)
+	{
+		for (const TPair<FString, FParameter*>& Parameter : Other.Map)
 		{
-			FParameter** Parameter = Map.Find(Name);
-			return (Parameter != nullptr) ? *Parameter : nullptr;
+			Map.Add(Parameter.Key, Parameter.Value);
 		}
+		return *this;
+	}
 
-		void SetFromString(const FString& ParameterStr);
+	virtual ~FParameters() = default;
 
-		FString ToString(bool bOnlyChanged = false) const;
+	FParameter* GetByName(const FString& Name)
+	{
+		FParameter** Parameter = Map.Find(Name);
+		return (Parameter != nullptr) ? *Parameter : nullptr;
+	}
 
-		void PrintParameterList();
+	void SetFromString(const FString& ParameterStr);
 
-		int32 Count()
-		{
-			return Map.Num();
-		}
+	FString ToString(bool bOnlyChanged = false) const;
 
-		const TMap<FString, FParameter*>& GetMap() const 
-		{
-			return Map;
-		}
+	void PrintParameterList();
 
-		friend FArchive& operator<<(FArchive& Ar, FParameters& Database)
-		{
-			return Ar;
-		}
+	int32 Count()
+	{
+		return Map.Num();
+	}
+
+	const TMap<FString, FParameter*>& GetMap() const
+	{
+		return Map;
+	}
+
+	friend FArchive& operator<<(FArchive& Ar, FParameters& Database)
+	{
+		return Ar;
+	}
 
-	};
-} 
+};
+}
diff --git a/Engine/Source/Runtime/Datasmith/CADKernel/Public/CADKernel/Core/Session.h b/Engine/Source/Runtime/Datasmith/CADKernel/Public/CADKernel/Core/Session.h
--- a/Engine/Source/Runtime/Datasmith/CADKernel/Public/CADKernel/Core/Session.h
+++ b/Engine/Source/Runtime/Datasmith/CADKernel/Public/CADKernel/Core/Session.h
@@ -8,155 +8,155 @@
 
 class FArchive;
 
-namespace CADKernel
+namespace UE::CADKernel
 {
-	class FEntity;
-	class FModel;
+class FEntity;
+class FModel;
 
-	class CADKERNEL_API FSession
-	{
-		friend FEntity;
-		friend FCADKernelArchive;
+class CADKERNEL_API FSession
+{
+	friend FEntity;
+	friend FCADKernelArchive;
 
-	protected:
+protected:
 
-		double GeometricTolerance;
-		FDatabase Database;
-		int32 LastHostId;
+	double GeometricTolerance;
+	FDatabase Database;
+	int32 LastHostId;
 
-	public:
+public:
 #if defined(CADKERNEL_DEV) || defined(CADKERNEL_STDA)
-		static FSession Session;
+	static FSession Session;
 #endif
 
-		FSession(double InGeometricTolerance)
-			: GeometricTolerance(InGeometricTolerance)
-		{
-		}
-
-		FModel& GetModel();
-
-		void Serialize(FCADKernelArchive& Ar)
-		{
-			Ar << GeometricTolerance;
-		}
-
-		FDatabase& GetDatabase()
-		{
-			return Database;
-		}
-
-		/**
-		 * Tolerance must not be modified as soon as a geometric entity has been build.
-		 */
-		void SetGeometricTolerance(double NewTolerance);
-
-		double GetGeometricTolerance() const
-		{
-			return GeometricTolerance;
-		}
-
-		/**
-		 * Save the database as a FAchive in a file
-		 * Mandatory: all entity have to have a defined ID
-		 * Use SpawnEntityIdent if needed
-		 */
-		void SaveDatabase(const TCHAR* FilePath);
-
-		/**
-		 * Save a selection and all the dependencies as a FAchive in a file
-		 * Mandatory: all entity have to have a defined ID
-		 * Use SpawnEntityIdent if needed
-		 */
-		void SaveDatabase(const TCHAR* FileName, const TArray<TSharedPtr<FEntity>>& Entities);
-
-		/**
-		 * Save a selection and all the dependencies as a FAchive in a file
-		 * Mandatory: all entity have to have a defined ID
-		 * Use SpawnEntityIdent if needed
-		 */
-		void SaveDatabase(const TCHAR* FileName, const TArray<FEntity*>& Entities);
-
-		/**
-		 * Save a selection and all the dependencies as a FAchive in a file
-		 * Mandatory: all entity have to have a defined ID
-		 * Use SpawnEntityIdent if needed
-		 */
-		void SaveDatabase(const TCHAR* FileName, FEntity& Entity)
-		{
-			TArray<FEntity*> Entities;
-			Entities.Emplace(&Entity);
-			SaveDatabase(FileName, Entities);
-		}
-
-		/**
-		 * Save a selection and all the dependencies as a FAchive in a file
-		 * Mandatory: all entity have to have a defined ID
-		 * Use SpawnEntityIdent if needed
-		 */
-		void SaveDatabase(const TCHAR* FileName, const TSharedPtr<FEntity> Entity)
-		{
-			TArray<TSharedPtr<FEntity>> Entities;
-			Entities.Emplace(Entity);
-			SaveDatabase(FileName, Entities);
-		}
-
-		/**
-		 * Load and add a database in the current session database
-		 * Entity ID is set for all loaded entities
-		 */
-		void LoadDatabase(const TCHAR* FilePath);
-
-		/**
-		 * Add a database defined by a RawData in the current session database
-		 * Entity ID is set for all loaded entities
-		 */
-		void AddDatabase(const TArray<uint8>& InRawData);
-
-		void Clear()
-		{
-			Database.Empty();
-		}
-
-		/**
-		 * To be consistent,  all entity to save have to had an Id.
-		 * This method browses all sub entities and set their Id if needed
-		 * @param bForceSpawning If false, the process does not iterate through the children of entities with a defined ID
-		 */
-		uint32 SpawnEntityIdent(FEntity& SelectedEntity, bool bForceSpawning = false)
-		{
-			return Database.SpawnEntityIdent(SelectedEntity, bForceSpawning);
-		}
-
-		uint32 SpawnEntityIdents(const TArray<TSharedPtr<FEntity>>& SelectedEntities, bool bForceSpawning = false)
-		{
-			return Database.SpawnEntityIdents(SelectedEntities, bForceSpawning);
-		}
-
-		uint32 SpawnEntityIdents(const TArray<FEntity*>& SelectedEntities, bool bForceSpawning = false)
-		{
-			return Database.SpawnEntityIdents(SelectedEntities, bForceSpawning);
-		}
-
-		int32 GetLastHostId() const 
-		{
-			return LastHostId;
-		}
-
-		int32 NewHostId()
-		{
-			return ++LastHostId;
-		}
-
-		/**
-		 * For stitching purpose, stitching can generate new body needing a host id. 
-		 * To avoid duplicate, the first generated host id can be set.
-		 */
-		void SetFirstNewHostId(int32 StartHostId)
-		{
-			LastHostId = StartHostId;
-		}
-	};
-
-} // namespace CADKernel
+	FSession(double InGeometricTolerance)
+		: GeometricTolerance(InGeometricTolerance)
+	{
+	}
+
+	FModel& GetModel();
+
+	void Serialize(FCADKernelArchive& Ar)
+	{
+		Ar << GeometricTolerance;
+	}
+
+	FDatabase& GetDatabase()
+	{
+		return Database;
+	}
+
+	/**
+	 * Tolerance must not be modified as soon as a geometric entity has been build.
+	 */
+	void SetGeometricTolerance(double NewTolerance);
+
+	double GetGeometricTolerance() const
+	{
+		return GeometricTolerance;
+	}
+
+	/**
+	 * Save the database as a FAchive in a file
+	 * Mandatory: all entity have to have a defined ID
+	 * Use SpawnEntityIdent if needed
+	 */
+	void SaveDatabase(const TCHAR* FilePath);
+
+	/**
+	 * Save a selection and all the dependencies as a FAchive in a file
+	 * Mandatory: all entity have to have a defined ID
+	 * Use SpawnEntityIdent if needed
+	 */
+	void SaveDatabase(const TCHAR* FileName, const TArray<TSharedPtr<FEntity>>& Entities);
+
+	/**
+	 * Save a selection and all the dependencies as a FAchive in a file
+	 * Mandatory: all entity have to have a defined ID
+	 * Use SpawnEntityIdent if needed
+	 */
+	void SaveDatabase(const TCHAR* FileName, const TArray<FEntity*>& Entities);
+
+	/**
+	 * Save a selection and all the dependencies as a FAchive in a file
+	 * Mandatory: all entity have to have a defined ID
+	 * Use SpawnEntityIdent if needed
+	 */
+	void SaveDatabase(const TCHAR* FileName, FEntity& Entity)
+	{
+		TArray<FEntity*> Entities;
+		Entities.Emplace(&Entity);
+		SaveDatabase(FileName, Entities);
+	}
+
+	/**
+	 * Save a selection and all the dependencies as a FAchive in a file
+	 * Mandatory: all entity have to have a defined ID
+	 * Use SpawnEntityIdent if needed
+	 */
+	void SaveDatabase(const TCHAR* FileName, const TSharedPtr<FEntity> Entity)
+	{
+		TArray<TSharedPtr<FEntity>> Entities;
+		Entities.Emplace(Entity);
+		SaveDatabase(FileName, Entities);
+	}
+
+	/**
+	 * Load and add a database in the current session database
+	 * Entity ID is set for all loaded entities
+	 */
+	void LoadDatabase(const TCHAR* FilePath);
+
+	/**
+	 * Add a database defined by a RawData in the current session database
+	 * Entity ID is set for all loaded entities
+	 */
+	void AddDatabase(const TArray<uint8>& InRawData);
+
+	void Clear()
+	{
+		Database.Empty();
+	}
+
+	/**
+	 * To be consistent,  all entity to save have to had an Id.
+	 * This method browses all sub entities and set their Id if needed
+	 * @param bForceSpawning If false, the process does not iterate through the children of entities with a defined ID
+	 */
+	uint32 SpawnEntityIdent(FEntity& SelectedEntity, bool bForceSpawning = false)
+	{
+		return Database.SpawnEntityIdent(SelectedEntity, bForceSpawning);
+	}
+
+	uint32 SpawnEntityIdents(const TArray<TSharedPtr<FEntity>>& SelectedEntities, bool bForceSpawning = false)
+	{
+		return Database.SpawnEntityIdents(SelectedEntities, bForceSpawning);
+	}
+
+	uint32 SpawnEntityIdents(const TArray<FEntity*>& SelectedEntities, bool bForceSpawning = false)
+	{
+		return Database.SpawnEntityIdents(SelectedEntities, bForceSpawning);
+	}
+
+	int32 GetLastHostId() const
+	{
+		return LastHostId;
+	}
+
+	int32 NewHostId()
+	{
+		return ++LastHostId;
+	}
+
+	/**
+	 * For stitching purpose, stitching can generate new body needing a host id.
+	 * To avoid duplicate, the first generated host id can be set.
+	 */
+	void SetFirstNewHostId(int32 StartHostId)
+	{
+		LastHostId = StartHostId;
+	}
+};
+
+} // namespace UE::CADKernel
 
diff --git a/Engine/Source/Runtime/Datasmith/CADKernel/Public/CADKernel/Core/System.h b/Engine/Source/Runtime/Datasmith/CADKernel/Public/CADKernel/Core/System.h
--- a/Engine/Source/Runtime/Datasmith/CADKernel/Public/CADKernel/Core/System.h
+++ b/Engine/Source/Runtime/Datasmith/CADKernel/Public/CADKernel/Core/System.h
@@ -8,151 +8,151 @@
 
 class FArchive;
 
-namespace CADKernel
+namespace UE::CADKernel
 {
-	class FConsole;
-	class FKernelParameters;
-	class FVisu;
+class FConsole;
+class FKernelParameters;
+class FVisu;
 
-	class FSystem
-	{
-		friend class FDatabase;
+class FSystem
+{
+	friend class FDatabase;
 
-	protected:
+protected:
 
-		FString ProductName = FString(TEXT("CADKernel"));
+	FString ProductName = FString(TEXT("CADKernel"));
 
-		TSharedRef<FKernelParameters> Parameters;
+	TSharedRef<FKernelParameters> Parameters;
 
-		FVisu DefaultVisu;
-		FVisu* Viewer;
+	FVisu DefaultVisu;
+	FVisu* Viewer;
 
-		FConsole DefaultConsole;
-		FConsole* Console;
+	FConsole DefaultConsole;
+	FConsole* Console;
 
-		FProgressManager DefaultProgressManager;
-		FProgressManager* ProgressManager;
+	FProgressManager DefaultProgressManager;
+	FProgressManager* ProgressManager;
 
-		static TUniquePtr<FSystem> Instance;
+	static TUniquePtr<FSystem> Instance;
 
-		TSharedPtr<FArchive> LogFile;
-		EVerboseLevel LogLevel;
-		TSharedPtr<FArchive> SpyFile;
+	TSharedPtr<FArchive> LogFile;
+	EVerboseLevel LogLevel;
+	TSharedPtr<FArchive> SpyFile;
 
 #if defined(CADKERNEL_DEV) || defined(CADKERNEL_STDA)
-		TSharedPtr<FArchive> ReportFile;
-		TSharedPtr<FArchive> ReportHeaderFile;
+	TSharedPtr<FArchive> ReportFile;
+	TSharedPtr<FArchive> ReportHeaderFile;
 #endif
 
-		EVerboseLevel VerboseLevel;
+	EVerboseLevel VerboseLevel;
 
-	public:
+public:
 
-		FSystem();
+	FSystem();
 
-		void Initialize(bool bIsDll = true, const FString& LogFilePath = TEXT(""), const FString& SpyFilePath = TEXT(""));
+	void Initialize(bool bIsDll = true, const FString& LogFilePath = TEXT(""), const FString& SpyFilePath = TEXT(""));
 
-		void Shutdown();
+	void Shutdown();
 
-		void CloseLogFiles();
+	void CloseLogFiles();
 
-		static FString GetToolkitVersion();
-		static FString GetCompilationDate();
+	static FString GetToolkitVersion();
+	static FString GetCompilationDate();
 
-		FVisu* GetVisu() const
-		{
-			return Viewer;
-		}
+	FVisu* GetVisu() const
+	{
+		return Viewer;
+	}
 
-		FConsole& GetConsole() const
-		{
-			ensureCADKernel(Console);
-			return *Console;
-		}
+	FConsole& GetConsole() const
+	{
+		ensureCADKernel(Console);
+		return *Console;
+	}
 
-		void SetViewer(FVisu* NewViewer)
-		{
-			Viewer = NewViewer;
-		}
+	void SetViewer(FVisu* NewViewer)
+	{
+		Viewer = NewViewer;
+	}
 
-		void SetConsole(FConsole* InConsole)
-		{
-			Console = InConsole;
-		}
+	void SetConsole(FConsole* InConsole)
+	{
+		Console = InConsole;
+	}
 
-		FProgressManager& GetProgressManager()
-		{
-			ensureCADKernel(ProgressManager);
-			return *ProgressManager;
-		}
-			
-		void SetProgressManager(FProgressManager* InProgressManager)
-		{
-			ProgressManager = InProgressManager;
-		}
+	FProgressManager& GetProgressManager()
+	{
+		ensureCADKernel(ProgressManager);
+		return *ProgressManager;
+	}
 
-		TSharedRef<FKernelParameters> GetParameters()
-		{
-			return Parameters;
-		}
+	void SetProgressManager(FProgressManager* InProgressManager)
+	{
+		ProgressManager = InProgressManager;
+	}
 
-		EVerboseLevel GetVerboseLevel() const
-		{
-			return VerboseLevel;
-		}
+	TSharedRef<FKernelParameters> GetParameters()
+	{
+		return Parameters;
+	}
 
-		void SetVerboseLevel(EVerboseLevel Level)
-		{
-			VerboseLevel = Level;
-		}
+	EVerboseLevel GetVerboseLevel() const
+	{
+		return VerboseLevel;
+	}
 
-		void InitializeCADKernel();
+	void SetVerboseLevel(EVerboseLevel Level)
+	{
+		VerboseLevel = Level;
+	}
 
-		EVerboseLevel GetLogLevel() const
-		{
-			return LogLevel;
-		}
+	void InitializeCADKernel();
 
-		void DefineLogFile(const FString& LogFilePath, EVerboseLevel Level = Log);
-		TSharedPtr<FArchive> GetLogFile() const
-		{
-			return LogFile;
-		}
+	EVerboseLevel GetLogLevel() const
+	{
+		return LogLevel;
+	}
 
-		void DefineSpyFile(const FString& SpyFilePath);
-		TSharedPtr<FArchive> GetSpyFile() const
-		{
-			return SpyFile;
-		}
+	void DefineLogFile(const FString& LogFilePath, EVerboseLevel Level = Log);
+	TSharedPtr<FArchive> GetLogFile() const
+	{
+		return LogFile;
+	}
+
+	void DefineSpyFile(const FString& SpyFilePath);
+	TSharedPtr<FArchive> GetSpyFile() const
+	{
+		return SpyFile;
+	}
 
 
 #if defined(CADKERNEL_DEV) || defined(CADKERNEL_STDA)
-		void DefineReportFile(const FString& InLogFilePath);
-		TSharedPtr<FArchive> GetReportFile() const
-		{
-			return ReportFile;
-		}
+	void DefineReportFile(const FString& InLogFilePath);
+	TSharedPtr<FArchive> GetReportFile() const
+	{
+		return ReportFile;
+	}
 
-		TSharedPtr<FArchive> GetReportHeaderFile() const
-		{
-			return ReportHeaderFile;
-		}
+	TSharedPtr<FArchive> GetReportHeaderFile() const
+	{
+		return ReportHeaderFile;
+	}
 #endif
 
-		static FSystem& Get()
+	static FSystem& Get()
+	{
+		if (!Instance.IsValid())
 		{
-			if (!Instance.IsValid())
-			{
-				Instance = MakeUnique<FSystem>();
-			}
-			return *Instance;
+			Instance = MakeUnique<FSystem>();
 		}
+		return *Instance;
+	}
 
-	protected:
+protected:
 
-		void PrintHeader();
+	void PrintHeader();
 
-	};
+};
 
-} // namespace CADKernel
+} // namespace UE::CADKernel
 
diff --git a/Engine/Source/Runtime/Datasmith/CADKernel/Public/CADKernel/Core/Types.h b/Engine/Source/Runtime/Datasmith/CADKernel/Public/CADKernel/Core/Types.h
--- a/Engine/Source/Runtime/Datasmith/CADKernel/Public/CADKernel/Core/Types.h
+++ b/Engine/Source/Runtime/Datasmith/CADKernel/Public/CADKernel/Core/Types.h
@@ -18,31 +18,31 @@
 
 typedef uint32 FIdent;
 
-namespace CADKernel
+namespace UE::CADKernel
 {
-	enum class EValue : uint8
-	{
-		Entity,
-		OrientedEntity,
-		Point,
-		Matrix,
-		Integer,
-		Double,
-		String,
-		Boolean,
-		Tuple,
-		List,
-		Array
-	};
+enum class EValue : uint8
+{
+	Entity,
+	OrientedEntity,
+	Point,
+	Matrix,
+	Integer,
+	Double,
+	String,
+	Boolean,
+	Tuple,
+	List,
+	Array
+};
 
-	enum EVerboseLevel : uint8
-	{
-		NoVerbose = 0,
-		Spy,
-		Log,
-		Debug
-	};
-} // namespace CADKernel
+enum EVerboseLevel : uint8
+{
+	NoVerbose = 0,
+	Spy,
+	Log,
+	Debug
+};
+} // namespace UE::CADKernel
 
 #ifdef DO_ENSURE_CADKERNEL
 #define ensureCADKernel(InExpression) ensure(InExpression)
diff --git a/Engine/Source/Runtime/Datasmith/CADKernel/Public/CADKernel/Geo/Curves/Curve.h b/Engine/Source/Runtime/Datasmith/CADKernel/Public/CADKernel/Geo/Curves/Curve.h
--- a/Engine/Source/Runtime/Datasmith/CADKernel/Public/CADKernel/Geo/Curves/Curve.h
+++ b/Engine/Source/Runtime/Datasmith/CADKernel/Public/CADKernel/Geo/Curves/Curve.h
@@ -12,228 +12,228 @@
 #include "CADKernel/UI/Message.h"
 #include "CADKernel/Utils/Cache.h"
 
-namespace CADKernel
+namespace UE::CADKernel
 {
-	struct FCurvePoint2D;
-	struct FCurvePoint;
-	struct FNurbsCurveData;
+struct FCurvePoint2D;
+struct FCurvePoint;
+struct FNurbsCurveData;
 
-	class CADKERNEL_API FCurve : public FEntityGeom
-	{
-		friend class FEntity;
+class CADKERNEL_API FCurve : public FEntityGeom
+{
+	friend class FEntity;
 
-	protected:
+protected:
 
-		mutable TCache<double> GlobalLength;
-		FLinearBoundary Boundary;
+	mutable TCache<double> GlobalLength;
+	FLinearBoundary Boundary;
 
-		int8 Dimension;
+	int8 Dimension;
 
-		FCurve(int8 InDimension = 3)
-			: FEntityGeom()
-			, Boundary()
-			, Dimension(InDimension)
-		{
-		}
+	FCurve(int8 InDimension = 3)
+		: FEntityGeom()
+		, Boundary()
+		, Dimension(InDimension)
+	{
+	}
 
-		FCurve(const FLinearBoundary& InBounds, int8 InDimension = 3)
-			: FEntityGeom()
-			, Boundary(InBounds)
-			, Dimension(InDimension)
-		{
-		}
+	FCurve(const FLinearBoundary& InBounds, int8 InDimension = 3)
+		: FEntityGeom()
+		, Boundary(InBounds)
+		, Dimension(InDimension)
+	{
+	}
 
-	public:
+public:
 
-		static TSharedPtr<FCurve> MakeNurbsCurve(FNurbsCurveData& InNurbsData);
-		// TODO
-		//static TSharedPtr<FCurve> MakeBezierCurve(const double InToleranceGeometric, int32 InUDegre, int32 InVDegre, const TArray<FPoint>& InPoles);
-		//static TSharedPtr<FCurve> MakeConeCurve(const double InToleranceGeometric, const FMatrixH& InMatrix, double InStartRadius, double InConeAngle, const FSurfacicBoundary& InBoundary);
-		//static TSharedPtr<FCurve> MakeCylinderCurve(const double InToleranceGeometric, const FMatrixH& InMatrix, const double InRadius, const FSurfacicBoundary& InBoundary);
-		//static TSharedPtr<FCurve> MakePlaneCurve(const double InToleranceGeometric, const FMatrixH& InMatrix, const FSurfacicBoundary& InBoundary);
-		//static TSharedPtr<FCurve> MakeSphericalCurve(const double InToleranceGeometric, const FMatrixH& InMatrix, double InRadius, const FSurfacicBoundary& InBoundary);
-		//static TSharedPtr<FCurve> MakeTorusCurve(const double InToleranceGeometric, const FMatrixH& InMatrix, double InMajorRadius, double InMinorRadius, const FSurfacicBoundary& InBoundary);
+	static TSharedPtr<FCurve> MakeNurbsCurve(FNurbsCurveData& InNurbsData);
+	// TODO
+	//static TSharedPtr<FCurve> MakeBezierCurve(const double InToleranceGeometric, int32 InUDegre, int32 InVDegre, const TArray<FPoint>& InPoles);
+	//static TSharedPtr<FCurve> MakeConeCurve(const double InToleranceGeometric, const FMatrixH& InMatrix, double InStartRadius, double InConeAngle, const FSurfacicBoundary& InBoundary);
+	//static TSharedPtr<FCurve> MakeCylinderCurve(const double InToleranceGeometric, const FMatrixH& InMatrix, const double InRadius, const FSurfacicBoundary& InBoundary);
+	//static TSharedPtr<FCurve> MakePlaneCurve(const double InToleranceGeometric, const FMatrixH& InMatrix, const FSurfacicBoundary& InBoundary);
+	//static TSharedPtr<FCurve> MakeSphericalCurve(const double InToleranceGeometric, const FMatrixH& InMatrix, double InRadius, const FSurfacicBoundary& InBoundary);
+	//static TSharedPtr<FCurve> MakeTorusCurve(const double InToleranceGeometric, const FMatrixH& InMatrix, double InMajorRadius, double InMinorRadius, const FSurfacicBoundary& InBoundary);
 
-		virtual void Serialize(FCADKernelArchive& Ar) override
+	virtual void Serialize(FCADKernelArchive& Ar) override
+	{
+		// Curve's type is serialize because it is used to instantiate the correct entity on deserialization (@see Deserialize(FCADKernelArchive& Archive)) 
+		if (Ar.IsSaving())
 		{
-			// Curve's type is serialize because it is used to instantiate the correct entity on deserialization (@see Deserialize(FCADKernelArchive& Archive)) 
-			if (Ar.IsSaving())
-			{
-				ECurve CurveType = GetCurveType();
-				Ar << CurveType;
-			}
-			FEntityGeom::Serialize(Ar);
-			Ar << Dimension;
-			Ar << Boundary;
+			ECurve CurveType = GetCurveType();
+			Ar << CurveType;
 		}
+		FEntityGeom::Serialize(Ar);
+		Ar << Dimension;
+		Ar << Boundary;
+	}
 
-		/**
-		 * Specific method for curve family to instantiate the correct derived class of FCurve 
-		 */
-		static TSharedPtr<FEntity> Deserialize(FCADKernelArchive& Archive);
+	/**
+	 * Specific method for curve family to instantiate the correct derived class of FCurve
+	 */
+	static TSharedPtr<FEntity> Deserialize(FCADKernelArchive& Archive);
 
 #ifdef CADKERNEL_DEV
-		virtual FInfoEntity& GetInfo(FInfoEntity&) const override;
+	virtual FInfoEntity& GetInfo(FInfoEntity&) const override;
 #endif
 
-		int32 GetDimension() const
-		{
-			return Dimension;
-		}
-
-		virtual EEntity GetEntityType() const override
-		{
-			return EEntity::Curve;
-		}
-
-		virtual ECurve GetCurveType() const = 0;
-
-		double GetUMin() const
-		{
-			return Boundary.Min;
-		};
-
-		double GetUMax() const
-		{
-			return Boundary.Max;
-		};
-
-		const FLinearBoundary& GetBoundary() const
-		{
-			return Boundary;
-		}
-
-		virtual TSharedPtr<FEntityGeom> ApplyMatrix(const FMatrixH& InMatrix) const override = 0;
-
-		double GetLength(double Tolerance) const;
-
-		/**
-		 * Evaluate exact 3D point of the curve at the input Coordinate
-		 * The function can only be used with 3D curve (Dimension == 3) 
-		 */
-		virtual void EvaluatePoint(double Coordinate, FCurvePoint& OutPoint, int32 DerivativeOrder = 0) const
-		{
-			ensureCADKernel(false);
-		}
-
-		/**
-		 * Evaluate exact 3D point of the curve at the input Coordinate
-		 * The function can only be used with 3D curve (Dimension == 3)
-		 */
-		virtual FPoint EvaluatePoint(double Coordinate) const
-		{
-			FCurvePoint OutPoint;
-			EvaluatePoint(Coordinate, OutPoint);
-			return OutPoint.Point;
-		}
+	int32 GetDimension() const
+	{
+		return Dimension;
+	}
 
-		/**
-		 * Evaluate exact 2D point of the curve at the input Coordinate
-		 * The function can only be used with 2D curve (Dimension == 2)
-		 */
-		virtual void Evaluate2DPoint(double Coordinate, FCurvePoint2D& OutPoint, int32 DerivativeOrder = 0) const
-		{
-			ensureCADKernel(false);
-		}
+	virtual EEntity GetEntityType() const override
+	{
+		return EEntity::Curve;
+	}
 
-		/**
-		 * Evaluate exact 3D point of the curve at the input Coordinate
-		 * The function can only be used with 3D curve (Dimension == 3)
-		 */
-		virtual FPoint2D Evaluate2DPoint(double Coordinate) const
-		{
-			FCurvePoint2D OutPoint;
-			Evaluate2DPoint(Coordinate, OutPoint);
-			return OutPoint.Point;
-		}
+	virtual ECurve GetCurveType() const = 0;
 
-		/**
-		 * Evaluate exact 2D point of the curve at the input Coordinate
-		 * The function can only be used with 2D curve (Dimension == 2)
-		 */
-		virtual void Evaluate2DPoint(double Coordinate, FPoint2D& OutPoint) const
-		{
-			ensure(Dimension == 2);
-			FCurvePoint2D CurvePoint;
-			Evaluate2DPoint(Coordinate, CurvePoint);
-			OutPoint = CurvePoint.Point;
-		}
+	double GetUMin() const
+	{
+		return Boundary.Min;
+	};
 
-		/**
-		 * Evaluate exact 3D points of the curve at the input Coordinates
-		 * The function can only be used with 3D curve (Dimension == 3)
-		 */
-		virtual void EvaluatePoints(const TArray<double>& Coordinates, TArray<FCurvePoint>& OutPoints, int32 DerivativeOrder = 0) const;
-
-		/**
-		 * Evaluate exact 3D points of the curve at the input Coordinates
-		 * The function can only be used with 3D curve (Dimension == 3)
-		 */
-		virtual void EvaluatePoints(const TArray<double>& Coordinates, TArray<FPoint>& OutPoints) const;
-
-		/**
-		 * Evaluate exact 2D points of the curve at the input Coordinates
-		 * The function can only be used with 2D curve (Dimension == 2)
-		 */
-		virtual void Evaluate2DPoints(const TArray<double>& Coordinates, TArray<FPoint2D>& OutPoints) const;
-
-		/**
-		 * Evaluate exact 2D points of the curve at the input Coordinates
-		 * The function can only be used with 2D curve (Dimension == 2)
-		 */
-		virtual void Evaluate2DPoints(const TArray<double>& Coordinates, TArray<FCurvePoint2D>& OutPoints, int32 DerivativeOrder = 0) const;
-
-		void FindNotDerivableCoordinates(int32 DerivativeOrder, TArray<double>& OutNotDerivableCoordinates) const;
-		virtual void FindNotDerivableCoordinates(const FLinearBoundary& InBoundary, int32 DerivativeOrder, TArray<double>& OutNotDerivableCoordinates) const;
-
-
-		/**
-		 * Generate a pre-sampling of the curve saved in OutCoordinates.
-		 * This sampling is light enough to allow a fast computation, precise enough to compute accurately meshing criteria
-		 */
-		void Presample(TArray<double>& OutSampling, double Tolerance) const
-		{
-			Presample(Boundary, Tolerance, OutSampling);
-		}
+	double GetUMax() const
+	{
+		return Boundary.Max;
+	};
 
-		virtual void Presample(const FLinearBoundary& InBoundary, double Tolerance, TArray<double>& OutSampling) const;
+	const FLinearBoundary& GetBoundary() const
+	{
+		return Boundary;
+	}
 
-		/**
-		 * Make a new curve based on the new bounds.
-		 * If the new bound is nearly equal to the initial bound, no curve is made (return TSharedPtr<FCurve>())
-		 */
-		virtual TSharedPtr<FCurve> MakeBoundedCurve(const FLinearBoundary& InBoundary);
+	virtual TSharedPtr<FEntityGeom> ApplyMatrix(const FMatrixH& InMatrix) const override = 0;
 
-		/**
-		 * Rebound the curve, if not possible make a new curve based on the new bounds.
-		 */
-		virtual TSharedPtr<FCurve> Rebound(const FLinearBoundary& InBoundary);
+	double GetLength(double Tolerance) const;
 
-		/**
-		 * Linear deformation of the curve along the axis [Start point, End point] so that the nearest extremity is at the desired position and the other is not modified
-		 */
-		virtual void ExtendTo(const FPoint& DesiredPosition)
-		{
-			ensure(Dimension == 3);
+	/**
+	 * Evaluate exact 3D point of the curve at the input Coordinate
+	 * The function can only be used with 3D curve (Dimension == 3)
+	 */
+	virtual void EvaluatePoint(double Coordinate, FCurvePoint& OutPoint, int32 DerivativeOrder = 0) const
+	{
+		ensureCADKernel(false);
+	}
+
+	/**
+	 * Evaluate exact 3D point of the curve at the input Coordinate
+	 * The function can only be used with 3D curve (Dimension == 3)
+	 */
+	virtual FPoint EvaluatePoint(double Coordinate) const
+	{
+		FCurvePoint OutPoint;
+		EvaluatePoint(Coordinate, OutPoint);
+		return OutPoint.Point;
+	}
+
+	/**
+	 * Evaluate exact 2D point of the curve at the input Coordinate
+	 * The function can only be used with 2D curve (Dimension == 2)
+	 */
+	virtual void Evaluate2DPoint(double Coordinate, FCurvePoint2D& OutPoint, int32 DerivativeOrder = 0) const
+	{
+		ensureCADKernel(false);
+	}
+
+	/**
+	 * Evaluate exact 3D point of the curve at the input Coordinate
+	 * The function can only be used with 3D curve (Dimension == 3)
+	 */
+	virtual FPoint2D Evaluate2DPoint(double Coordinate) const
+	{
+		FCurvePoint2D OutPoint;
+		Evaluate2DPoint(Coordinate, OutPoint);
+		return OutPoint.Point;
+	}
+
+	/**
+	 * Evaluate exact 2D point of the curve at the input Coordinate
+	 * The function can only be used with 2D curve (Dimension == 2)
+	 */
+	virtual void Evaluate2DPoint(double Coordinate, FPoint2D& OutPoint) const
+	{
+		ensure(Dimension == 2);
+		FCurvePoint2D CurvePoint;
+		Evaluate2DPoint(Coordinate, CurvePoint);
+		OutPoint = CurvePoint.Point;
+	}
+
+	/**
+	 * Evaluate exact 3D points of the curve at the input Coordinates
+	 * The function can only be used with 3D curve (Dimension == 3)
+	 */
+	virtual void EvaluatePoints(const TArray<double>& Coordinates, TArray<FCurvePoint>& OutPoints, int32 DerivativeOrder = 0) const;
+
+	/**
+	 * Evaluate exact 3D points of the curve at the input Coordinates
+	 * The function can only be used with 3D curve (Dimension == 3)
+	 */
+	virtual void EvaluatePoints(const TArray<double>& Coordinates, TArray<FPoint>& OutPoints) const;
+
+	/**
+	 * Evaluate exact 2D points of the curve at the input Coordinates
+	 * The function can only be used with 2D curve (Dimension == 2)
+	 */
+	virtual void Evaluate2DPoints(const TArray<double>& Coordinates, TArray<FPoint2D>& OutPoints) const;
+
+	/**
+	 * Evaluate exact 2D points of the curve at the input Coordinates
+	 * The function can only be used with 2D curve (Dimension == 2)
+	 */
+	virtual void Evaluate2DPoints(const TArray<double>& Coordinates, TArray<FCurvePoint2D>& OutPoints, int32 DerivativeOrder = 0) const;
+
+	void FindNotDerivableCoordinates(int32 DerivativeOrder, TArray<double>& OutNotDerivableCoordinates) const;
+	virtual void FindNotDerivableCoordinates(const FLinearBoundary& InBoundary, int32 DerivativeOrder, TArray<double>& OutNotDerivableCoordinates) const;
+
+
+	/**
+	 * Generate a pre-sampling of the curve saved in OutCoordinates.
+	 * This sampling is light enough to allow a fast computation, precise enough to compute accurately meshing criteria
+	 */
+	void Presample(TArray<double>& OutSampling, double Tolerance) const
+	{
+		Presample(Boundary, Tolerance, OutSampling);
+	}
+
+	virtual void Presample(const FLinearBoundary& InBoundary, double Tolerance, TArray<double>& OutSampling) const;
+
+	/**
+	 * Make a new curve based on the new bounds.
+	 * If the new bound is nearly equal to the initial bound, no curve is made (return TSharedPtr<FCurve>())
+	 */
+	virtual TSharedPtr<FCurve> MakeBoundedCurve(const FLinearBoundary& InBoundary);
+
+	/**
+	 * Rebound the curve, if not possible make a new curve based on the new bounds.
+	 */
+	virtual TSharedPtr<FCurve> Rebound(const FLinearBoundary& InBoundary);
+
+	/**
+	 * Linear deformation of the curve along the axis [Start point, End point] so that the nearest extremity is at the desired position and the other is not modified
+	 */
+	virtual void ExtendTo(const FPoint& DesiredPosition)
+	{
+		ensure(Dimension == 3);
 
-			NOT_IMPLEMENTED;
-			ensureCADKernel(false);
-		}
+		NOT_IMPLEMENTED;
+		ensureCADKernel(false);
+	}
 
-		/**
-		 * Linear deformation of the curve along the axis [Start point, End point] so that the nearest extremity is at the desired position and the other is not modified
-		 */
-		virtual void ExtendTo(const FPoint2D& DesiredPosition)
-		{
-			ensure(Dimension == 2);
-			FPoint Point = DesiredPosition;
-			ExtendTo(Point);
-		}
+	/**
+	 * Linear deformation of the curve along the axis [Start point, End point] so that the nearest extremity is at the desired position and the other is not modified
+	 */
+	virtual void ExtendTo(const FPoint2D& DesiredPosition)
+	{
+		ensure(Dimension == 2);
+		FPoint Point = DesiredPosition;
+		ExtendTo(Point);
+	}
 
-	protected:
+protected:
 
-		virtual double ComputeLength(const FLinearBoundary& InBoundary, double Tolerance) const;
-		virtual double ComputeLength2D(const FLinearBoundary& InBoundary, double Tolerance) const;
-	};
-} // namespace CADKernel
+	virtual double ComputeLength(const FLinearBoundary& InBoundary, double Tolerance) const;
+	virtual double ComputeLength2D(const FLinearBoundary& InBoundary, double Tolerance) const;
+};
+} // namespace UE::CADKernel
 
diff --git a/Engine/Source/Runtime/Datasmith/CADKernel/Public/CADKernel/Geo/Curves/NURBSCurveData.h b/Engine/Source/Runtime/Datasmith/CADKernel/Public/CADKernel/Geo/Curves/NURBSCurveData.h
--- a/Engine/Source/Runtime/Datasmith/CADKernel/Public/CADKernel/Geo/Curves/NURBSCurveData.h
+++ b/Engine/Source/Runtime/Datasmith/CADKernel/Public/CADKernel/Geo/Curves/NURBSCurveData.h
@@ -4,18 +4,18 @@
 #include "CADKernel/Core/Types.h"
 #include "CADKernel/Math/Point.h"
 
-namespace CADKernel
+namespace UE::CADKernel
 {
-	struct FNurbsCurveData
-	{
-		bool bIsRational;
-		int32 Dimension = 0;
+struct FNurbsCurveData
+{
+	bool bIsRational;
+	int32 Dimension = 0;
 
-		int32 Degree = 0;
-		TArray<double> NodalVector;
+	int32 Degree = 0;
+	TArray<double> NodalVector;
 
-		TArray<double> Weights;
-		TArray<FPoint> Poles;
-	};
-} // namespace CADKernel
+	TArray<double> Weights;
+	TArray<FPoint> Poles;
+};
+} // namespace UE::CADKernel
 
diff --git a/Engine/Source/Runtime/Datasmith/CADKernel/Public/CADKernel/Geo/Curves/RestrictionCurve.h b/Engine/Source/Runtime/Datasmith/CADKernel/Public/CADKernel/Geo/Curves/RestrictionCurve.h
--- a/Engine/Source/Runtime/Datasmith/CADKernel/Public/CADKernel/Geo/Curves/RestrictionCurve.h
+++ b/Engine/Source/Runtime/Datasmith/CADKernel/Public/CADKernel/Geo/Curves/RestrictionCurve.h
@@ -12,229 +12,229 @@
 #include "CADKernel/Math/MatrixH.h"
 #include "CADKernel/Math/Point.h"
 
-namespace CADKernel
+namespace UE::CADKernel
 {
-	class FCurve;
-	class FEntityGeom;
-
-	/**
-	 * A restriction curve is the curve carrying an edge
-	 *
-	 * It's defined by
-	 * - A surfacic curve defined by a 2D curve (@see Curve2D) and the carrier surface (@see CarrierSurface) of the TopologicalFace containing the edge:
-	 * - A linear approximation of the surfacic curve (@see Polyline) respecting the System Geometrical Tolerance (@see FKernelParameters::GeometricalTolerance)
-	 *   The linear approximation is:
-	 *      - an array of increasing coordinates (@see Polyline::Coordinates)
-	 *      - an array of 2D Points relating to coordinates: points of curve in the parametric space of the carrier surface (@see Polyline::Points2D)
-	 *      - an array of 3D Points relating to coordinates: 3D points of curve (@see Polyline::Points3D)
-	 *      - an array of Normal relating to coordinates: Surface's normal (@see Polyline::Normal)
-	 */
-	class CADKERNEL_API FRestrictionCurve : public FSurfacicCurve
-	{
-		friend class FEntity;
-		friend class FTopologicalEdge;
+class FCurve;
+class FEntityGeom;
+
+/**
+ * A restriction curve is the curve carrying an edge
+ *
+ * It's defined by
+ * - A surfacic curve defined by a 2D curve (@see Curve2D) and the carrier surface (@see CarrierSurface) of the TopologicalFace containing the edge:
+ * - A linear approximation of the surfacic curve (@see Polyline) respecting the System Geometrical Tolerance (@see FKernelParameters::GeometricalTolerance)
+ *   The linear approximation is:
+ *      - an array of increasing coordinates (@see Polyline::Coordinates)
+ *      - an array of 2D Points relating to coordinates: points of curve in the parametric space of the carrier surface (@see Polyline::Points2D)
+ *      - an array of 3D Points relating to coordinates: 3D points of curve (@see Polyline::Points3D)
+ *      - an array of Normal relating to coordinates: Surface's normal (@see Polyline::Normal)
+ */
+class CADKERNEL_API FRestrictionCurve : public FSurfacicCurve
+{
+	friend class FEntity;
+	friend class FTopologicalEdge;
 
-	protected:
+protected:
 
-		FSurfacicPolyline Polyline;
-		double MinLinearTolerance;
+	FSurfacicPolyline Polyline;
+	double MinLinearTolerance;
 
-		FRestrictionCurve(TSharedRef<FSurface> InCarrierSurface, TSharedRef<FCurve> InCurve2D)
-			: FSurfacicCurve(InCurve2D, InCarrierSurface)
-			, Polyline(InCarrierSurface, InCurve2D)
-		{
-			MinLinearTolerance = Boundary.ComputeMinimalTolerance();
-		}
+	FRestrictionCurve(TSharedRef<FSurface> InCarrierSurface, TSharedRef<FCurve> InCurve2D)
+		: FSurfacicCurve(InCurve2D, InCarrierSurface)
+		, Polyline(InCarrierSurface, InCurve2D)
+	{
+		MinLinearTolerance = Boundary.ComputeMinimalTolerance();
+	}
 
-		FRestrictionCurve() = default;
+	FRestrictionCurve() = default;
 
-	public:
+public:
 
-		virtual void Serialize(FCADKernelArchive& Ar) override
+	virtual void Serialize(FCADKernelArchive& Ar) override
+	{
+		FSurfacicCurve::Serialize(Ar);
+		Polyline.Serialize(Ar);
+		if (Ar.IsLoading())
 		{
-			FSurfacicCurve::Serialize(Ar);
-			Polyline.Serialize(Ar);
-			if (Ar.IsLoading())
-			{
-				MinLinearTolerance = Boundary.ComputeMinimalTolerance();
-			}
+			MinLinearTolerance = Boundary.ComputeMinimalTolerance();
 		}
+	}
 
 #ifdef CADKERNEL_DEV
-		virtual FInfoEntity& GetInfo(FInfoEntity&) const override;
+	virtual FInfoEntity& GetInfo(FInfoEntity&) const override;
 #endif
 
-		virtual ECurve GetCurveType() const override
-		{
-			return ECurve::Restriction;
-		}
+	virtual ECurve GetCurveType() const override
+	{
+		return ECurve::Restriction;
+	}
 
-		const TSharedRef<FCurve> Get2DCurve() const
-		{
-			return Curve2D.ToSharedRef();
-		}
+	const TSharedRef<FCurve> Get2DCurve() const
+	{
+		return Curve2D.ToSharedRef();
+	}
 
-		TSharedPtr<FEntityGeom> ApplyMatrix(const FMatrixH& InMatrix) const override
-		{
-			ensureCADKernel(false);
-			return TSharedPtr<FEntityGeom>();
-		}
+	TSharedPtr<FEntityGeom> ApplyMatrix(const FMatrixH& InMatrix) const override
+	{
+		ensureCADKernel(false);
+		return TSharedPtr<FEntityGeom>();
+	}
 
-		/**
-		 * Fast computation of the point in the parametric space of the carrier surface.
-		 * This approximation is based on FSurfacicPolyline::Polyline
-		 */
-		FPoint2D Approximate2DPoint(double InCoordinate) const
-		{
-			return Polyline.Approximate2DPoint(InCoordinate);
-		}
+	/**
+	 * Fast computation of the point in the parametric space of the carrier surface.
+	 * This approximation is based on FSurfacicPolyline::Polyline
+	 */
+	FPoint2D Approximate2DPoint(double InCoordinate) const
+	{
+		return Polyline.Approximate2DPoint(InCoordinate);
+	}
 
-		/**
-		 * Fast computation of the point in the parametric space of the carrier surface.
-		 * This approximation is based on FSurfacicPolyline::Polyline
-		 */
-		FPoint Approximate3DPoint(double InCoordinate) const
-		{
-			return Polyline.Approximate3DPoint(InCoordinate);
-		}
+	/**
+	 * Fast computation of the point in the parametric space of the carrier surface.
+	 * This approximation is based on FSurfacicPolyline::Polyline
+	 */
+	FPoint Approximate3DPoint(double InCoordinate) const
+	{
+		return Polyline.Approximate3DPoint(InCoordinate);
+	}
 
-		FPoint GetTangentAt(double InCoordinate) const
-		{
-			return Polyline.GetTangentAt(InCoordinate);
-		}
+	FPoint GetTangentAt(double InCoordinate) const
+	{
+		return Polyline.GetTangentAt(InCoordinate);
+	}
 
-		FPoint2D GetTangent2DAt(double InCoordinate) const
-		{
-			return Polyline.GetTangent2DAt(InCoordinate);
-		}
+	FPoint2D GetTangent2DAt(double InCoordinate) const
+	{
+		return Polyline.GetTangent2DAt(InCoordinate);
+	}
 
-		/**
-		 * Fast computation of the point in the parametric space of the carrier surface.
-		 * This approximation is based on FSurfacicPolyline::Polyline
-		 */
-		void Approximate2DPoints(const TArray<double>& InCoordinates, TArray<FPoint2D>& OutPoints) const
-		{
-			return Polyline.Approximate2DPoints(InCoordinates, OutPoints);
-		}
+	/**
+	 * Fast computation of the point in the parametric space of the carrier surface.
+	 * This approximation is based on FSurfacicPolyline::Polyline
+	 */
+	void Approximate2DPoints(const TArray<double>& InCoordinates, TArray<FPoint2D>& OutPoints) const
+	{
+		return Polyline.Approximate2DPoints(InCoordinates, OutPoints);
+	}
 
-		/**
-		 * Fast computation of the 3D point of the curve.
-		 * This approximation is based on FSurfacicPolyline::Polyline
-		 */
-		void Approximate3DPoints(const TArray<double>& InCoordinates, TArray<FPoint>& OutPoints) const
-		{
-			return Polyline.Approximate3DPoints(InCoordinates, OutPoints);
-		}
+	/**
+	 * Fast computation of the 3D point of the curve.
+	 * This approximation is based on FSurfacicPolyline::Polyline
+	 */
+	void Approximate3DPoints(const TArray<double>& InCoordinates, TArray<FPoint>& OutPoints) const
+	{
+		return Polyline.Approximate3DPoints(InCoordinates, OutPoints);
+	}
 
-		/**
-		 * Approximation of surfacic polyline (points 2d, 3d, normals, tangents) defined by its coordinates compute with carrier surface polyline
-		 */
-		void ApproximatePolyline(FSurfacicPolyline& OutPolyline) const
-		{
-			Polyline.ApproximatePolyline(OutPolyline);
-		}
+	/**
+	 * Approximation of surfacic polyline (points 2d, 3d, normals, tangents) defined by its coordinates compute with carrier surface polyline
+	 */
+	void ApproximatePolyline(FSurfacicPolyline& OutPolyline) const
+	{
+		Polyline.ApproximatePolyline(OutPolyline);
+	}
 
-		template<class PointType>
-		double GetCoordinateOfProjectedPoint(const FLinearBoundary& InBoundary, const PointType& PointOnEdge, PointType& ProjectedPoint) const
-		{
-			return Polyline.GetCoordinateOfProjectedPoint(InBoundary, PointOnEdge, ProjectedPoint);
-		}
+	template<class PointType>
+	double GetCoordinateOfProjectedPoint(const FLinearBoundary& InBoundary, const PointType& PointOnEdge, PointType& ProjectedPoint) const
+	{
+		return Polyline.GetCoordinateOfProjectedPoint(InBoundary, PointOnEdge, ProjectedPoint);
+	}
 
-		template<class PointType>
-		void ProjectPoints(const FLinearBoundary& InBoundary, const TArray<PointType>& InPointsToProject, TArray<double>& ProjectedPointCoordinates, TArray<PointType>& ProjectedPoints) const
-		{
-			Polyline.ProjectPoints(InBoundary, InPointsToProject, ProjectedPointCoordinates, ProjectedPoints);
-		}
+	template<class PointType>
+	void ProjectPoints(const FLinearBoundary& InBoundary, const TArray<PointType>& InPointsToProject, TArray<double>& ProjectedPointCoordinates, TArray<PointType>& ProjectedPoints) const
+	{
+		Polyline.ProjectPoints(InBoundary, InPointsToProject, ProjectedPointCoordinates, ProjectedPoints);
+	}
 
-		/**
-		 * Project a set of points of a twin curve on the 3D polyline and return the coordinate of the projected point
-		 */
-		void ProjectTwinCurvePoints(const TArray<FPoint>& InPointsToProject, bool bSameOrientation, TArray<double>& OutProjectedPointCoords) const
-		{
-			Polyline.ProjectCoincidentalPolyline(InPointsToProject, bSameOrientation, OutProjectedPointCoords);
-		}
+	/**
+	 * Project a set of points of a twin curve on the 3D polyline and return the coordinate of the projected point
+	 */
+	void ProjectTwinCurvePoints(const TArray<FPoint>& InPointsToProject, bool bSameOrientation, TArray<double>& OutProjectedPointCoords) const
+	{
+		Polyline.ProjectCoincidentalPolyline(InPointsToProject, bSameOrientation, OutProjectedPointCoords);
+	}
 
-		void ComputeIntersectionsWithIsos(const FLinearBoundary& InBoundary, const TArray<double>& InIsoCoordinates, const EIso InTypeIso, const FSurfacicTolerance& ToleranceIso, TArray<double>& OutIntersection) const
-		{
-			Polyline.ComputeIntersectionsWithIsos(InBoundary, InIsoCoordinates, InTypeIso, ToleranceIso, OutIntersection);
-		}
+	void ComputeIntersectionsWithIsos(const FLinearBoundary& InBoundary, const TArray<double>& InIsoCoordinates, const EIso InTypeIso, const FSurfacicTolerance& ToleranceIso, TArray<double>& OutIntersection) const
+	{
+		Polyline.ComputeIntersectionsWithIsos(InBoundary, InIsoCoordinates, InTypeIso, ToleranceIso, OutIntersection);
+	}
 
-		/**
-		 * A check is done to verify that:
-		 * - the curve is degenerated in the parametric space of the carrier surface i.e. the 2D length of the curve is not nearly equal to zero
-		 * - the curve is degenerated in 3D i.e. the 3D length of the curve is not nearly equal to zero
-		 *
-		 * A curve can be degenerated in 3D and not in 2D in the case of locally degenerated carrier surface.
-		 */
-		void CheckIfDegenerated(const FLinearBoundary& InBoundary, bool& bDegeneration2D, bool& bDegeneration3D, double& Length3D) const
-		{
-			double Tolerance = GetCarrierSurface()->Get3DTolerance();
-			const FSurfacicTolerance& Tolerances2D = GetCarrierSurface()->GetIsoTolerances();
-			Polyline.CheckIfDegenerated(Tolerance, Tolerances2D, InBoundary, bDegeneration2D, bDegeneration3D, Length3D);
-		}
+	/**
+	 * A check is done to verify that:
+	 * - the curve is degenerated in the parametric space of the carrier surface i.e. the 2D length of the curve is not nearly equal to zero
+	 * - the curve is degenerated in 3D i.e. the 3D length of the curve is not nearly equal to zero
+	 *
+	 * A curve can be degenerated in 3D and not in 2D in the case of locally degenerated carrier surface.
+	 */
+	void CheckIfDegenerated(const FLinearBoundary& InBoundary, bool& bDegeneration2D, bool& bDegeneration3D, double& Length3D) const
+	{
+		double Tolerance = GetCarrierSurface()->Get3DTolerance();
+		const FSurfacicTolerance& Tolerances2D = GetCarrierSurface()->GetIsoTolerances();
+		Polyline.CheckIfDegenerated(Tolerance, Tolerances2D, InBoundary, bDegeneration2D, bDegeneration3D, Length3D);
+	}
 
-		void GetExtremities(const FLinearBoundary& InBoundary, FSurfacicCurveExtremities& Extremities) const
-		{
-			double Tolerance = GetCarrierSurface()->Get3DTolerance();
-			const FSurfacicTolerance& Tolerances2D = GetCarrierSurface()->GetIsoTolerances();
-			Polyline.GetExtremities(InBoundary, Tolerance, Tolerances2D, Extremities);
-		}
+	void GetExtremities(const FLinearBoundary& InBoundary, FSurfacicCurveExtremities& Extremities) const
+	{
+		double Tolerance = GetCarrierSurface()->Get3DTolerance();
+		const FSurfacicTolerance& Tolerances2D = GetCarrierSurface()->GetIsoTolerances();
+		Polyline.GetExtremities(InBoundary, Tolerance, Tolerances2D, Extremities);
+	}
 
-		double GetToleranceAt(const double InCoordinate)
-		{
-			FDichotomyFinder Finder(Polyline.GetCoordinates());
-			const int32 Index = Finder.Find(InCoordinate);
-			const double Tolerance = GetCarrierSurface()->Get3DTolerance();
+	double GetToleranceAt(const double InCoordinate)
+	{
+		FDichotomyFinder Finder(Polyline.GetCoordinates());
+		const int32 Index = Finder.Find(InCoordinate);
+		const double Tolerance = GetCarrierSurface()->Get3DTolerance();
 
-			return Polyline.ComputeLinearToleranceAt(Tolerance, MinLinearTolerance, Index);
-		}
+		return Polyline.ComputeLinearToleranceAt(Tolerance, MinLinearTolerance, Index);
+	}
 
-		FPoint2D GetExtremityTolerances(const FLinearBoundary& InBoundary)
-		{
-			FDichotomyFinder Finder(Polyline.GetCoordinates());
-			const int32 StartIndex = Finder.Find(InBoundary.Min);
-			const int32 EndIndex = Finder.Find(InBoundary.Max);
+	FPoint2D GetExtremityTolerances(const FLinearBoundary& InBoundary)
+	{
+		FDichotomyFinder Finder(Polyline.GetCoordinates());
+		const int32 StartIndex = Finder.Find(InBoundary.Min);
+		const int32 EndIndex = Finder.Find(InBoundary.Max);
 
-			const double Tolerance = GetCarrierSurface()->Get3DTolerance();
+		const double Tolerance = GetCarrierSurface()->Get3DTolerance();
 
-			FPoint2D ExtremityTolerance;
-			ExtremityTolerance[0] = Polyline.ComputeLinearToleranceAt(Tolerance, MinLinearTolerance, StartIndex);
-			ExtremityTolerance[1] = Polyline.ComputeLinearToleranceAt(Tolerance, MinLinearTolerance, EndIndex);
-			return ExtremityTolerance;
-		}
+		FPoint2D ExtremityTolerance;
+		ExtremityTolerance[0] = Polyline.ComputeLinearToleranceAt(Tolerance, MinLinearTolerance, StartIndex);
+		ExtremityTolerance[1] = Polyline.ComputeLinearToleranceAt(Tolerance, MinLinearTolerance, EndIndex);
+		return ExtremityTolerance;
+	}
 
-		/**
-		 * @return the size of the polyline i.e. the count of points.
-		 */
-		int32 GetPolylineSize()
-		{
-			return Polyline.Size();
-		}
+	/**
+	 * @return the size of the polyline i.e. the count of points.
+	 */
+	int32 GetPolylineSize()
+	{
+		return Polyline.Size();
+	}
 
-		/**
-		 * Get the sub polyline bounded by the input InBoundary in the orientation of the input InOrientation and append it to the output OutPoints
-		 */
-		template<class PointType>
-		void GetDiscretizationPoints(const FLinearBoundary& InBoundary, EOrientation Orientation, TArray<PointType>& OutPoints) const
-		{
-			Polyline.GetSubPolyline(InBoundary, Orientation, OutPoints);
-		}
+	/**
+	 * Get the sub polyline bounded by the input InBoundary in the orientation of the input InOrientation and append it to the output OutPoints
+	 */
+	template<class PointType>
+	void GetDiscretizationPoints(const FLinearBoundary& InBoundary, EOrientation Orientation, TArray<PointType>& OutPoints) const
+	{
+		Polyline.GetSubPolyline(InBoundary, Orientation, OutPoints);
+	}
 
-		/**
-		 * Samples the sub curve limited by the boundary respecting the input Desired segment length
-		 */
-		void Sample(const FLinearBoundary& InBoundary, const double DesiredSegmentLength, TArray<double>& OutCoordinates) const
-		{
-			Polyline.Sample(InBoundary, DesiredSegmentLength, OutCoordinates);
-		}
+	/**
+	 * Samples the sub curve limited by the boundary respecting the input Desired segment length
+	 */
+	void Sample(const FLinearBoundary& InBoundary, const double DesiredSegmentLength, TArray<double>& OutCoordinates) const
+	{
+		Polyline.Sample(InBoundary, DesiredSegmentLength, OutCoordinates);
+	}
 
-		double ApproximateLength(const FLinearBoundary& InBoundary) const
-		{
-			return Polyline.GetLength(InBoundary);
-		}
+	double ApproximateLength(const FLinearBoundary& InBoundary) const
+	{
+		return Polyline.GetLength(InBoundary);
+	}
 
-		void ExtendTo(const FPoint2D& Point) override;
-	};
+	void ExtendTo(const FPoint2D& Point) override;
+};
 
-} // namespace CADKernel
+} // namespace UE::CADKernel
 
diff --git a/Engine/Source/Runtime/Datasmith/CADKernel/Public/CADKernel/Geo/Curves/SurfacicCurve.h b/Engine/Source/Runtime/Datasmith/CADKernel/Public/CADKernel/Geo/Curves/SurfacicCurve.h
--- a/Engine/Source/Runtime/Datasmith/CADKernel/Public/CADKernel/Geo/Curves/SurfacicCurve.h
+++ b/Engine/Source/Runtime/Datasmith/CADKernel/Public/CADKernel/Geo/Curves/SurfacicCurve.h
@@ -8,135 +8,135 @@
 #include "CADKernel/Geo/Surfaces/Surface.h"
 #include "CADKernel/Math/MatrixH.h"
 
-namespace CADKernel
+namespace UE::CADKernel
 {
-	class FCADKernelArchive;
-	class FDatabase;
-	class FEntityGeom;
-	class FPoint2D;
-	class FSurfacicPolyline;
-	struct FCurvePoint2D;
-	struct FCurvePoint;
-	struct FLinearBoundary;
-
-	class CADKERNEL_API FSurfacicCurve : public FCurve
-	{
-		friend class FEntity;
+class FCADKernelArchive;
+class FDatabase;
+class FEntityGeom;
+class FPoint2D;
+class FSurfacicPolyline;
+struct FCurvePoint2D;
+struct FCurvePoint;
+struct FLinearBoundary;
+
+class CADKERNEL_API FSurfacicCurve : public FCurve
+{
+	friend class FEntity;
 
-	protected:
+protected:
 
-		TSharedPtr<FCurve> Curve2D;
-		TSharedPtr<FSurface> CarrierSurface;
+	TSharedPtr<FCurve> Curve2D;
+	TSharedPtr<FSurface> CarrierSurface;
 
-		FSurfacicCurve(TSharedRef<FCurve> InCurve2D, TSharedRef<FSurface> InSurface)
-			: FCurve(InCurve2D->GetBoundary())
-			, Curve2D(InCurve2D)
-			, CarrierSurface(InSurface)
-		{
-			ensureCADKernel(Curve2D.IsValid());
-		}
+	FSurfacicCurve(TSharedRef<FCurve> InCurve2D, TSharedRef<FSurface> InSurface)
+		: FCurve(InCurve2D->GetBoundary())
+		, Curve2D(InCurve2D)
+		, CarrierSurface(InSurface)
+	{
+		ensureCADKernel(Curve2D.IsValid());
+	}
 
-		FSurfacicCurve() = default;
+	FSurfacicCurve() = default;
 
-	public:
+public:
 
-		virtual void Serialize(FCADKernelArchive& Ar) override
-		{
-			FCurve::Serialize(Ar);
-			SerializeIdent(Ar, Curve2D);
-			SerializeIdent(Ar, CarrierSurface);
-		}
+	virtual void Serialize(FCADKernelArchive& Ar) override
+	{
+		FCurve::Serialize(Ar);
+		SerializeIdent(Ar, Curve2D);
+		SerializeIdent(Ar, CarrierSurface);
+	}
 
-		virtual void SpawnIdent(FDatabase& Database) override
+	virtual void SpawnIdent(FDatabase& Database) override
+	{
+		if (!FEntity::SetId(Database))
 		{
-			if (!FEntity::SetId(Database))
-			{
-				return;
-			}
-
-			Curve2D->SpawnIdent(Database);
-			CarrierSurface->SpawnIdent(Database);
+			return;
 		}
 
-		virtual void ResetMarkersRecursively() override
-		{
-			ResetMarkers();
-			Curve2D->ResetMarkersRecursively();
-			CarrierSurface->ResetMarkersRecursively();
-		}
+		Curve2D->SpawnIdent(Database);
+		CarrierSurface->SpawnIdent(Database);
+	}
+
+	virtual void ResetMarkersRecursively() override
+	{
+		ResetMarkers();
+		Curve2D->ResetMarkersRecursively();
+		CarrierSurface->ResetMarkersRecursively();
+	}
 
 #ifdef CADKERNEL_DEV
-		virtual FInfoEntity& GetInfo(FInfoEntity&) const override;
+	virtual FInfoEntity& GetInfo(FInfoEntity&) const override;
 #endif
 
-		virtual ECurve GetCurveType() const override
-		{
-			return ECurve::Surfacic;
-		}
+	virtual ECurve GetCurveType() const override
+	{
+		return ECurve::Surfacic;
+	}
 
-		const TSharedPtr<FCurve>& Get2DCurve() const
-		{
-			return Curve2D;
-		}
+	const TSharedPtr<FCurve>& Get2DCurve() const
+	{
+		return Curve2D;
+	}
 
-		TSharedPtr<FCurve>& Get2DCurve()
-		{
-			return Curve2D;
-		}
+	TSharedPtr<FCurve>& Get2DCurve()
+	{
+		return Curve2D;
+	}
 
-		const TSharedRef<FSurface> GetCarrierSurface() const
-		{
-			return CarrierSurface.ToSharedRef();
-		}
+	const TSharedRef<FSurface> GetCarrierSurface() const
+	{
+		return CarrierSurface.ToSharedRef();
+	}
 
-		void Set2DCurve(TSharedPtr<FCurve>& NewCurve2D)
-		{
-			ensureCADKernel(NewCurve2D.IsValid());
-			Curve2D = NewCurve2D;
-		}
+	void Set2DCurve(TSharedPtr<FCurve>& NewCurve2D)
+	{
+		ensureCADKernel(NewCurve2D.IsValid());
+		Curve2D = NewCurve2D;
+	}
 
-		const TSharedPtr<FSurface>& GetSurface() const
-		{
-			return CarrierSurface;
-		}
+	const TSharedPtr<FSurface>& GetSurface() const
+	{
+		return CarrierSurface;
+	}
 
-		TSharedPtr<FSurface>& GetSurface()
-		{
-			return CarrierSurface;
-		}
+	TSharedPtr<FSurface>& GetSurface()
+	{
+		return CarrierSurface;
+	}
 
-		virtual TSharedPtr<FEntityGeom> ApplyMatrix(const FMatrixH& InMatrix) const override;
+	virtual TSharedPtr<FEntityGeom> ApplyMatrix(const FMatrixH& InMatrix) const override;
 
-		virtual void EvaluatePoint(double Coordinate, FCurvePoint& OutPoint, int32 DerivativeOrder = 0) const override;
+	virtual void EvaluatePoint(double Coordinate, FCurvePoint& OutPoint, int32 DerivativeOrder = 0) const override;
 
-		virtual void Evaluate2DPoint(double Coordinate, FCurvePoint2D& OutPoint, int32 DerivativeOrder = 0) const override
-		{
-			Curve2D->Evaluate2DPoint(Coordinate, OutPoint, DerivativeOrder);
-		}
+	virtual void Evaluate2DPoint(double Coordinate, FCurvePoint2D& OutPoint, int32 DerivativeOrder = 0) const override
+	{
+		Curve2D->Evaluate2DPoint(Coordinate, OutPoint, DerivativeOrder);
+	}
 
-		virtual void Evaluate2DPoint(double Coordinate, FPoint2D& OutPoint) const override
-		{
-			Curve2D->Evaluate2DPoint(Coordinate, OutPoint);
-		}
+	virtual void Evaluate2DPoint(double Coordinate, FPoint2D& OutPoint) const override
+	{
+		Curve2D->Evaluate2DPoint(Coordinate, OutPoint);
+	}
 
-		virtual void EvaluatePoints(const TArray<double>& Coordinates, TArray<FCurvePoint>& OutPoints, int32 DerivativeOrder = 0) const override;
+	virtual void EvaluatePoints(const TArray<double>& Coordinates, TArray<FCurvePoint>& OutPoints, int32 DerivativeOrder = 0) const override;
 
-		void EvaluateSurfacicPolyline(FSurfacicPolyline& OutPolyline) const;
+	void EvaluateSurfacicPolyline(FSurfacicPolyline& OutPolyline) const;
 
-		virtual void Evaluate2DPoints(const TArray<double>& Coordinates, TArray<FPoint2D>& OutPoints) const override
-		{
-			Curve2D->Evaluate2DPoints(Coordinates, OutPoints);
-		}
+	virtual void Evaluate2DPoints(const TArray<double>& Coordinates, TArray<FPoint2D>& OutPoints) const override
+	{
+		Curve2D->Evaluate2DPoints(Coordinates, OutPoints);
+	}
 
-		virtual void Evaluate2DPoints(const TArray<double>& Coordinates, TArray<FCurvePoint2D>& OutPoints, int32 DerivativeOrder = 0) const override
-		{
-			Curve2D->Evaluate2DPoints(Coordinates, OutPoints, DerivativeOrder);
-		}
+	virtual void Evaluate2DPoints(const TArray<double>& Coordinates, TArray<FCurvePoint2D>& OutPoints, int32 DerivativeOrder = 0) const override
+	{
+		Curve2D->Evaluate2DPoints(Coordinates, OutPoints, DerivativeOrder);
+	}
 
-		virtual void FindNotDerivableCoordinates(const FLinearBoundary& InBoundary, int32 DerivativeOrder, TArray<double>& OutNotDerivableCoordinates) const override;
+	virtual void FindNotDerivableCoordinates(const FLinearBoundary& InBoundary, int32 DerivativeOrder, TArray<double>& OutNotDerivableCoordinates) const override;
 private:
-		void EvaluateSurfacicPolylineWithNormalAndTangent(FSurfacicPolyline& OutPolyline) const;
-	};
+	void EvaluateSurfacicPolylineWithNormalAndTangent(FSurfacicPolyline& OutPolyline) const;
+};
 
-} // namespace CADKernel
+} // namespace UE::CADKernel
 
diff --git a/Engine/Source/Runtime/Datasmith/CADKernel/Public/CADKernel/Geo/GeoPoint.h b/Engine/Source/Runtime/Datasmith/CADKernel/Public/CADKernel/Geo/GeoPoint.h
--- a/Engine/Source/Runtime/Datasmith/CADKernel/Public/CADKernel/Geo/GeoPoint.h
+++ b/Engine/Source/Runtime/Datasmith/CADKernel/Public/CADKernel/Geo/GeoPoint.h
@@ -5,154 +5,154 @@
 #include "CADKernel/Geo/GeoEnum.h"
 #include "CADKernel/Math/Point.h"
 
-namespace CADKernel
+namespace UE::CADKernel
 {
-	struct FCurvePoint2D
-	{
-		int32 DerivativeOrder = -1;
-		FPoint2D Point = FPoint2D::ZeroPoint;
-		FPoint2D Gradient = FPoint2D::ZeroPoint;
-		FPoint2D Laplacian = FPoint2D::ZeroPoint;
-	};
+struct FCurvePoint2D
+{
+	int32 DerivativeOrder = -1;
+	FPoint2D Point = FPoint2D::ZeroPoint;
+	FPoint2D Gradient = FPoint2D::ZeroPoint;
+	FPoint2D Laplacian = FPoint2D::ZeroPoint;
+};
 
-	struct FSurfacicPoint
-	{
-		int32 DerivativeOrder = -1;
-		FPoint Point = FPoint::ZeroPoint;
-		FPoint GradientU = FPoint::ZeroPoint;
-		FPoint GradientV = FPoint::ZeroPoint;
-		FPoint LaplacianU = FPoint::ZeroPoint;
-		FPoint LaplacianV = FPoint::ZeroPoint;
-		FPoint LaplacianUV = FPoint::ZeroPoint;
-	};
-
-	struct FSurfacicCurvePoint
-	{
-		bool bWithNormals;
-		bool bWithTangent;
+struct FSurfacicPoint
+{
+	int32 DerivativeOrder = -1;
+	FPoint Point = FPoint::ZeroPoint;
+	FPoint GradientU = FPoint::ZeroPoint;
+	FPoint GradientV = FPoint::ZeroPoint;
+	FPoint LaplacianU = FPoint::ZeroPoint;
+	FPoint LaplacianV = FPoint::ZeroPoint;
+	FPoint LaplacianUV = FPoint::ZeroPoint;
+};
+
+struct FSurfacicCurvePoint
+{
+	bool bWithNormals;
+	bool bWithTangent;
 
-		FPoint2D Point2D = FPoint2D::ZeroPoint;
-		FPoint Point = FPoint::ZeroPoint;
-		FPoint Normal = FPoint::ZeroPoint;
-		FPoint Tangent = FPoint::ZeroPoint;
-	};
+	FPoint2D Point2D = FPoint2D::ZeroPoint;
+	FPoint Point = FPoint::ZeroPoint;
+	FPoint Normal = FPoint::ZeroPoint;
+	FPoint Tangent = FPoint::ZeroPoint;
+};
 
-	struct FSurfacicCurvePointWithTolerance
-	{
-		FPoint2D Point2D = FPoint2D::ZeroPoint;
-		FPoint Point = FPoint::ZeroPoint;
-		FSurfacicTolerance Tolerance = FPoint2D::FarawayPoint;
-	};
+struct FSurfacicCurvePointWithTolerance
+{
+	FPoint2D Point2D = FPoint2D::ZeroPoint;
+	FPoint Point = FPoint::ZeroPoint;
+	FSurfacicTolerance Tolerance = FPoint2D::FarawayPoint;
+};
+
+typedef FSurfacicCurvePointWithTolerance FSurfacicCurveExtremities[2];
 
-	typedef FSurfacicCurvePointWithTolerance FSurfacicCurveExtremities[2];
+struct FCurvePoint
+{
+	int32 DerivativeOrder = -1;
+	FPoint Point = FPoint::ZeroPoint;
+	FPoint Gradient = FPoint::ZeroPoint;
+	FPoint Laplacian = FPoint::ZeroPoint;
+
+	FCurvePoint() = default;
 
-	struct FCurvePoint
+	FCurvePoint(FPoint InPoint)
+		: Point(InPoint)
 	{
-		int32 DerivativeOrder = -1;
-		FPoint Point = FPoint::ZeroPoint;
-		FPoint Gradient = FPoint::ZeroPoint;
-		FPoint Laplacian = FPoint::ZeroPoint;
+		Gradient = FPoint::ZeroPoint;
+		Laplacian = FPoint::ZeroPoint;
+	}
 
-		FCurvePoint() = default;
+	void Init()
+	{
+		Point = FPoint::ZeroPoint;
+		Gradient = FPoint::ZeroPoint;
+		Laplacian = FPoint::ZeroPoint;
+	}
+
+	/**
+	 * Compute the 3D surface curve point property (3D Coordinate, Gradient, Laplacian) according to
+	 * its 2D curve point property and the 3D surface point property
+	 */
+	void Combine(const FCurvePoint2D& Point2D, const FSurfacicPoint& SurfacicPoint)
+	{
+		ensureCADKernel(Point2D.DerivativeOrder >= 0);
+		ensureCADKernel(SurfacicPoint.DerivativeOrder >= 0);
 
-		FCurvePoint(FPoint InPoint)
-			: Point(InPoint)
-		{
- 			Gradient = FPoint::ZeroPoint;
-			Laplacian = FPoint::ZeroPoint;
-		}
+		ensureCADKernel(Point2D.DerivativeOrder <= SurfacicPoint.DerivativeOrder);
 
-		void Init()
+		DerivativeOrder = Point2D.DerivativeOrder;
+		Point = SurfacicPoint.Point;
+
+		if (DerivativeOrder > 0)
 		{
-			Point = FPoint::ZeroPoint;
-			Gradient = FPoint::ZeroPoint;
-			Laplacian = FPoint::ZeroPoint;
+			Gradient = SurfacicPoint.GradientU * Point2D.Gradient.U + SurfacicPoint.GradientV * Point2D.Gradient.V;
 		}
 
-		/**
-		 * Compute the 3D surface curve point property (3D Coordinate, Gradient, Laplacian) according to 
-		 * its 2D curve point property and the 3D surface point property
-		 */
-		void Combine(const FCurvePoint2D& Point2D, const FSurfacicPoint& SurfacicPoint)
+		if (DerivativeOrder > 1)
 		{
-			ensureCADKernel(Point2D.DerivativeOrder >= 0);
-			ensureCADKernel(SurfacicPoint.DerivativeOrder >= 0);
-
-			ensureCADKernel(Point2D.DerivativeOrder <= SurfacicPoint.DerivativeOrder);
-
-			DerivativeOrder = Point2D.DerivativeOrder;
-			Point = SurfacicPoint.Point;
-
-			if (DerivativeOrder > 0)
-			{
-				Gradient = SurfacicPoint.GradientU * Point2D.Gradient.U + SurfacicPoint.GradientV * Point2D.Gradient.V;
-			}
-
-			if (DerivativeOrder > 1) 
-			{
-				Laplacian = SurfacicPoint.LaplacianU * FMath::Square(Point2D.Gradient.U)
-					+ 2.0 * SurfacicPoint.LaplacianUV * Point2D.Gradient.U * Point2D.Gradient.V
-					+ SurfacicPoint.LaplacianV * FMath::Square(Point2D.Gradient.V)
-					+ SurfacicPoint.GradientU * Point2D.Laplacian.U 
-					+ SurfacicPoint.GradientV * Point2D.Laplacian.V;
-			}
+			Laplacian = SurfacicPoint.LaplacianU * FMath::Square(Point2D.Gradient.U)
+				+ 2.0 * SurfacicPoint.LaplacianUV * Point2D.Gradient.U * Point2D.Gradient.V
+				+ SurfacicPoint.LaplacianV * FMath::Square(Point2D.Gradient.V)
+				+ SurfacicPoint.GradientU * Point2D.Laplacian.U
+				+ SurfacicPoint.GradientV * Point2D.Laplacian.V;
 		}
-	};
+	}
+};
 
-	struct FCoordinateGrid
-	{
-		TArray<double> Coordinates[2];
+struct FCoordinateGrid
+{
+	TArray<double> Coordinates[2];
 
-		FCoordinateGrid()
-		{
-		}
+	FCoordinateGrid()
+	{
+	}
 
-		FCoordinateGrid(const TArray<double>& InUCoordinates, const TArray<double>& InVCoordinates)
-		{
-			Coordinates[EIso::IsoU] = InUCoordinates;
-			Coordinates[EIso::IsoV] = InVCoordinates;
-		}
+	FCoordinateGrid(const TArray<double>& InUCoordinates, const TArray<double>& InVCoordinates)
+	{
+		Coordinates[EIso::IsoU] = InUCoordinates;
+		Coordinates[EIso::IsoV] = InVCoordinates;
+	}
 
-		void Swap(TArray<double>& InUCoordinates, TArray<double>& InVCoordinates)
-		{
-			::Swap(Coordinates[EIso::IsoU], InUCoordinates);
-			::Swap(Coordinates[EIso::IsoV], InVCoordinates);
-		}
+	void Swap(TArray<double>& InUCoordinates, TArray<double>& InVCoordinates)
+	{
+		::Swap(Coordinates[EIso::IsoU], InUCoordinates);
+		::Swap(Coordinates[EIso::IsoV], InVCoordinates);
+	}
 
-		int32 Count() const
-		{
-			return Coordinates[0].Num() * Coordinates[1].Num();
-		}
+	int32 Count() const
+	{
+		return Coordinates[0].Num() * Coordinates[1].Num();
+	}
 
-		int32 IsoCount(EIso Iso) const
-		{
-			return Coordinates[Iso].Num();
-		}
+	int32 IsoCount(EIso Iso) const
+	{
+		return Coordinates[Iso].Num();
+	}
 
-		void SetNum(int32 UNumber, int32 VNumber)
-		{
-			Coordinates[0].SetNum(UNumber);
-			Coordinates[1].SetNum(VNumber);
-		}
+	void SetNum(int32 UNumber, int32 VNumber)
+	{
+		Coordinates[0].SetNum(UNumber);
+		Coordinates[1].SetNum(VNumber);
+	}
 
-		void Empty(int32 UNumber = 0, int32 VNumber = 0)
-		{
-			Coordinates[0].Empty(UNumber);
-			Coordinates[1].Empty(VNumber);
-		}
+	void Empty(int32 UNumber = 0, int32 VNumber = 0)
+	{
+		Coordinates[0].Empty(UNumber);
+		Coordinates[1].Empty(VNumber);
+	}
 
-		constexpr TArray<double>& operator[](EIso Iso)
-		{
-			ensureCADKernel(Iso == 0 || Iso == 1);
-			return Coordinates[Iso];
-		}
+	constexpr TArray<double>& operator[](EIso Iso)
+	{
+		ensureCADKernel(Iso == 0 || Iso == 1);
+		return Coordinates[Iso];
+	}
 
-		constexpr const TArray<double>& operator[](EIso Iso) const
-		{
-			ensureCADKernel(Iso == 0 || Iso == 1);
-			return Coordinates[Iso];
-		}
-	};
+	constexpr const TArray<double>& operator[](EIso Iso) const
+	{
+		ensureCADKernel(Iso == 0 || Iso == 1);
+		return Coordinates[Iso];
+	}
+};
 
 
 }
\ No newline at end of file
diff --git a/Engine/Source/Runtime/Datasmith/CADKernel/Public/CADKernel/Geo/Sampling/Polyline.h b/Engine/Source/Runtime/Datasmith/CADKernel/Public/CADKernel/Geo/Sampling/Polyline.h
--- a/Engine/Source/Runtime/Datasmith/CADKernel/Public/CADKernel/Geo/Sampling/Polyline.h
+++ b/Engine/Source/Runtime/Datasmith/CADKernel/Public/CADKernel/Geo/Sampling/Polyline.h
@@ -9,7 +9,7 @@
 
 #include "Serialization/Archive.h"
 
-namespace CADKernel
+namespace UE::CADKernel
 {
 
 template<class PointType>
@@ -214,5 +214,5 @@ public:
 	}
 };
 
-} // ns CADKernel
+} // ns UE::CADKernel
 
diff --git a/Engine/Source/Runtime/Datasmith/CADKernel/Public/CADKernel/Geo/Sampling/PolylineTools.h b/Engine/Source/Runtime/Datasmith/CADKernel/Public/CADKernel/Geo/Sampling/PolylineTools.h
--- a/Engine/Source/Runtime/Datasmith/CADKernel/Public/CADKernel/Geo/Sampling/PolylineTools.h
+++ b/Engine/Source/Runtime/Datasmith/CADKernel/Public/CADKernel/Geo/Sampling/PolylineTools.h
@@ -13,7 +13,7 @@
 #include "Algo/ForEach.h"
 #include "Algo/Reverse.h"
 
-namespace CADKernel
+namespace UE::CADKernel
 {
 
 struct FPolylineBBox
@@ -623,7 +623,7 @@ public:
 		int32 BoundaryIndices[2];
 		GetStartEndIndex(InBoundary, BoundaryIndices);
 
-		if(BoundaryIndices[1] - BoundaryIndices[0] > 0)
+		if (BoundaryIndices[1] - BoundaryIndices[0] > 0)
 		{
 			if (FMath::IsNearlyEqual(PolylineCoordinates[BoundaryIndices[0] + 1], InBoundary.Min, (double)DOUBLE_SMALL_NUMBER))
 			{
@@ -701,5 +701,5 @@ public:
 	}
 };
 
-} // namespace CADKernel
+} // namespace UE::CADKernel
 
diff --git a/Engine/Source/Runtime/Datasmith/CADKernel/Public/CADKernel/Geo/Surfaces/Surface.h b/Engine/Source/Runtime/Datasmith/CADKernel/Public/CADKernel/Geo/Surfaces/Surface.h
--- a/Engine/Source/Runtime/Datasmith/CADKernel/Public/CADKernel/Geo/Surfaces/Surface.h
+++ b/Engine/Source/Runtime/Datasmith/CADKernel/Public/CADKernel/Geo/Surfaces/Surface.h
@@ -12,271 +12,271 @@
 #include "CADKernel/Math/Point.h"
 #include "CADKernel/Utils/Cache.h"
 
-namespace CADKernel
+namespace UE::CADKernel
 {
-	class FGrid;
-	class FInfoEntity;
-	class FSurfacicPolyline;
-	struct FCurvePoint2D;
-	struct FCurvePoint;
-	struct FSurfacicSampling;
-	struct FNurbsSurfaceData;
-	struct FNurbsSurfaceHomogeneousData;
-
-	class CADKERNEL_API FSurface : public FEntityGeom
-	{
-		friend FEntity;
+class FGrid;
+class FInfoEntity;
+class FSurfacicPolyline;
+struct FCurvePoint2D;
+struct FCurvePoint;
+struct FSurfacicSampling;
+struct FNurbsSurfaceData;
+struct FNurbsSurfaceHomogeneousData;
+
+class CADKERNEL_API FSurface : public FEntityGeom
+{
+	friend FEntity;
 
-	protected:
+protected:
 
-		FSurfacicBoundary Boundary;
-		double Tolerance3D;
+	FSurfacicBoundary Boundary;
+	double Tolerance3D;
 
-		/**
-		 * Tolerance along Iso U/V is very costly to compute and not accurate.
-		 * A first approximation is based on the surface boundary along U and along V
-		 * Indeed, defining a Tolerance2D has no sense has the boundary length along an Iso could be very very huge compare to the boundary length along the other Iso like [[0, 1000] [0, 1]]
-		 * The tolerance along an iso is the length of the boundary along this iso divided by 100 000: if the curve length in 3d is 10m, the tolerance is 0.01mm
-		 * In the used, a local tolerance has to be estimated
-		 */
-		TCache<FSurfacicTolerance> MinToleranceIso;
+	/**
+	 * Tolerance along Iso U/V is very costly to compute and not accurate.
+	 * A first approximation is based on the surface boundary along U and along V
+	 * Indeed, defining a Tolerance2D has no sense has the boundary length along an Iso could be very very huge compare to the boundary length along the other Iso like [[0, 1000] [0, 1]]
+	 * The tolerance along an iso is the length of the boundary along this iso divided by 100 000: if the curve length in 3d is 10m, the tolerance is 0.01mm
+	 * In the used, a local tolerance has to be estimated
+	 */
+	TCache<FSurfacicTolerance> MinToleranceIso;
 
-	private:
+private:
 
-		virtual void InitBoundary()
-		{
-			Boundary.Set();
-		}
+	virtual void InitBoundary()
+	{
+		Boundary.Set();
+	}
 
-	protected:
+protected:
 
-		FSurface() = default;
+	FSurface() = default;
 
-		FSurface(double InToleranceGeometric)
-			: FEntityGeom()
-			, Tolerance3D(InToleranceGeometric)
-		{
+	FSurface(double InToleranceGeometric)
+		: FEntityGeom()
+		, Tolerance3D(InToleranceGeometric)
+	{
 
-		}
+	}
 
-		FSurface(double InToleranceGeometric, const FSurfacicBoundary& InBoundary)
-			: FEntityGeom()
-			, Boundary(InBoundary)
-			, Tolerance3D(InToleranceGeometric)
-		{
-			Boundary.IsValid();
-		}
+	FSurface(double InToleranceGeometric, const FSurfacicBoundary& InBoundary)
+		: FEntityGeom()
+		, Boundary(InBoundary)
+		, Tolerance3D(InToleranceGeometric)
+	{
+		Boundary.IsValid();
+	}
 
-		FSurface(double InToleranceGeometric, double UMin, double UMax, double VMin, double VMax)
-			: FEntityGeom()
-			, Tolerance3D(InToleranceGeometric)
-		{
-			Boundary.Set(UMin, UMax, VMin, VMax);
-		}		
+	FSurface(double InToleranceGeometric, double UMin, double UMax, double VMin, double VMax)
+		: FEntityGeom()
+		, Tolerance3D(InToleranceGeometric)
+	{
+		Boundary.Set(UMin, UMax, VMin, VMax);
+	}
 
-		void ComputeDefaultMinToleranceIso()
-		{
-			MinToleranceIso.Set(Boundary[EIso::IsoU].ComputeMinimalTolerance(), Boundary[EIso::IsoV].ComputeMinimalTolerance());
-		}
+	void ComputeDefaultMinToleranceIso()
+	{
+		MinToleranceIso.Set(Boundary[EIso::IsoU].ComputeMinimalTolerance(), Boundary[EIso::IsoV].ComputeMinimalTolerance());
+	}
 
-	public:
+public:
 
-		static TSharedPtr<FSurface> MakeBezierSurface(const double InToleranceGeometric, int32 InUDegre, int32 InVDegre, const TArray<FPoint>& InPoles);
-		static TSharedPtr<FSurface> MakeConeSurface(const double InToleranceGeometric, const FMatrixH& InMatrix, double InStartRadius, double InConeAngle, const FSurfacicBoundary& InBoundary);
-		static TSharedPtr<FSurface> MakeCylinderSurface(const double InToleranceGeometric, const FMatrixH& InMatrix, const double InRadius, const FSurfacicBoundary& InBoundary);
-		static TSharedPtr<FSurface> MakeNurbsSurface(const double InToleranceGeometric, const FNurbsSurfaceData& NurbsData);
-		static TSharedPtr<FSurface> MakeNurbsSurface(const double InToleranceGeometric, const FNurbsSurfaceHomogeneousData& NurbsData);
-		static TSharedPtr<FSurface> MakePlaneSurface(const double InToleranceGeometric, const FMatrixH& InMatrix, const FSurfacicBoundary& InBoundary);
-		static TSharedPtr<FSurface> MakeSphericalSurface(const double InToleranceGeometric, const FMatrixH& InMatrix, double InRadius, const FSurfacicBoundary& InBoundary);
-		static TSharedPtr<FSurface> MakeTorusSurface(const double InToleranceGeometric, const FMatrixH& InMatrix, double InMajorRadius, double InMinorRadius, const FSurfacicBoundary& InBoundary);
+	static TSharedPtr<FSurface> MakeBezierSurface(const double InToleranceGeometric, int32 InUDegre, int32 InVDegre, const TArray<FPoint>& InPoles);
+	static TSharedPtr<FSurface> MakeConeSurface(const double InToleranceGeometric, const FMatrixH& InMatrix, double InStartRadius, double InConeAngle, const FSurfacicBoundary& InBoundary);
+	static TSharedPtr<FSurface> MakeCylinderSurface(const double InToleranceGeometric, const FMatrixH& InMatrix, const double InRadius, const FSurfacicBoundary& InBoundary);
+	static TSharedPtr<FSurface> MakeNurbsSurface(const double InToleranceGeometric, const FNurbsSurfaceData& NurbsData);
+	static TSharedPtr<FSurface> MakeNurbsSurface(const double InToleranceGeometric, const FNurbsSurfaceHomogeneousData& NurbsData);
+	static TSharedPtr<FSurface> MakePlaneSurface(const double InToleranceGeometric, const FMatrixH& InMatrix, const FSurfacicBoundary& InBoundary);
+	static TSharedPtr<FSurface> MakeSphericalSurface(const double InToleranceGeometric, const FMatrixH& InMatrix, double InRadius, const FSurfacicBoundary& InBoundary);
+	static TSharedPtr<FSurface> MakeTorusSurface(const double InToleranceGeometric, const FMatrixH& InMatrix, double InMajorRadius, double InMinorRadius, const FSurfacicBoundary& InBoundary);
 
-		virtual void Serialize(FCADKernelArchive& Ar) override
+	virtual void Serialize(FCADKernelArchive& Ar) override
+	{
+		// Surface's type is serialize because it is used to instantiate the correct entity on de-serialization (@see Deserialize(FCADKernelArchive& Archive))
+		if (Ar.IsSaving())
 		{
-			// Surface's type is serialize because it is used to instantiate the correct entity on de-serialization (@see Deserialize(FCADKernelArchive& Archive))
-			if (Ar.IsSaving())
-			{
-				ESurface SurfaceType = GetSurfaceType();
-				Ar << SurfaceType;
-			}
-			FEntityGeom::Serialize(Ar);
-
-			Ar << Tolerance3D;
-			Ar << MinToleranceIso;
-			Ar << Boundary;
+			ESurface SurfaceType = GetSurfaceType();
+			Ar << SurfaceType;
 		}
+		FEntityGeom::Serialize(Ar);
 
-		/**
-		 * Specific method for surface family to instantiate the correct derived class of FCurve
-		 */
-		static TSharedPtr<FSurface> Deserialize(FCADKernelArchive& Archive);
+		Ar << Tolerance3D;
+		Ar << MinToleranceIso;
+		Ar << Boundary;
+	}
 
-		virtual ~FSurface() = default;
+	/**
+	 * Specific method for surface family to instantiate the correct derived class of FCurve
+	 */
+	static TSharedPtr<FSurface> Deserialize(FCADKernelArchive& Archive);
+
+	virtual ~FSurface() = default;
 
 #ifdef CADKERNEL_DEV
-		virtual FInfoEntity& GetInfo(FInfoEntity&) const override;
+	virtual FInfoEntity& GetInfo(FInfoEntity&) const override;
 #endif
 
-		virtual EEntity GetEntityType() const override
-		{
-			return EEntity::Surface;
-		}
-
-		virtual ESurface GetSurfaceType() const = 0;
-
-		const FSurfacicBoundary& GetBoundary() const
-		{
-			ensureCADKernel(Boundary.IsValid());
-			return Boundary;
-		};
-
-		void TrimBoundaryTo(const FSurfacicBoundary NewLimit)
-		{
-			Boundary.TrimAt(NewLimit);
-		}
-
-		void ExtendBoundaryTo(const FSurfacicBoundary MaxLimit)
-		{
-			Boundary.ExtendTo(MaxLimit);
-		}
-
-		virtual TSharedPtr<FEntityGeom> ApplyMatrix(const FMatrixH& InMatrix) const = 0;
-
-		/**
-		 * Tolerance along Iso U/V is very costly to compute and not accurate.
-		 * A first approximation is based on the surface parametric space length along U and along V
-		 * Indeed, defining a Tolerance2D has no sense has the boundary length along an Iso could be very very huge compare to the boundary length along the other Iso like [[0, 1000] [0, 1]]
-		 * @see FBoundary::ComputeMinimalTolerance
-		 * In the used, a local tolerance has to be estimated
-		 */
-		const FSurfacicTolerance& GetIsoTolerances() const
-		{
-			ensureCADKernel(MinToleranceIso.IsValid());
-			return *MinToleranceIso;
-		}
-
-		/**
-		 * Return the minimum tolerance in the parametric space of the surface along the specified axis
-		 * With Tolerance3D = FSysteme.GeometricalTolerance
-		 * @see FBoundary::ComputeMinimalTolerance
-		 */
-		double GetIsoTolerance(EIso Iso) const
-		{
-			ensureCADKernel(MinToleranceIso.IsValid());
-			return (*MinToleranceIso)[Iso];
-		}
-
-		double Get3DTolerance() const
-		{
-			return Tolerance3D;
-		}
-
-
-		virtual void EvaluatePoint(const FPoint2D& InSurfacicCoordinate, FSurfacicPoint& OutPoint3D, int32 InDerivativeOrder = 0) const = 0;
-
-		virtual void EvaluatePoints(const TArray<FPoint2D>& InSurfacicCoordinates, TArray<FSurfacicPoint>& OutPoint3D, int32 InDerivativeOrder = 0) const;
-		virtual void EvaluatePoints(const TArray<FCurvePoint2D>& InSurfacicCoordinates, TArray<FSurfacicPoint>& OutPoint3D, int32 InDerivativeOrder = 0) const;
-		virtual void EvaluatePoints(const TArray<FCurvePoint2D>& InSurfacicCoordinates, TArray<FCurvePoint>& OutPoint3D, int32 InDerivativeOrder = 0) const;
-		virtual void EvaluatePoints(FSurfacicPolyline& Polyline) const;
-		virtual void EvaluatePoints(const TArray<FCurvePoint2D>& Points2D, FSurfacicPolyline& Polyline) const;
-
-		/**
-		 * X = Rho cos(Alpha) 
-		 * Y = Rho sin(Alpha)
-		 * Z = Z
-		 * 
-		 * YZ => cylindrical projection
-		 * XY => Plan projection along rotation axis
-		 * 
-		 * For Normal(u,v) parallel Rotation Axis -> use Plan projection 
-		 * For Normal(u,v) perpendicular to Rotation Axis -> use cylindrical projection
-		 */
-		virtual FPoint2D EvaluatePointInCylindricalSpace(const FPoint2D& InSurfacicCoordinate) const {return FPoint::ZeroPoint;}
-  		virtual void EvaluatePointGridInCylindricalSpace(const FCoordinateGrid& Coordinates, TArray<FPoint2D>&) const {;}
-
-		virtual void EvaluatePointGrid(const FCoordinateGrid& Coordinates, FSurfacicSampling& OutPoints, bool bComputeNormals = false) const;
-		void EvaluateGrid(FGrid& Grid) const;
-
-		virtual void EvaluateNormals(const TArray<FPoint2D>& Points2D, TArray<FVector3f>& Normals) const;
-
-		virtual FVector EvaluateNormal(const FPoint2D& InPoint2D) const
-		{
-			int32 DerivativeOrder = 1;
-			FSurfacicPoint Point3D;
-			EvaluatePoint(InPoint2D, Point3D, DerivativeOrder);
-
-			FVector Normal = Point3D.GradientU ^ Point3D.GradientV;
-			Normal.Normalize();
-			return Normal;
-		}
+	virtual EEntity GetEntityType() const override
+	{
+		return EEntity::Surface;
+	}
 
-		/**
-		 * Divide the parametric space in the desired number of regular subdivisions and compute the associated PointGrid
-		 */
-		void Sample(const FSurfacicBoundary& Bounds, int32 NumberOfSubdivisions[2], FSurfacicSampling& OutPointSampling) const;
+	virtual ESurface GetSurfaceType() const = 0;
 
-		/**
-		 * Divide the parametric space in the desired number of regular subdivisions and compute the associated CoordinateGrid
-		 */
-		void Sample(const FSurfacicBoundary& Boundary, int32 NumberOfSubdivisions[2], FCoordinateGrid& OutCoordinateSampling) const;
+	const FSurfacicBoundary& GetBoundary() const
+	{
+		ensureCADKernel(Boundary.IsValid());
+		return Boundary;
+	};
 
-		/**
-		 * Generate a pre-sampling of the surface saved in OutCoordinates.
-		 * This sampling is light enough to allow a fast computation of the grid, precise enough to compute accurately meshing criteria
-		 */
-		virtual void Presample(const FSurfacicBoundary& InBoundaries, FCoordinateGrid& OutCoordinates);
+	void TrimBoundaryTo(const FSurfacicBoundary NewLimit)
+	{
+		Boundary.TrimAt(NewLimit);
+	}
 
-		virtual void LinesNotDerivables(const FSurfacicBoundary& Bounds, int32 InDerivativeOrder, FCoordinateGrid& OutNotDerivableCoordinates) const
-		{
-			OutNotDerivableCoordinates.Empty();
-		}
+	void ExtendBoundaryTo(const FSurfacicBoundary MaxLimit)
+	{
+		Boundary.ExtendTo(MaxLimit);
+	}
+
+	virtual TSharedPtr<FEntityGeom> ApplyMatrix(const FMatrixH& InMatrix) const = 0;
+
+	/**
+	 * Tolerance along Iso U/V is very costly to compute and not accurate.
+	 * A first approximation is based on the surface parametric space length along U and along V
+	 * Indeed, defining a Tolerance2D has no sense has the boundary length along an Iso could be very very huge compare to the boundary length along the other Iso like [[0, 1000] [0, 1]]
+	 * @see FBoundary::ComputeMinimalTolerance
+	 * In the used, a local tolerance has to be estimated
+	 */
+	const FSurfacicTolerance& GetIsoTolerances() const
+	{
+		ensureCADKernel(MinToleranceIso.IsValid());
+		return *MinToleranceIso;
+	}
+
+	/**
+	 * Return the minimum tolerance in the parametric space of the surface along the specified axis
+	 * With Tolerance3D = FSysteme.GeometricalTolerance
+	 * @see FBoundary::ComputeMinimalTolerance
+	 */
+	double GetIsoTolerance(EIso Iso) const
+	{
+		ensureCADKernel(MinToleranceIso.IsValid());
+		return (*MinToleranceIso)[Iso];
+	}
 
-		void LinesNotDerivables(int32 InDerivativeOrder, FCoordinateGrid& OutNotDerivableCoordinates) const
-		{
-			LinesNotDerivables(Boundary, InDerivativeOrder, OutNotDerivableCoordinates);
-		}
+	double Get3DTolerance() const
+	{
+		return Tolerance3D;
+	}
+
+
+	virtual void EvaluatePoint(const FPoint2D& InSurfacicCoordinate, FSurfacicPoint& OutPoint3D, int32 InDerivativeOrder = 0) const = 0;
+
+	virtual void EvaluatePoints(const TArray<FPoint2D>& InSurfacicCoordinates, TArray<FSurfacicPoint>& OutPoint3D, int32 InDerivativeOrder = 0) const;
+	virtual void EvaluatePoints(const TArray<FCurvePoint2D>& InSurfacicCoordinates, TArray<FSurfacicPoint>& OutPoint3D, int32 InDerivativeOrder = 0) const;
+	virtual void EvaluatePoints(const TArray<FCurvePoint2D>& InSurfacicCoordinates, TArray<FCurvePoint>& OutPoint3D, int32 InDerivativeOrder = 0) const;
+	virtual void EvaluatePoints(FSurfacicPolyline& Polyline) const;
+	virtual void EvaluatePoints(const TArray<FCurvePoint2D>& Points2D, FSurfacicPolyline& Polyline) const;
+
+	/**
+	 * X = Rho cos(Alpha)
+	 * Y = Rho sin(Alpha)
+	 * Z = Z
+	 *
+	 * YZ => cylindrical projection
+	 * XY => Plan projection along rotation axis
+	 *
+	 * For Normal(u,v) parallel Rotation Axis -> use Plan projection
+	 * For Normal(u,v) perpendicular to Rotation Axis -> use cylindrical projection
+	 */
+	virtual FPoint2D EvaluatePointInCylindricalSpace(const FPoint2D& InSurfacicCoordinate) const { return FPoint::ZeroPoint; }
+	virtual void EvaluatePointGridInCylindricalSpace(const FCoordinateGrid& Coordinates, TArray<FPoint2D>&) const { ; }
+
+	virtual void EvaluatePointGrid(const FCoordinateGrid& Coordinates, FSurfacicSampling& OutPoints, bool bComputeNormals = false) const;
+	void EvaluateGrid(FGrid& Grid) const;
+
+	virtual void EvaluateNormals(const TArray<FPoint2D>& Points2D, TArray<FVector3f>& Normals) const;
+
+	virtual FVector EvaluateNormal(const FPoint2D& InPoint2D) const
+	{
+		int32 DerivativeOrder = 1;
+		FSurfacicPoint Point3D;
+		EvaluatePoint(InPoint2D, Point3D, DerivativeOrder);
+
+		FVector Normal = Point3D.GradientU ^ Point3D.GradientV;
+		Normal.Normalize();
+		return Normal;
+	}
+
+	/**
+	 * Divide the parametric space in the desired number of regular subdivisions and compute the associated PointGrid
+	 */
+	void Sample(const FSurfacicBoundary& Bounds, int32 NumberOfSubdivisions[2], FSurfacicSampling& OutPointSampling) const;
+
+	/**
+	 * Divide the parametric space in the desired number of regular subdivisions and compute the associated CoordinateGrid
+	 */
+	void Sample(const FSurfacicBoundary& Boundary, int32 NumberOfSubdivisions[2], FCoordinateGrid& OutCoordinateSampling) const;
+
+	/**
+	 * Generate a pre-sampling of the surface saved in OutCoordinates.
+	 * This sampling is light enough to allow a fast computation of the grid, precise enough to compute accurately meshing criteria
+	 */
+	virtual void Presample(const FSurfacicBoundary& InBoundaries, FCoordinateGrid& OutCoordinates);
+
+	virtual void LinesNotDerivables(const FSurfacicBoundary& Bounds, int32 InDerivativeOrder, FCoordinateGrid& OutNotDerivableCoordinates) const
+	{
+		OutNotDerivableCoordinates.Empty();
+	}
 
-		/**
-		 * A surface is closed along an iso axis if it's connected to itself e.g. the edge carried by the iso curve U = UMin is linked to the edge carried by the iso curve U = UMax
-		 * A complete sphere, torus are example of surface closed along U and V
-		 * A cylinder with circle section or a cone are example of surface closed along an iso
-		 */
-		virtual void IsSurfaceClosed(bool& bOutClosedAlongU, bool& bOutClosedAlongV) const
-		{
-			bOutClosedAlongU = false;
-			bOutClosedAlongV = false;
-		}
+	void LinesNotDerivables(int32 InDerivativeOrder, FCoordinateGrid& OutNotDerivableCoordinates) const
+	{
+		LinesNotDerivables(Boundary, InDerivativeOrder, OutNotDerivableCoordinates);
+	}
+
+	/**
+	 * A surface is closed along an iso axis if it's connected to itself e.g. the edge carried by the iso curve U = UMin is linked to the edge carried by the iso curve U = UMax
+	 * A complete sphere, torus are example of surface closed along U and V
+	 * A cylinder with circle section or a cone are example of surface closed along an iso
+	 */
+	virtual void IsSurfaceClosed(bool& bOutClosedAlongU, bool& bOutClosedAlongV) const
+	{
+		bOutClosedAlongU = false;
+		bOutClosedAlongV = false;
+	}
 
-		void PresampleIsoCircle(const FSurfacicBoundary& InBoundaries, FCoordinateGrid& Coordinates, EIso Iso)
+	void PresampleIsoCircle(const FSurfacicBoundary& InBoundaries, FCoordinateGrid& Coordinates, EIso Iso)
+	{
+		constexpr double SixOverPi = 6 / DOUBLE_PI;
+		double Delta = InBoundaries.Length(Iso);
+		int32 SampleCount = /*FMath::Max(3, */ (int32)(Delta * SixOverPi + 1)/*)*/;
+		Delta /= SampleCount;
+
+		Coordinates[Iso].Empty(SampleCount + 1);
+		double Sample = InBoundaries[Iso].GetMin();
+		for (int32 Index = 0; Index <= SampleCount; ++Index)
 		{
-			constexpr double SixOverPi = 6 / DOUBLE_PI;
-			double Delta = InBoundaries.Length(Iso);
-			int32 SampleCount = /*FMath::Max(3, */ (int32)(Delta * SixOverPi + 1)/*)*/;
-			Delta /= SampleCount;
-
-			Coordinates[Iso].Empty(SampleCount + 1);
-			double Sample = InBoundaries[Iso].GetMin();
-			for (int32 Index = 0; Index <= SampleCount; ++Index)
-			{
-				Coordinates[Iso].Add(Sample);
-				Sample += Delta;
-			}
+			Coordinates[Iso].Add(Sample);
+			Sample += Delta;
 		}
-
-		/**
-		 * This function return the scale of the input Axis.
-		 * This function is useful to estimate tolerance when scales are defined in the Matrix
-		 * @param InAxis a vetor of length 1
-		 * @param InMatrix
-		 * @param InOrigin = InMatrix*FPoint::ZeroPoint
-		 */
-		static double ComputeScaleAlongAxis(const FPoint& InAxis, const FMatrixH& InMatrix, const FPoint& InOrigin)
-		{
-			FPoint Point = InMatrix.Multiply(InAxis);
-			double Length = InOrigin.Distance(Point);
-			return Length;
-		};
+	}
+
+	/**
+	 * This function return the scale of the input Axis.
+	 * This function is useful to estimate tolerance when scales are defined in the Matrix
+	 * @param InAxis a vetor of length 1
+	 * @param InMatrix
+	 * @param InOrigin = InMatrix*FPoint::ZeroPoint
+	 */
+	static double ComputeScaleAlongAxis(const FPoint& InAxis, const FMatrixH& InMatrix, const FPoint& InOrigin)
+	{
+		FPoint Point = InMatrix.Multiply(InAxis);
+		double Length = InOrigin.Distance(Point);
+		return Length;
+	};
 
 
-	};
-} // namespace CADKernel
+};
+} // namespace UE::CADKernel
 
