Rename CADKernel namespace into UE::CADKernel

diff --git a/Engine/Source/Runtime/Datasmith/CADKernel/Public/CADKernel/Math/Aabb.h b/Engine/Source/Runtime/Datasmith/CADKernel/Public/CADKernel/Math/Aabb.h
--- a/Engine/Source/Runtime/Datasmith/CADKernel/Public/CADKernel/Math/Aabb.h
+++ b/Engine/Source/Runtime/Datasmith/CADKernel/Public/CADKernel/Math/Aabb.h
@@ -16,245 +16,245 @@ enum EAABBBoundary : uint32
 
 ENUM_CLASS_FLAGS(EAABBBoundary);
 
-namespace CADKernel
+namespace UE::CADKernel
+{
+template<class PointType>
+class TAABB
 {
-	template<class PointType>
-	class TAABB
-	{
 
-	protected:
-		PointType MinCorner;
-		PointType MaxCorner;
+protected:
+	PointType MinCorner;
+	PointType MaxCorner;
 
-	public:
-		TAABB()
-			: MinCorner(PointType::FarawayPoint)
-			, MaxCorner(-PointType::FarawayPoint)
-		{
-		}
+public:
+	TAABB()
+		: MinCorner(PointType::FarawayPoint)
+		, MaxCorner(-PointType::FarawayPoint)
+	{
+	}
 
-		TAABB(const PointType& InMinCorner, const PointType& InMaxCorner)
-			: MinCorner(InMinCorner)
-			, MaxCorner(InMaxCorner)
-		{
-		}
+	TAABB(const PointType& InMinCorner, const PointType& InMaxCorner)
+		: MinCorner(InMinCorner)
+		, MaxCorner(InMaxCorner)
+	{
+	}
 
-		friend FArchive& operator<<(FArchive& Ar, TAABB& AABB)
-		{
-			Ar << AABB.MinCorner;
-			Ar << AABB.MaxCorner;
-			return Ar;
-		}
+	friend FArchive& operator<<(FArchive& Ar, TAABB& AABB)
+	{
+		Ar << AABB.MinCorner;
+		Ar << AABB.MaxCorner;
+		return Ar;
+	}
 
-		bool IsValid() const
+	bool IsValid() const
+	{
+		for (int32 Axis = 0; Axis < PointType::Dimension; Axis++)
 		{
-			for (int32 Axis = 0; Axis < PointType::Dimension; Axis++)
+			if (MinCorner[Axis] > MaxCorner[Axis])
 			{
-				if (MinCorner[Axis] > MaxCorner[Axis])
-				{
-					return false;
-				}
+				return false;
 			}
-			return true;
 		}
+		return true;
+	}
 
-		void Empty()
-		{
-			MinCorner = PointType::FarawayPoint;
-			MaxCorner = -PointType::FarawayPoint;
-		}
+	void Empty()
+	{
+		MinCorner = PointType::FarawayPoint;
+		MaxCorner = -PointType::FarawayPoint;
+	}
 
-		bool Contains(const PointType& Point) const
+	bool Contains(const PointType& Point) const
+	{
+		for (int32 Axis = 0; Axis < PointType::Dimension; Axis++)
 		{
-			for (int32 Axis = 0; Axis < PointType::Dimension; Axis++)
+			if ((Point[Axis] < MinCorner[Axis]) || (Point[Axis] > MaxCorner[Axis]))
 			{
-				if ((Point[Axis] < MinCorner[Axis]) || (Point[Axis] > MaxCorner[Axis]))
-				{
-					return false;
-				}
+				return false;
 			}
-			return true;
 		}
+		return true;
+	}
 
-		void SetMinSize(double MinSize)
+	void SetMinSize(double MinSize)
+	{
+		for (int32 Axis = 0; Axis < PointType::Dimension; Axis++)
 		{
-			for (int32 Axis = 0; Axis < PointType::Dimension; Axis++)
+			double AxisSize = GetSize(Axis);
+			if (AxisSize < MinSize)
 			{
-				double AxisSize = GetSize(Axis);
-				if (AxisSize < MinSize)
-				{
-					double Offset = (MinSize - AxisSize) / 2;
-					MinCorner[Axis] -= Offset;
-					MaxCorner[Axis] += Offset;
-				}
+				double Offset = (MinSize - AxisSize) / 2;
+				MinCorner[Axis] -= Offset;
+				MaxCorner[Axis] += Offset;
 			}
 		}
+	}
 
-		double GetMaxSize() const
+	double GetMaxSize() const
+	{
+		double MaxSideSize = 0;
+		for (int32 Index = 0; Index < PointType::Dimension; Index++)
 		{
-			double MaxSideSize = 0;
-			for (int32 Index = 0; Index < PointType::Dimension; Index++)
+			double Size = GetSize(Index);
+			if (Size > MaxSideSize)
 			{
-				double Size = GetSize(Index);
-				if (Size > MaxSideSize)
-				{
-					MaxSideSize = Size;
-				}
+				MaxSideSize = Size;
 			}
-			return MaxSideSize;
 		}
+		return MaxSideSize;
+	}
 
-		double GetSize(int32 Axis) const
-		{
-			return MaxCorner[Axis] - MinCorner[Axis];
-		}
+	double GetSize(int32 Axis) const
+	{
+		return MaxCorner[Axis] - MinCorner[Axis];
+	}
 
-		double DiagonalLength() const
-		{
-			return MaxCorner.Distance(MinCorner);
-		}
+	double DiagonalLength() const
+	{
+		return MaxCorner.Distance(MinCorner);
+	}
 
-		PointType Diagonal() const
-		{
-			return MaxCorner - MinCorner;
-		}
+	PointType Diagonal() const
+	{
+		return MaxCorner - MinCorner;
+	}
 
-		bool Contains(const TAABB& Aabb) const
-		{
-			return IsValid() && Aabb.IsValid() && Contains(Aabb.MinCorner) && Contains(Aabb.MaxCorner);
-		}
+	bool Contains(const TAABB& Aabb) const
+	{
+		return IsValid() && Aabb.IsValid() && Contains(Aabb.MinCorner) && Contains(Aabb.MaxCorner);
+	}
 
-		const PointType& GetMin() const
-		{
-			return MinCorner;
-		}
+	const PointType& GetMin() const
+	{
+		return MinCorner;
+	}
 
-		const PointType& GetMax() const
-		{
-			return MaxCorner;
-		}
+	const PointType& GetMax() const
+	{
+		return MaxCorner;
+	}
 
-		TAABB& operator+= (const double* Point)
+	TAABB& operator+= (const double* Point)
+	{
+		for (int32 Index = 0; Index < PointType::Dimension; Index++)
 		{
-			for (int32 Index = 0; Index < PointType::Dimension; Index++)
+			if (Point[Index] < MinCorner[Index])
 			{
-				if (Point[Index] < MinCorner[Index])
-				{
-					MinCorner[Index] = Point[Index];
-				}
-				if (Point[Index] > MaxCorner[Index])
-				{
-					MaxCorner[Index] = Point[Index];
-				}
+				MinCorner[Index] = Point[Index];
 			}
-			return *this;
-		}
-
-		TAABB& operator+= (const PointType& Point)
-		{
-			for (int32 Index = 0; Index < PointType::Dimension; Index++)
+			if (Point[Index] > MaxCorner[Index])
 			{
-				if (Point[Index] < MinCorner[Index])
-				{
-					MinCorner[Index] = Point[Index];
-				}
-				if (Point[Index] > MaxCorner[Index])
-				{
-					MaxCorner[Index] = Point[Index];
-				}
+				MaxCorner[Index] = Point[Index];
 			}
-			return *this;
 		}
+		return *this;
+	}
 
-
-		TAABB& operator+= (const TArray<PointType>& Points)
+	TAABB& operator+= (const PointType& Point)
+	{
+		for (int32 Index = 0; Index < PointType::Dimension; Index++)
 		{
-			for (const PointType& Point : Points)
+			if (Point[Index] < MinCorner[Index])
 			{
-				*this += Point;
+				MinCorner[Index] = Point[Index];
 			}
-			return *this;
-		}
-
-
-		void Offset(double Offset)
-		{
-			for (int32 Index = 0; Index < PointType::Dimension; Index++)
+			if (Point[Index] > MaxCorner[Index])
 			{
-				MinCorner[Index] -= Offset;
-				MaxCorner[Index] += Offset;
+				MaxCorner[Index] = Point[Index];
 			}
 		}
+		return *this;
+	}
 
-		TAABB& operator+= (const TAABB& aabb)
-		{
-			*this += aabb.MinCorner;
-			*this += aabb.MaxCorner;
-			return *this;
-		}
 
-		TAABB operator+ (const PointType& Point) const
+	TAABB& operator+= (const TArray<PointType>& Points)
+	{
+		for (const PointType& Point : Points)
 		{
-			TAABB Other = *this;
-			Other += Point;
-			return Other;
+			*this += Point;
 		}
+		return *this;
+	}
 
 
-		TAABB operator+ (const TAABB& Aabb) const
+	void Offset(double Offset)
+	{
+		for (int32 Index = 0; Index < PointType::Dimension; Index++)
 		{
-			TAABB Other = *this;
-			Other += Aabb;
-			return Other;
+			MinCorner[Index] -= Offset;
+			MaxCorner[Index] += Offset;
 		}
-	};
+	}
 
-	class FAABB : public TAABB<FPoint>
+	TAABB& operator+= (const TAABB& aabb)
 	{
+		*this += aabb.MinCorner;
+		*this += aabb.MaxCorner;
+		return *this;
+	}
 
-	public:
-		FAABB()
-			: TAABB<FPoint>()
-		{
-		}
+	TAABB operator+ (const PointType& Point) const
+	{
+		TAABB Other = *this;
+		Other += Point;
+		return Other;
+	}
 
-		FAABB(const FPoint& InMinCorner, const FPoint& InMaxCorner)
-			: TAABB<FPoint>(InMinCorner, InMaxCorner)
-		{
-		}
 
-		FPoint GetCorner(int32 Corner) const
-		{
-			return FPoint(
-				Corner & EAABBBoundary::XMin ? MinCorner[0] : MaxCorner[0],
-				Corner & EAABBBoundary::YMin ? MinCorner[1] : MaxCorner[1],
-				Corner & EAABBBoundary::ZMin ? MinCorner[2] : MaxCorner[2]
-			);
-		}
-	};
+	TAABB operator+ (const TAABB& Aabb) const
+	{
+		TAABB Other = *this;
+		Other += Aabb;
+		return Other;
+	}
+};
 
-	class CADKERNEL_API FAABB2D : public TAABB<FPoint2D>
+class FAABB : public TAABB<FPoint>
+{
+
+public:
+	FAABB()
+		: TAABB<FPoint>()
 	{
-	public:
-		FAABB2D()
-			: TAABB<FPoint2D>()
-		{
-		}
+	}
 
-		FAABB2D(const FPoint& InMinCorner, const FPoint& InMaxCorner)
-			: TAABB<FPoint2D>(InMinCorner, InMaxCorner)
-		{
-		}
+	FAABB(const FPoint& InMinCorner, const FPoint& InMaxCorner)
+		: TAABB<FPoint>(InMinCorner, InMaxCorner)
+	{
+	}
 
-		FPoint2D GetCorner(int32 CornerIndex) const
-		{
-			return FPoint2D(
-				CornerIndex & EAABBBoundary::XMin ? MinCorner[0] : MaxCorner[0],
-				CornerIndex & EAABBBoundary::YMin ? MinCorner[1] : MaxCorner[1]
-			);
-		}
+	FPoint GetCorner(int32 Corner) const
+	{
+		return FPoint(
+			Corner & EAABBBoundary::XMin ? MinCorner[0] : MaxCorner[0],
+			Corner & EAABBBoundary::YMin ? MinCorner[1] : MaxCorner[1],
+			Corner & EAABBBoundary::ZMin ? MinCorner[2] : MaxCorner[2]
+		);
+	}
+};
 
-	};
+class CADKERNEL_API FAABB2D : public TAABB<FPoint2D>
+{
+public:
+	FAABB2D()
+		: TAABB<FPoint2D>()
+	{
+	}
+
+	FAABB2D(const FPoint& InMinCorner, const FPoint& InMaxCorner)
+		: TAABB<FPoint2D>(InMinCorner, InMaxCorner)
+	{
+	}
+
+	FPoint2D GetCorner(int32 CornerIndex) const
+	{
+		return FPoint2D(
+			CornerIndex & EAABBBoundary::XMin ? MinCorner[0] : MaxCorner[0],
+			CornerIndex & EAABBBoundary::YMin ? MinCorner[1] : MaxCorner[1]
+		);
+	}
+
+};
 
-} // namespace CADKernel
+} // namespace UE::CADKernel
 
diff --git a/Engine/Source/Runtime/Datasmith/CADKernel/Public/CADKernel/Math/BSpline.h b/Engine/Source/Runtime/Datasmith/CADKernel/Public/CADKernel/Math/BSpline.h
--- a/Engine/Source/Runtime/Datasmith/CADKernel/Public/CADKernel/Math/BSpline.h
+++ b/Engine/Source/Runtime/Datasmith/CADKernel/Public/CADKernel/Math/BSpline.h
@@ -5,37 +5,37 @@
 #include "CADKernel/Math/Boundary.h"
 #include "CADKernel/Math/Point.h"
 
-namespace CADKernel
+namespace UE::CADKernel
 {
-	class FGrid;
-	class FNURBSCurve;
-	class FNURBSSurface;
-	class FNURBSSurface;
-	class FSurfacicBoundary;
-	struct FCoordinateGrid;
-	struct FCurvePoint2D;
-	struct FCurvePoint;
-	struct FLinearBoundary;
-	struct FSurfacicPoint;
-	struct FSurfacicSampling;
-
-	namespace BSpline
-	{
-		/**
-		 * Compute the values of Bernstein polynomial
-		 */
-		CADKERNEL_API void Bernstein(int32 Degree, double InCoordinateU, TArray<double>& BernsteinValuesAtU, TArray<double>& BernsteinGradientsAtU, TArray<double>& BernsteinLaplaciansAtU);
-
-		CADKERNEL_API void FindNotDerivableParameters(const FNURBSCurve&, int32 InDerivativeOrder, const FLinearBoundary& Boundary, TArray<double>& OutNotDerivableParameters);
-		CADKERNEL_API void FindNotDerivableParameters(const FNURBSSurface&, int32 InDerivativeOrder, const FSurfacicBoundary& Boundary, FCoordinateGrid& OutNotDerivableParameters);
-
-		CADKERNEL_API void Evaluate2DPoint(const FNURBSCurve&, double Coordinate, FCurvePoint2D& OutPoint, int32 DerivativeOrder);
-		CADKERNEL_API void EvaluatePoint(const FNURBSCurve&, double Coordinate, FCurvePoint& OutPoint, int32 DerivativeOrder);
-
-		CADKERNEL_API void EvaluatePoint(const FNURBSSurface&, const FPoint2D& InPoint2D, FSurfacicPoint& OutPoint3D, int32 InDerivativeOrder);
-		CADKERNEL_API void EvaluatePointGrid(const FNURBSSurface&, const FCoordinateGrid& Coords, FSurfacicSampling& OutPoints, bool bComputeNormals);
-
-		CADKERNEL_API TSharedPtr<FNURBSCurve> DuplicateNurbsCurveWithHigherDegree(int32 degre, const FNURBSCurve& InCurve);
-	}
+class FGrid;
+class FNURBSCurve;
+class FNURBSSurface;
+class FNURBSSurface;
+class FSurfacicBoundary;
+struct FCoordinateGrid;
+struct FCurvePoint2D;
+struct FCurvePoint;
+struct FLinearBoundary;
+struct FSurfacicPoint;
+struct FSurfacicSampling;
+
+namespace BSpline
+{
+/**
+ * Compute the values of Bernstein polynomial
+ */
+CADKERNEL_API void Bernstein(int32 Degree, double InCoordinateU, TArray<double>& BernsteinValuesAtU, TArray<double>& BernsteinGradientsAtU, TArray<double>& BernsteinLaplaciansAtU);
+
+CADKERNEL_API void FindNotDerivableParameters(const FNURBSCurve&, int32 InDerivativeOrder, const FLinearBoundary& Boundary, TArray<double>& OutNotDerivableParameters);
+CADKERNEL_API void FindNotDerivableParameters(const FNURBSSurface&, int32 InDerivativeOrder, const FSurfacicBoundary& Boundary, FCoordinateGrid& OutNotDerivableParameters);
+
+CADKERNEL_API void Evaluate2DPoint(const FNURBSCurve&, double Coordinate, FCurvePoint2D& OutPoint, int32 DerivativeOrder);
+CADKERNEL_API void EvaluatePoint(const FNURBSCurve&, double Coordinate, FCurvePoint& OutPoint, int32 DerivativeOrder);
+
+CADKERNEL_API void EvaluatePoint(const FNURBSSurface&, const FPoint2D& InPoint2D, FSurfacicPoint& OutPoint3D, int32 InDerivativeOrder);
+CADKERNEL_API void EvaluatePointGrid(const FNURBSSurface&, const FCoordinateGrid& Coords, FSurfacicSampling& OutPoints, bool bComputeNormals);
+
+CADKERNEL_API TSharedPtr<FNURBSCurve> DuplicateNurbsCurveWithHigherDegree(int32 degre, const FNURBSCurve& InCurve);
+}
 }
 
diff --git a/Engine/Source/Runtime/Datasmith/CADKernel/Public/CADKernel/Math/Boundary.h b/Engine/Source/Runtime/Datasmith/CADKernel/Public/CADKernel/Math/Boundary.h
--- a/Engine/Source/Runtime/Datasmith/CADKernel/Public/CADKernel/Math/Boundary.h
+++ b/Engine/Source/Runtime/Datasmith/CADKernel/Public/CADKernel/Math/Boundary.h
@@ -6,396 +6,396 @@
 #include "CADKernel/Math/MathConst.h"
 #include "CADKernel/Math/Point.h"
 
-namespace CADKernel
+namespace UE::CADKernel
 {
-	/**
-	 * MINIMAL_UNIT_LINEAR_TOLERANCE allows to define the minimal tolerance value of a parametric space
-	 * @see FLinearBoundary::ComputeMinimalTolerance
-	 */
-	#define MINIMAL_UNIT_LINEAR_TOLERANCE 10e-5
+/**
+ * MINIMAL_UNIT_LINEAR_TOLERANCE allows to define the minimal tolerance value of a parametric space
+ * @see FLinearBoundary::ComputeMinimalTolerance
+ */
+#define MINIMAL_UNIT_LINEAR_TOLERANCE 10e-5
+
+struct CADKERNEL_API FLinearBoundary
+{
+
+	/** A default boundary (0., 1.)*/
+	static const FLinearBoundary DefaultBoundary;
+
+	double Min;
+	double Max;
 
-	struct CADKERNEL_API FLinearBoundary
+	FLinearBoundary()
 	{
+		Min = 0.;
+		Max = 1.;
+	}
 
-		/** A default boundary (0., 1.)*/
-		static const FLinearBoundary DefaultBoundary;
+	FLinearBoundary(const FLinearBoundary& Boundary)
+		: Min(Boundary.Min)
+		, Max(Boundary.Max)
+	{
+	}
 
-		double Min;
-		double Max;
+	FLinearBoundary(double UMin, double UMax)
+	{
+		Set(UMin, UMax);
+	}
 
-		FLinearBoundary()
-		{
-			Min = 0.;
-			Max = 1.;
-		}
+	friend FArchive& operator<<(FArchive& Ar, FLinearBoundary& Boundary)
+	{
+		Ar.Serialize(&Boundary, sizeof(FLinearBoundary));
+		return Ar;
+	}
 
-		FLinearBoundary(const FLinearBoundary& Boundary)
-			: Min(Boundary.Min)
-			, Max(Boundary.Max)
-		{
-		}
+	constexpr double GetMin() const
+	{
+		return Min;
+	}
 
-		FLinearBoundary(double UMin, double UMax)
-		{
-			Set(UMin, UMax);
-		}
+	constexpr double GetMax() const
+	{
+		return Max;
+	}
 
-		friend FArchive& operator<<(FArchive& Ar, FLinearBoundary& Boundary)
-		{
-			Ar.Serialize(&Boundary, sizeof(FLinearBoundary));
-			return Ar;
-		}
+	constexpr double GetAt(double Coordinate) const
+	{
+		return Min + (Max - Min) * Coordinate;
+	}
 
-		constexpr double GetMin() const
-		{
-			return Min;
-		}
+	constexpr double GetMiddle() const
+	{
+		return (Min + Max) * 0.5;
+	}
 
-		constexpr double GetMax() const
-		{
-			return Max;
-		}
+	double Size() const { return Max - Min; }
 
-		constexpr double GetAt(double Coordinate) const
-		{
-			return Min + (Max - Min) * Coordinate;
-		}
+	void SetMin(double Coordinates)
+	{
+		GetMinMax(Coordinates, Max, Min, Max);
+	}
 
-		constexpr double GetMiddle() const
-		{
-			return (Min + Max) * 0.5;
-		}
+	void SetMax(double Coordinates)
+	{
+		GetMinMax(Min, Coordinates, Min, Max);
+	}
 
-		double Size() const { return Max - Min; }
+	void Set(double InUMin = 0., double InUMax = 1.)
+	{
+		GetMinMax(InUMin, InUMax, Min, Max);
+	}
 
-		void SetMin(double Coordinates)
+	/**
+	 * Set the boundary with the min and max of the array
+	 */
+	void Set(const TArray<double>& Coordinates)
+	{
+		Init();
+		for (const double& Coordinate : Coordinates)
 		{
-			GetMinMax(Coordinates, Max, Min, Max);
+			ExtendTo(Coordinate);
 		}
+	}
 
-		void SetMax(double Coordinates)
-		{
-			GetMinMax(Min, Coordinates, Min, Max);
-		}
 
-		void Set(double InUMin = 0., double InUMax = 1.)
-		{
-			GetMinMax(InUMin, InUMax, Min, Max);
-		}
+	bool IsValid() const
+	{
+		return Min <= Max;
+	}
 
-		/**
-		 * Set the boundary with the min and max of the array
-		 */
-		void Set(const TArray<double>& Coordinates)
-		{
-			Init();
-			for (const double& Coordinate : Coordinates)
-			{
-				ExtendTo(Coordinate);
-			}
-		}
+	bool Contains(double Coordinate) const
+	{
+		return RealCompare(Coordinate, Min) >= 0 && RealCompare(Coordinate, Max) <= 0;
+	}
 
+	double Length() const
+	{
+		return GetMax() - GetMin();
+	}
 
-		bool IsValid() const
-		{
-			return Min <= Max;
-		}
+	/**
+	 * Return true if the parametric domain is to small
+	 */
+	bool IsDegenerated() const
+	{
+		double DeltaU = (Max - Min);
+		return (DeltaU < DOUBLE_KINDA_SMALL_NUMBER);
+	}
 
-		bool Contains(double Coordinate) const
-		{
-			return RealCompare(Coordinate, Min) >= 0 && RealCompare(Coordinate, Max) <= 0;
-		}
+	/**
+	 * Compute the minimal tolerance of the parametric domain i.e.
+	 * ToleranceMin = Boundary.Length() * MINIMAL_UNIT_LINEAR_TOLERANCE
+	 * e.g. for a curve of 1m with a parametric space define between [0, 1], the parametric tolerance is 0.01
+	 * This is a minimal value that has to be replace with a more accurate value when its possible
+	 */
+	double ComputeMinimalTolerance() const
+	{
+		return Length() * MINIMAL_UNIT_LINEAR_TOLERANCE;
+	}
 
-		double Length() const
+	/**
+	 * If a coordinate is outside the bounds, set the coordinate at the closed limit
+	 */
+	void MoveInsideIfNot(double& Coordinate, double Tolerance = DOUBLE_SMALL_NUMBER) const
+	{
+		if (Coordinate <= Min)
 		{
-			return GetMax() - GetMin();
+			Coordinate = Min + Tolerance;
 		}
-
-		/**
-		 * Return true if the parametric domain is to small
-		 */
-		bool IsDegenerated() const
+		else if (Coordinate >= Max)
 		{
-			double DeltaU = (Max - Min);
-			return (DeltaU < DOUBLE_KINDA_SMALL_NUMBER);
+			Coordinate = Max - Tolerance;
 		}
+	}
 
-		/**
-		 * Compute the minimal tolerance of the parametric domain i.e. 
-		 * ToleranceMin = Boundary.Length() * MINIMAL_UNIT_LINEAR_TOLERANCE
-		 * e.g. for a curve of 1m with a parametric space define between [0, 1], the parametric tolerance is 0.01
-		 * This is a minimal value that has to be replace with a more accurate value when its possible
-		 */
-		double ComputeMinimalTolerance() const
-		{
-			return Length() * MINIMAL_UNIT_LINEAR_TOLERANCE;
-		}
+	/**
+	 * Uses to initiate a boundary computation with ExtendTo
+	 */
+	void Init()
+	{
+		Min = HUGE_VALUE;
+		Max = -HUGE_VALUE;
+	}
 
-		/**
-		 * If a coordinate is outside the bounds, set the coordinate at the closed limit
-		 */
-		void MoveInsideIfNot(double& Coordinate, double Tolerance = DOUBLE_SMALL_NUMBER) const
-		{
-			if (Coordinate <= Min)
-			{
-				Coordinate = Min + Tolerance;
-			}
-			else if (Coordinate >= Max)
-			{
-				Coordinate = Max - Tolerance;
-			}
-		}
+	void ExtendTo(double MinCoordinate, double MaxCoordinate)
+	{
+		GetMinMax(MinCoordinate, MaxCoordinate);
+		Min = FMath::Min(Min, MinCoordinate);
+		Max = FMath::Max(Max, MaxCoordinate);
+	}
 
-		/**
-		 * Uses to initiate a boundary computation with ExtendTo 
-		 */
-		void Init()
+	void TrimAt(const FLinearBoundary& MaxBound)
+	{
+		if (Max < MaxBound.Min || Min > MaxBound.Max)
 		{
-			Min = HUGE_VALUE;
-			Max = -HUGE_VALUE;
+			*this = MaxBound;
+			return;
 		}
 
-		void ExtendTo(double MinCoordinate, double MaxCoordinate)
-		{
-			GetMinMax(MinCoordinate, MaxCoordinate);
-			Min = FMath::Min(Min, MinCoordinate);
-			Max = FMath::Max(Max, MaxCoordinate);
-		}
+		Min = FMath::Max(Min, MaxBound.Min);
+		Max = FMath::Min(Max, MaxBound.Max);
+	}
 
-		void TrimAt(const FLinearBoundary& MaxBound)
-		{
-			if (Max < MaxBound.Min || Min > MaxBound.Max)
-			{
-				*this = MaxBound;
-				return;
-			}
-
-			Min = FMath::Max(Min, MaxBound.Min);
-			Max = FMath::Min(Max, MaxBound.Max);
-		}
+	void ExtendTo(const FLinearBoundary& MaxBound)
+	{
+		Min = FMath::Min(Min, MaxBound.Min);
+		Max = FMath::Max(Max, MaxBound.Max);
+	}
 
-		void ExtendTo(const FLinearBoundary& MaxBound)
+	void ExtendTo(double Coordinate)
+	{
+		if (Coordinate < Min)
 		{
-			Min = FMath::Min(Min, MaxBound.Min);
-			Max = FMath::Max(Max, MaxBound.Max);
+			Min = Coordinate;
 		}
 
-		void ExtendTo(double Coordinate)
+		if (Coordinate > Max)
 		{
-			if (Coordinate < Min)
-			{
-				Min = Coordinate;
-			}
-
-			if (Coordinate > Max)
-			{
-				Max = Coordinate;
-			}
+			Max = Coordinate;
 		}
+	}
 
-		void RestrictTo(const FLinearBoundary& MaxBound)
+	void RestrictTo(const FLinearBoundary& MaxBound)
+	{
+		if (MaxBound.Min > Min)
 		{
-			if (MaxBound.Min > Min)
-			{
-				Min = MaxBound.Min;
-			}
-			if (MaxBound.Max < Max)
-			{
-				Max = MaxBound.Max;
-			}
+			Min = MaxBound.Min;
 		}
-
-		/**
-		 * If the boundary width is near or equal to zero, it's widened by +/- DOUBLE_SMALL_NUMBER
-		 */
-		void WidenIfDegenerated()
+		if (MaxBound.Max < Max)
 		{
-			if (FMath::IsNearlyEqual(Min, Max))
-			{
-				Min -= DOUBLE_SMALL_NUMBER;
-				Max += DOUBLE_SMALL_NUMBER;
-			}
+			Max = MaxBound.Max;
 		}
+	}
 
-		FLinearBoundary& operator=(const FLinearBoundary& InBounds)
+	/**
+	 * If the boundary width is near or equal to zero, it's widened by +/- DOUBLE_SMALL_NUMBER
+	 */
+	void WidenIfDegenerated()
+	{
+		if (FMath::IsNearlyEqual(Min, Max))
 		{
-			Min = InBounds.Min;
-			Max = InBounds.Max;
-			return *this;
+			Min -= DOUBLE_SMALL_NUMBER;
+			Max += DOUBLE_SMALL_NUMBER;
 		}
+	}
 
-	};
-
-	class CADKERNEL_API FSurfacicBoundary
+	FLinearBoundary& operator=(const FLinearBoundary& InBounds)
 	{
-	private:
-		FLinearBoundary UVBoundaries[2];
+		Min = InBounds.Min;
+		Max = InBounds.Max;
+		return *this;
+	}
 
-	public:
-		/** A default boundary (0., 1., 0., 1.)*/
-		static const FSurfacicBoundary DefaultBoundary;
+};
 
-		FSurfacicBoundary() = default;
+class CADKERNEL_API FSurfacicBoundary
+{
+private:
+	FLinearBoundary UVBoundaries[2];
 
-		FSurfacicBoundary(double InUMin, double InUMax, double InVMin, double InVMax)
-		{
-			UVBoundaries[EIso::IsoU].Set(InUMin, InUMax);
-			UVBoundaries[EIso::IsoV].Set(InVMin, InVMax);
-		}
+public:
+	/** A default boundary (0., 1., 0., 1.)*/
+	static const FSurfacicBoundary DefaultBoundary;
 
-		FSurfacicBoundary(const FPoint2D& Point1, const FPoint2D& Point2)
-		{
-			Set(Point1, Point2);
-		}
-		
-		void Set(const FPoint2D& Point1, const FPoint2D& Point2)
-		{
-			UVBoundaries[EIso::IsoU].Set(Point1.U, Point2.U);
-			UVBoundaries[EIso::IsoV].Set(Point1.V, Point2.V);
-		}
+	FSurfacicBoundary() = default;
 
-		friend FArchive& operator<<(FArchive& Ar, FSurfacicBoundary& Boundary)
-		{
-			Ar << Boundary[EIso::IsoU];
-			Ar << Boundary[EIso::IsoV];
-			return Ar;
-		}
+	FSurfacicBoundary(double InUMin, double InUMax, double InVMin, double InVMax)
+	{
+		UVBoundaries[EIso::IsoU].Set(InUMin, InUMax);
+		UVBoundaries[EIso::IsoV].Set(InVMin, InVMax);
+	}
 
-		void Set(const FLinearBoundary& BoundU, const FLinearBoundary& BoundV)
-		{
-			UVBoundaries[EIso::IsoU] = BoundU;
-			UVBoundaries[EIso::IsoV] = BoundV;
-		}
+	FSurfacicBoundary(const FPoint2D& Point1, const FPoint2D& Point2)
+	{
+		Set(Point1, Point2);
+	}
 
-		void Set(double InUMin, double InUMax, double InVMin, double InVMax)
-		{
-			UVBoundaries[EIso::IsoU].Set(InUMin, InUMax);
-			UVBoundaries[EIso::IsoV].Set(InVMin, InVMax);
-		}
+	void Set(const FPoint2D& Point1, const FPoint2D& Point2)
+	{
+		UVBoundaries[EIso::IsoU].Set(Point1.U, Point2.U);
+		UVBoundaries[EIso::IsoV].Set(Point1.V, Point2.V);
+	}
 
-		void Set()
-		{
-			UVBoundaries[EIso::IsoU].Set();
-			UVBoundaries[EIso::IsoV].Set();
-		}
+	friend FArchive& operator<<(FArchive& Ar, FSurfacicBoundary& Boundary)
+	{
+		Ar << Boundary[EIso::IsoU];
+		Ar << Boundary[EIso::IsoV];
+		return Ar;
+	}
 
-		/**
-		 * Set the boundary with the min and max of this array
-		 */
-		void Set(const TArray<FPoint2D>& Points)
-		{
-			Init();
-			for (const FPoint2D& Point : Points)
-			{
-				ExtendTo(Point);
-			}
-		}
+	void Set(const FLinearBoundary& BoundU, const FLinearBoundary& BoundV)
+	{
+		UVBoundaries[EIso::IsoU] = BoundU;
+		UVBoundaries[EIso::IsoV] = BoundV;
+	}
 
-		const FLinearBoundary& Get(EIso Type) const
-		{
-			return UVBoundaries[Type];
-		}
+	void Set(double InUMin, double InUMax, double InVMin, double InVMax)
+	{
+		UVBoundaries[EIso::IsoU].Set(InUMin, InUMax);
+		UVBoundaries[EIso::IsoV].Set(InVMin, InVMax);
+	}
 
-		bool IsValid() const
-		{
-			return UVBoundaries[EIso::IsoU].IsValid() && UVBoundaries[EIso::IsoV].IsValid();
-		}
+	void Set()
+	{
+		UVBoundaries[EIso::IsoU].Set();
+		UVBoundaries[EIso::IsoV].Set();
+	}
 
-		/**
-		 * Return true if the parametric domain is to small
-		 */
-		bool IsDegenerated() const
+	/**
+	 * Set the boundary with the min and max of this array
+	 */
+	void Set(const TArray<FPoint2D>& Points)
+	{
+		Init();
+		for (const FPoint2D& Point : Points)
 		{
-			return UVBoundaries[EIso::IsoU].IsDegenerated() || UVBoundaries[EIso::IsoV].IsDegenerated();
+			ExtendTo(Point);
 		}
+	}
 
-		/**
-		 * Uses to initiate a boundary computation with ExtendTo
-		 */
-		void Init()
-		{
-			UVBoundaries[EIso::IsoU].Init();
-			UVBoundaries[EIso::IsoV].Init();
-		}
+	const FLinearBoundary& Get(EIso Type) const
+	{
+		return UVBoundaries[Type];
+	}
 
-		void TrimAt(const FSurfacicBoundary MaxLimit)
-		{
-			UVBoundaries[EIso::IsoU].TrimAt(MaxLimit[EIso::IsoU]);
-			UVBoundaries[EIso::IsoV].TrimAt(MaxLimit[EIso::IsoV]);
-		}
+	bool IsValid() const
+	{
+		return UVBoundaries[EIso::IsoU].IsValid() && UVBoundaries[EIso::IsoV].IsValid();
+	}
 
-		void ExtendTo(const FSurfacicBoundary MaxLimit)
-		{
-			UVBoundaries[EIso::IsoU].ExtendTo(MaxLimit[EIso::IsoU]);
-			UVBoundaries[EIso::IsoV].ExtendTo(MaxLimit[EIso::IsoV]);
-		}
+	/**
+	 * Return true if the parametric domain is to small
+	 */
+	bool IsDegenerated() const
+	{
+		return UVBoundaries[EIso::IsoU].IsDegenerated() || UVBoundaries[EIso::IsoV].IsDegenerated();
+	}
 
-		void ExtendTo(FPoint2D Point)
-		{
-			UVBoundaries[EIso::IsoU].ExtendTo(Point.U);
-			UVBoundaries[EIso::IsoV].ExtendTo(Point.V);
-		}
+	/**
+	 * Uses to initiate a boundary computation with ExtendTo
+	 */
+	void Init()
+	{
+		UVBoundaries[EIso::IsoU].Init();
+		UVBoundaries[EIso::IsoV].Init();
+	}
 
-		void ExtendTo(FPoint Point)
-		{
-			UVBoundaries[EIso::IsoU].ExtendTo(Point.X);
-			UVBoundaries[EIso::IsoV].ExtendTo(Point.Y);
-		}
+	void TrimAt(const FSurfacicBoundary MaxLimit)
+	{
+		UVBoundaries[EIso::IsoU].TrimAt(MaxLimit[EIso::IsoU]);
+		UVBoundaries[EIso::IsoV].TrimAt(MaxLimit[EIso::IsoV]);
+	}
 
-		void RestrictTo(const FSurfacicBoundary& MaxBound)
-		{
-			UVBoundaries[EIso::IsoU].RestrictTo(MaxBound.UVBoundaries[EIso::IsoU]);
-			UVBoundaries[EIso::IsoV].RestrictTo(MaxBound.UVBoundaries[EIso::IsoV]);
-		}
+	void ExtendTo(const FSurfacicBoundary MaxLimit)
+	{
+		UVBoundaries[EIso::IsoU].ExtendTo(MaxLimit[EIso::IsoU]);
+		UVBoundaries[EIso::IsoV].ExtendTo(MaxLimit[EIso::IsoV]);
+	}
 
-		/**
-		 * If Along each axis, the bound width is near equal to zero, it's widened by +/- DOUBLE_SMALL_NUMBER
-		 */
-		void WidenIfDegenerated()
-		{
-			UVBoundaries[EIso::IsoU].WidenIfDegenerated();
-			UVBoundaries[EIso::IsoV].WidenIfDegenerated();
-		}
+	void ExtendTo(FPoint2D Point)
+	{
+		UVBoundaries[EIso::IsoU].ExtendTo(Point.U);
+		UVBoundaries[EIso::IsoV].ExtendTo(Point.V);
+	}
 
-		/**
-		 * If a point is outside the bounds, set the coordinate to insert the point inside the bounds
-		 */
-		void MoveInsideIfNot(FPoint& Point, double Tolerance = DOUBLE_SMALL_NUMBER) const
-		{
-			UVBoundaries[EIso::IsoU].MoveInsideIfNot(Point.X, Tolerance);
-			UVBoundaries[EIso::IsoV].MoveInsideIfNot(Point.Y, Tolerance);
-		}
+	void ExtendTo(FPoint Point)
+	{
+		UVBoundaries[EIso::IsoU].ExtendTo(Point.X);
+		UVBoundaries[EIso::IsoV].ExtendTo(Point.Y);
+	}
 
-		/**
-		 * If a point is outside the bounds, set the coordinate to insert the point inside the bounds
-		 */
-		void MoveInsideIfNot(FPoint2D& Point, double Tolerance = DOUBLE_SMALL_NUMBER) const
-		{
-			UVBoundaries[EIso::IsoU].MoveInsideIfNot(Point.U, Tolerance);
-			UVBoundaries[EIso::IsoV].MoveInsideIfNot(Point.V, Tolerance);
-		}
+	void RestrictTo(const FSurfacicBoundary& MaxBound)
+	{
+		UVBoundaries[EIso::IsoU].RestrictTo(MaxBound.UVBoundaries[EIso::IsoU]);
+		UVBoundaries[EIso::IsoV].RestrictTo(MaxBound.UVBoundaries[EIso::IsoV]);
+	}
 
-		FSurfacicBoundary& operator=(const FSurfacicBoundary& InBounds)
-		{
-			UVBoundaries[EIso::IsoU] = InBounds.UVBoundaries[EIso::IsoU];
-			UVBoundaries[EIso::IsoV] = InBounds.UVBoundaries[EIso::IsoV];
-			return *this;
-		}
+	/**
+	 * If Along each axis, the bound width is near equal to zero, it's widened by +/- DOUBLE_SMALL_NUMBER
+	 */
+	void WidenIfDegenerated()
+	{
+		UVBoundaries[EIso::IsoU].WidenIfDegenerated();
+		UVBoundaries[EIso::IsoV].WidenIfDegenerated();
+	}
 
-		double Length(const EIso& Iso) const
-		{
-			return UVBoundaries[Iso].Length();
-		}
+	/**
+	 * If a point is outside the bounds, set the coordinate to insert the point inside the bounds
+	 */
+	void MoveInsideIfNot(FPoint& Point, double Tolerance = DOUBLE_SMALL_NUMBER) const
+	{
+		UVBoundaries[EIso::IsoU].MoveInsideIfNot(Point.X, Tolerance);
+		UVBoundaries[EIso::IsoV].MoveInsideIfNot(Point.Y, Tolerance);
+	}
 
-		constexpr const FLinearBoundary& operator[](const EIso& Iso) const
-		{
-			return UVBoundaries[Iso];
-		}
+	/**
+	 * If a point is outside the bounds, set the coordinate to insert the point inside the bounds
+	 */
+	void MoveInsideIfNot(FPoint2D& Point, double Tolerance = DOUBLE_SMALL_NUMBER) const
+	{
+		UVBoundaries[EIso::IsoU].MoveInsideIfNot(Point.U, Tolerance);
+		UVBoundaries[EIso::IsoV].MoveInsideIfNot(Point.V, Tolerance);
+	}
 
-		constexpr FLinearBoundary& operator[](const EIso& Iso)
-		{
-			return UVBoundaries[Iso];
-		}
-	};
+	FSurfacicBoundary& operator=(const FSurfacicBoundary& InBounds)
+	{
+		UVBoundaries[EIso::IsoU] = InBounds.UVBoundaries[EIso::IsoU];
+		UVBoundaries[EIso::IsoV] = InBounds.UVBoundaries[EIso::IsoV];
+		return *this;
+	}
+
+	double Length(const EIso& Iso) const
+	{
+		return UVBoundaries[Iso].Length();
+	}
+
+	constexpr const FLinearBoundary& operator[](const EIso& Iso) const
+	{
+		return UVBoundaries[Iso];
+	}
+
+	constexpr FLinearBoundary& operator[](const EIso& Iso)
+	{
+		return UVBoundaries[Iso];
+	}
+};
 }
 
diff --git a/Engine/Source/Runtime/Datasmith/CADKernel/Public/CADKernel/Math/Curvature.h b/Engine/Source/Runtime/Datasmith/CADKernel/Public/CADKernel/Math/Curvature.h
--- a/Engine/Source/Runtime/Datasmith/CADKernel/Public/CADKernel/Math/Curvature.h
+++ b/Engine/Source/Runtime/Datasmith/CADKernel/Public/CADKernel/Math/Curvature.h
@@ -4,27 +4,27 @@
 #include "CADKernel/Core/Types.h"
 #include "CADKernel/Geo/GeoEnum.h"
 
-namespace CADKernel
+namespace UE::CADKernel
 {
-	struct FIsoCurvature
-	{
-		double Min = HUGE_VALUE;
-		double Max = 0;
-	};
+struct FIsoCurvature
+{
+	double Min = HUGE_VALUE;
+	double Max = 0;
+};
 
-	struct FSurfaceCurvature
+struct FSurfaceCurvature
+{
+	FIsoCurvature Curvatures[2];
+
+	constexpr const FIsoCurvature& operator[](const EIso& Iso) const
 	{
-		FIsoCurvature Curvatures[2];
+		return Curvatures[Iso];
+	}
 
-		constexpr const FIsoCurvature& operator[](const EIso& Iso) const
-		{
-			return Curvatures[Iso];
-		}
-		
-		constexpr FIsoCurvature& operator[](const EIso& Iso)
-		{
-			return Curvatures[Iso];
-		}
-	};
+	constexpr FIsoCurvature& operator[](const EIso& Iso)
+	{
+		return Curvatures[Iso];
+	}
+};
 
 }
diff --git a/Engine/Source/Runtime/Datasmith/CADKernel/Public/CADKernel/Math/Geometry.h b/Engine/Source/Runtime/Datasmith/CADKernel/Public/CADKernel/Math/Geometry.h
--- a/Engine/Source/Runtime/Datasmith/CADKernel/Public/CADKernel/Math/Geometry.h
+++ b/Engine/Source/Runtime/Datasmith/CADKernel/Public/CADKernel/Math/Geometry.h
@@ -8,7 +8,7 @@
 #include "CADKernel/Math/MatrixH.h"
 #include "CADKernel/Math/Point.h"
 
-namespace CADKernel
+namespace UE::CADKernel
 {
 enum EPolygonSide : uint8
 {
@@ -584,7 +584,7 @@ inline FPoint2D FindIntersectionOfSegments2D(const TSegment<FPoint2D>& SegmentAB
 }
 
 /**
- * @return false if the lines are parallele 
+ * @return false if the lines are parallele
  */
 inline bool FindIntersectionOfLines2D(const TSegment<FPoint2D>& LineAB, const TSegment<FPoint2D>& LineCD, FPoint2D& OutIntersectionPoint)
 {
@@ -602,7 +602,7 @@ inline bool FindIntersectionOfLines2D(const TSegment<FPoint2D>& LineAB, const TS
 	}
 
 	double OutABIntersectionCoordinate = (DC ^ AC) / ParallelCoef;
-	OutIntersectionPoint =  LineAB[0] + AB * OutABIntersectionCoordinate;
+	OutIntersectionPoint = LineAB[0] + AB * OutABIntersectionCoordinate;
 	return true;
 }
 
@@ -623,5 +623,5 @@ inline double ComputeCosinus(FVector Vector, FVector OtherVector)
 }
 
 
-} // namespace CADKernel
+} // namespace UE::CADKernel
 
diff --git a/Engine/Source/Runtime/Datasmith/CADKernel/Public/CADKernel/Math/MathConst.h b/Engine/Source/Runtime/Datasmith/CADKernel/Public/CADKernel/Math/MathConst.h
--- a/Engine/Source/Runtime/Datasmith/CADKernel/Public/CADKernel/Math/MathConst.h
+++ b/Engine/Source/Runtime/Datasmith/CADKernel/Public/CADKernel/Math/MathConst.h
@@ -3,61 +3,61 @@
 
 #include "CADKernel/Core/Types.h"
 
-namespace CADKernel
+namespace UE::CADKernel
 {
-	/**
-	 * Sort input A & B values into OutMin & OutMax
-	 */
-	template<typename ValueType> const void GetMinMax(const ValueType& ValueA, const ValueType& ValueB, ValueType& OutMin, ValueType& OutMax)
+/**
+ * Sort input A & B values into OutMin & OutMax
+ */
+template<typename ValueType> const void GetMinMax(const ValueType& ValueA, const ValueType& ValueB, ValueType& OutMin, ValueType& OutMax)
+{
+	if (ValueA < ValueB)
 	{
-		if (ValueA < ValueB)
-		{
-			OutMin = ValueA;
-			OutMax = ValueB;
-		}
-		else
-		{
-			OutMin = ValueB;
-			OutMax = ValueA;
-		}
+		OutMin = ValueA;
+		OutMax = ValueB;
 	}
-
-	/**
-	 * Sort input values to be Min and Max
-	 */
-	template<typename ValueType> const void GetMinMax(ValueType& Min, ValueType& Max)
+	else
 	{
-		if (Max < Min)
-		{
-			Swap(Min, Max);
-		}
+		OutMin = ValueB;
+		OutMax = ValueA;
 	}
+}
 
-	inline int32 RealCompare(const double Value1, const double Value2, const double Tolerance = DOUBLE_SMALL_NUMBER)
+/**
+ * Sort input values to be Min and Max
+ */
+template<typename ValueType> const void GetMinMax(ValueType& Min, ValueType& Max)
+{
+	if (Max < Min)
 	{
-		double Difference = Value1 - Value2;
-		if (Difference < -Tolerance) 
-		{
-			return -1;
-		}
-		if (Difference > Tolerance) 
-		{
-			return 1;
-		}
-		return 0;
+		Swap(Min, Max);
 	}
+}
 
-	template< class T >
-	static T Cubic(const T A)
+inline int32 RealCompare(const double Value1, const double Value2, const double Tolerance = DOUBLE_SMALL_NUMBER)
+{
+	double Difference = Value1 - Value2;
+	if (Difference < -Tolerance)
 	{
-		return A * A * A;
+		return -1;
 	}
-
-	template< class T >
-	uint8 ToUInt8(T Value)
+	if (Difference > Tolerance)
 	{
-		return FMath::Clamp((uint8)(Value / 255.), (uint8) 0, (uint8) 255);
+		return 1;
 	}
+	return 0;
+}
+
+template< class T >
+static T Cubic(const T A)
+{
+	return A * A * A;
+}
+
+template< class T >
+uint8 ToUInt8(T Value)
+{
+	return FMath::Clamp((uint8)(Value / 255.), (uint8)0, (uint8)255);
+}
 
 
-} // namespace CADKernel	
+} // namespace UE::CADKernel	
diff --git a/Engine/Source/Runtime/Datasmith/CADKernel/Public/CADKernel/Math/MatrixH.h b/Engine/Source/Runtime/Datasmith/CADKernel/Public/CADKernel/Math/MatrixH.h
--- a/Engine/Source/Runtime/Datasmith/CADKernel/Public/CADKernel/Math/MatrixH.h
+++ b/Engine/Source/Runtime/Datasmith/CADKernel/Public/CADKernel/Math/MatrixH.h
@@ -5,289 +5,289 @@
 #include "CADKernel/Math/MathConst.h"
 #include "CADKernel/Math/Point.h"
 
-namespace CADKernel
+namespace UE::CADKernel
 {
-	/**
-	* Should be unified with the math utilities implemented by the Geometry team.
-	*/
-	class CADKERNEL_API FMatrixH
-	{
-	private:
-		double Matrix[16];
+/**
+* Should be unified with the math utilities implemented by the Geometry team.
+*/
+class CADKERNEL_API FMatrixH
+{
+private:
+	double Matrix[16];
 
-	public:
+public:
 
-		static const FMatrixH Identity;
+	static const FMatrixH Identity;
 
-		FMatrixH()
-		{
-			SetIdentity();
-		}
+	FMatrixH()
+	{
+		SetIdentity();
+	}
 
-		FMatrixH(const double* const InMatrix16)
+	FMatrixH(const double* const InMatrix16)
+	{
+		for (int32 Index = 0; Index < 16; Index++)
 		{
-			for (int32 Index = 0; Index < 16; Index++)
-			{
-				Matrix[Index] = InMatrix16[Index];
-			}
+			Matrix[Index] = InMatrix16[Index];
 		}
+	}
 
-		FMatrixH(const double InMatrix44[][4])
+	FMatrixH(const double InMatrix44[][4])
+	{
+		for (int32 Row = 0; Row < 4; Row++)
 		{
-			for (int32 Row = 0; Row < 4; Row++)
+			for (int32 Column = 0; Column < 4; Column++)
 			{
-				for (int32 Column = 0; Column < 4; Column++)
-				{
-					Matrix[4 * Row + Column] = InMatrix44[Row][Column];
-				}
+				Matrix[4 * Row + Column] = InMatrix44[Row][Column];
 			}
 		}
+	}
 
-		FMatrixH(const FPoint& Origin, const FPoint& Ox, const FPoint& Oy, const FPoint& Oz)
-		{
-			BuildChangeOfCoordinateSystemMatrix(Ox, Oy, Oz, Origin);
-		}
+	FMatrixH(const FPoint& Origin, const FPoint& Ox, const FPoint& Oy, const FPoint& Oz)
+	{
+		BuildChangeOfCoordinateSystemMatrix(Ox, Oy, Oz, Origin);
+	}
 
-		friend FArchive& operator<<(FArchive& Ar, FMatrixH& InMatrix)
-		{
-			Ar.Serialize(InMatrix.Matrix, 16 * sizeof(double));
-			return Ar;
-		}
+	friend FArchive& operator<<(FArchive& Ar, FMatrixH& InMatrix)
+	{
+		Ar.Serialize(InMatrix.Matrix, 16 * sizeof(double));
+		return Ar;
+	}
 
-		void SetIdentity()
+	void SetIdentity()
+	{
+		for (int32 Row = 0; Row < 4; Row++)
 		{
-			for (int32 Row = 0; Row < 4; Row++)
+			for (int32 Column = 0; Column < 4; Column++)
 			{
-				for (int32 Column = 0; Column < 4; Column++)
-				{
-					Get(Row, Column) = Row == Column ? 1. : 0.;
-				}
+				Get(Row, Column) = Row == Column ? 1. : 0.;
 			}
 		}
+	}
 
-		void FromAxisOrigin(const FPoint& Axis, const FPoint& Origin);
+	void FromAxisOrigin(const FPoint& Axis, const FPoint& Origin);
 
-		/**
-		 * @param Xaxis the vector X of the old coordinate system in the new coordinate system
-		 * @param Yaxis the vector Y of the old coordinate system in the new coordinate system
-		 * @param Zaxis the vector Z of the old coordinate system in the new coordinate system
-		 * @param Origin Origin of the old coordinate system in the new coordinate system
-		 * @return the transform matrix
-		 */
-		void BuildChangeOfCoordinateSystemMatrix(const FPoint& Xaxis, const FPoint& Yaxis, const FPoint& Zaxis, const FPoint& Origin);
-
-		FPoint Multiply(const FPoint& InPoint) const
-		{
-			return FPoint(
-				InPoint.X * Get(0, 0) + InPoint.Y * Get(0, 1) + InPoint.Z * Get(0, 2) + Get(0, 3),
-				InPoint.X * Get(1, 0) + InPoint.Y * Get(1, 1) + InPoint.Z * Get(1, 2) + Get(1, 3),
-				InPoint.X * Get(2, 0) + InPoint.Y * Get(2, 1) + InPoint.Z * Get(2, 2) + Get(2, 3)
-			);
-		}
-
-		FPoint MultiplyVector(const FPoint& InVector) const
-		{
-			return FPoint(
-				InVector.X * Get(0, 0) + InVector.Y * Get(0, 1) + InVector.Z * Get(0, 2),
-				InVector.X * Get(1, 0) + InVector.Y * Get(1, 1) + InVector.Z * Get(1, 2),
-				InVector.X * Get(2, 0) + InVector.Y * Get(2, 1) + InVector.Z * Get(2, 2)
-			);
-		}
-
-		FVector MultiplyVector(const FVector& InVector) const
-		{
-			return FVector(
-				InVector.X * Get(0, 0) + InVector.Y * Get(0, 1) + InVector.Z * Get(0, 2),
-				InVector.X * Get(1, 0) + InVector.Y * Get(1, 1) + InVector.Z * Get(1, 2),
-				InVector.X * Get(2, 0) + InVector.Y * Get(2, 1) + InVector.Z * Get(2, 2)
-			);
-		}
-
-		FVector3f MultiplyVector(const FVector3f& InVector) const
-		{
-			return FVector3f(
-				(float) (InVector.X * Get(0, 0) + InVector.Y * Get(0, 1) + InVector.Z * Get(0, 2)),
-				(float) (InVector.X * Get(1, 0) + InVector.Y * Get(1, 1) + InVector.Z * Get(1, 2)),
-				(float) (InVector.X * Get(2, 0) + InVector.Y * Get(2, 1) + InVector.Z * Get(2, 2))
-			);
-		}
+	/**
+	 * @param Xaxis the vector X of the old coordinate system in the new coordinate system
+	 * @param Yaxis the vector Y of the old coordinate system in the new coordinate system
+	 * @param Zaxis the vector Z of the old coordinate system in the new coordinate system
+	 * @param Origin Origin of the old coordinate system in the new coordinate system
+	 * @return the transform matrix
+	 */
+	void BuildChangeOfCoordinateSystemMatrix(const FPoint& Xaxis, const FPoint& Yaxis, const FPoint& Zaxis, const FPoint& Origin);
+
+	FPoint Multiply(const FPoint& InPoint) const
+	{
+		return FPoint(
+			InPoint.X * Get(0, 0) + InPoint.Y * Get(0, 1) + InPoint.Z * Get(0, 2) + Get(0, 3),
+			InPoint.X * Get(1, 0) + InPoint.Y * Get(1, 1) + InPoint.Z * Get(1, 2) + Get(1, 3),
+			InPoint.X * Get(2, 0) + InPoint.Y * Get(2, 1) + InPoint.Z * Get(2, 2) + Get(2, 3)
+		);
+	}
+
+	FPoint MultiplyVector(const FPoint& InVector) const
+	{
+		return FPoint(
+			InVector.X * Get(0, 0) + InVector.Y * Get(0, 1) + InVector.Z * Get(0, 2),
+			InVector.X * Get(1, 0) + InVector.Y * Get(1, 1) + InVector.Z * Get(1, 2),
+			InVector.X * Get(2, 0) + InVector.Y * Get(2, 1) + InVector.Z * Get(2, 2)
+		);
+	}
+
+	FVector MultiplyVector(const FVector& InVector) const
+	{
+		return FVector(
+			InVector.X * Get(0, 0) + InVector.Y * Get(0, 1) + InVector.Z * Get(0, 2),
+			InVector.X * Get(1, 0) + InVector.Y * Get(1, 1) + InVector.Z * Get(1, 2),
+			InVector.X * Get(2, 0) + InVector.Y * Get(2, 1) + InVector.Z * Get(2, 2)
+		);
+	}
+
+	FVector3f MultiplyVector(const FVector3f& InVector) const
+	{
+		return FVector3f(
+			(float)(InVector.X * Get(0, 0) + InVector.Y * Get(0, 1) + InVector.Z * Get(0, 2)),
+			(float)(InVector.X * Get(1, 0) + InVector.Y * Get(1, 1) + InVector.Z * Get(1, 2)),
+			(float)(InVector.X * Get(2, 0) + InVector.Y * Get(2, 1) + InVector.Z * Get(2, 2))
+		);
+	}
 
-		static FMatrixH MakeRotationMatrix(double InAngle, const FPoint InAxe);
+	static FMatrixH MakeRotationMatrix(double InAngle, const FPoint InAxe);
 
-		static FMatrixH MakeTranslationMatrix(const FPoint& InPoint);
+	static FMatrixH MakeTranslationMatrix(const FPoint& InPoint);
 
-		static FMatrixH MakeScaleMatrix(FPoint& Scale);
-		static FMatrixH MakeScaleMatrix(double ScaleX, double ScaleY, double ScaleZ);
+	static FMatrixH MakeScaleMatrix(FPoint& Scale);
+	static FMatrixH MakeScaleMatrix(double ScaleX, double ScaleY, double ScaleZ);
 
-		/**
-		 * Apply the rotation centered in origin to PointToRotate
-		 */
-		FPoint PointRotation(const FPoint& PointToRotate, const FPoint& Origin) const
+	/**
+	 * Apply the rotation centered in origin to PointToRotate
+	 */
+	FPoint PointRotation(const FPoint& PointToRotate, const FPoint& Origin) const
+	{
+		FPoint Result = Origin;
+		for (int32 Index = 0; Index < 3; Index++)
 		{
-			FPoint Result = Origin;
-			for (int32 Index = 0; Index < 3; Index++)
+			for (int32 Jndex = 0; Jndex < 3; Jndex++)
 			{
-				for (int32 Jndex = 0; Jndex < 3; Jndex++)
-				{
-					Result[Index] += Get(Index, Jndex) * (PointToRotate[Jndex] - Origin[Jndex]);
-				}
+				Result[Index] += Get(Index, Jndex) * (PointToRotate[Jndex] - Origin[Jndex]);
 			}
-			return Result;
 		}
+		return Result;
+	}
 
-		FVector PointRotation(const FVector& PointToRotate, const FVector& Origin) const
+	FVector PointRotation(const FVector& PointToRotate, const FVector& Origin) const
+	{
+		FVector Result = Origin;
+		for (int32 Index = 0; Index < 3; Index++)
 		{
-			FVector Result = Origin;
-			for (int32 Index = 0; Index < 3; Index++)
+			for (int32 Jndex = 0; Jndex < 3; Jndex++)
 			{
-				for (int32 Jndex = 0; Jndex < 3; Jndex++)
-				{
-					Result[Index] += Get(Index, Jndex) * (PointToRotate[Jndex] - Origin[Jndex]);
-				}
+				Result[Index] += Get(Index, Jndex) * (PointToRotate[Jndex] - Origin[Jndex]);
 			}
-			return Result;
 		}
+		return Result;
+	}
 
-		FPoint2D PointRotation(const FPoint2D& PointToRotate, const FPoint2D& Origin) const
+	FPoint2D PointRotation(const FPoint2D& PointToRotate, const FPoint2D& Origin) const
+	{
+		FPoint2D Result = Origin;
+		for (int32 Index = 0; Index < 2; Index++)
 		{
-			FPoint2D Result = Origin;
-			for (int32 Index = 0; Index < 2; Index++)
+			for (int32 Jndex = 0; Jndex < 2; Jndex++)
 			{
-				for (int32 Jndex = 0; Jndex < 2; Jndex++)
-				{
-					Result[Index] += Get(Index, Jndex) * (PointToRotate[Jndex] - Origin[Jndex]);
-				}
+				Result[Index] += Get(Index, Jndex) * (PointToRotate[Jndex] - Origin[Jndex]);
 			}
-			return Result;
 		}
+		return Result;
+	}
 
-		void Inverse();
+	void Inverse();
 
-		FMatrixH GetInverse() const
-		{
-			FMatrixH NewMatrix = *this;
-			NewMatrix.Inverse();
-			return NewMatrix;
-		}
+	FMatrixH GetInverse() const
+	{
+		FMatrixH NewMatrix = *this;
+		NewMatrix.Inverse();
+		return NewMatrix;
+	}
 
-		void Transpose()
+	void Transpose()
+	{
+		FMatrixH Tmp = *this;
+		for (int32 Index = 0; Index < 4; Index++)
 		{
-			FMatrixH Tmp = *this;
-			for (int32 Index = 0; Index < 4; Index++)
+			for (int32 Jndex = 0; Jndex < 4; Jndex++)
 			{
-				for (int32 Jndex = 0; Jndex < 4; Jndex++)
-				{
-					Get(Index, Jndex) = Tmp.Get(Jndex, Index);
-				}
+				Get(Index, Jndex) = Tmp.Get(Jndex, Index);
 			}
 		}
+	}
 
-		double& Get(int32 Row, int32 Column)
-		{
-			return Matrix[Row * 4 + Column];
-		}
+	double& Get(int32 Row, int32 Column)
+	{
+		return Matrix[Row * 4 + Column];
+	}
 
-		double Get(int32 Row, int32 Column) const
-		{
-			return Matrix[Row * 4 + Column];
-		}
+	double Get(int32 Row, int32 Column) const
+	{
+		return Matrix[Row * 4 + Column];
+	}
 
-		double& operator()(int32 Row, int32 Column)
-		{
-			return Matrix[Row * 4 + Column];
-		}
+	double& operator()(int32 Row, int32 Column)
+	{
+		return Matrix[Row * 4 + Column];
+	}
 
-		double operator()(int32 Row, int32 Column) const
-		{
-			return Matrix[Row * 4 + Column];
-		}
+	double operator()(int32 Row, int32 Column) const
+	{
+		return Matrix[Row * 4 + Column];
+	}
 
-		double& operator[](int32 Index)
-		{
-			return Matrix[Index];
-		}
+	double& operator[](int32 Index)
+	{
+		return Matrix[Index];
+	}
 
-		FMatrixH operator*(const FMatrixH& InMatrix) const
-		{
-			FMatrixH Result;
+	FMatrixH operator*(const FMatrixH& InMatrix) const
+	{
+		FMatrixH Result;
 
-			for (int32 Index = 0; Index < 4; Index++)
+		for (int32 Index = 0; Index < 4; Index++)
+		{
+			for (int32 Jndex = 0; Jndex < 4; Jndex++)
 			{
-				for (int32 Jndex = 0; Jndex < 4; Jndex++)
+				Result.Get(Index, Jndex) = 0;
+				for (int32 Kndex = 0; Kndex < 4; Kndex++)
 				{
-					Result.Get(Index, Jndex) = 0;
-					for (int32 Kndex = 0; Kndex < 4; Kndex++)
-					{
-						Result.Get(Index, Jndex) = Result.Get(Index, Jndex) + (Get(Index, Kndex) * InMatrix.Get(Kndex, Jndex));
-					}
+					Result.Get(Index, Jndex) = Result.Get(Index, Jndex) + (Get(Index, Kndex) * InMatrix.Get(Kndex, Jndex));
 				}
 			}
-			return Result;
 		}
+		return Result;
+	}
 
-		void operator*=(const FMatrixH& InMatrix)
-		{
-			FMatrixH Result = *this * InMatrix;
-			Move(*this, Result);
-		}
+	void operator*=(const FMatrixH& InMatrix)
+	{
+		FMatrixH Result = *this * InMatrix;
+		Move(*this, Result);
+	}
 
-		FPoint operator*(const FPoint& Point) const
-		{
-			return Multiply(Point);
-		}
+	FPoint operator*(const FPoint& Point) const
+	{
+		return Multiply(Point);
+	}
 
-		FMatrixH operator+(const FMatrixH& InMatrix) const
+	FMatrixH operator+(const FMatrixH& InMatrix) const
+	{
+		FMatrixH Result;
+		for (int32 i = 0; i < 16; i++)
 		{
-			FMatrixH Result;
-			for (int32 i = 0; i < 16; i++)
-			{
-				Result.Matrix[i] = Matrix[i] + InMatrix.Matrix[i];
-			}
-			return Result;
+			Result.Matrix[i] = Matrix[i] + InMatrix.Matrix[i];
 		}
+		return Result;
+	}
 
-		void GetMatrixDouble(double* OutMatrix) const
-		{
-			memcpy(OutMatrix, Matrix, 16 * sizeof(double));
-		}
+	void GetMatrixDouble(double* OutMatrix) const
+	{
+		memcpy(OutMatrix, Matrix, 16 * sizeof(double));
+	}
 
-		FPoint Column(int32 Index) const
-		{
-			return FPoint(Get(0, Index), Get(1, Index), Get(2, Index));
-		}
+	FPoint Column(int32 Index) const
+	{
+		return FPoint(Get(0, Index), Get(1, Index), Get(2, Index));
+	}
 
-		FPoint Row(int32 Index) const
-		{
-			return FPoint(Get(Index, 0), Get(Index, 1), Get(Index, 2));
-		}
+	FPoint Row(int32 Index) const
+	{
+		return FPoint(Get(Index, 0), Get(Index, 1), Get(Index, 2));
+	}
 
-		void Print(EVerboseLevel level) const;
+	void Print(EVerboseLevel level) const;
 
-		bool IsId() const
+	bool IsId() const
+	{
+		for (int32 Row = 0; Row < 4; Row++)
 		{
-			for (int32 Row = 0; Row < 4; Row++)
+			for (int32 Column = 0; Column < 4; Column++)
 			{
-				for (int32 Column = 0; Column < 4; Column++)
+				if (Row == Column)
+				{
+					if (!FMath::IsNearlyEqual(Get(Row, Column), 1.)) return false;
+				}
+				else
 				{
-					if (Row == Column)
-					{
-						if (!FMath::IsNearlyEqual(Get(Row, Column), 1.)) return false;
-					}
-					else
-					{
-						if (!FMath::IsNearlyZero(Get(Row, Column))) return false;
-					}
+					if (!FMath::IsNearlyZero(Get(Row, Column))) return false;
 				}
 			}
-			return true;
 		}
-	};
+		return true;
+	}
+};
 
-	CADKERNEL_API void InverseMatrixN(double* Matrix, int32 n);
+CADKERNEL_API void InverseMatrixN(double* Matrix, int32 n);
 
-	CADKERNEL_API void MatrixProduct(int32 ARowNum, int32 AColumnNum, int32 ResultRank, const double* InMatrixA, const double* InMatrixB, double* OutMatrix);
+CADKERNEL_API void MatrixProduct(int32 ARowNum, int32 AColumnNum, int32 ResultRank, const double* InMatrixA, const double* InMatrixB, double* OutMatrix);
 
-	CADKERNEL_API void TransposeMatrix(int32 RowNum, int32 ColumnNum, const double* InMatrix, double* OutMatrix);
+CADKERNEL_API void TransposeMatrix(int32 RowNum, int32 ColumnNum, const double* InMatrix, double* OutMatrix);
 
-} // namespace CADKernel
+} // namespace UE::CADKernel
 
diff --git a/Engine/Source/Runtime/Datasmith/CADKernel/Public/CADKernel/Math/Point.h b/Engine/Source/Runtime/Datasmith/CADKernel/Public/CADKernel/Math/Point.h
--- a/Engine/Source/Runtime/Datasmith/CADKernel/Public/CADKernel/Math/Point.h
+++ b/Engine/Source/Runtime/Datasmith/CADKernel/Public/CADKernel/Math/Point.h
@@ -5,7 +5,7 @@
 #include "CADKernel/Math/MathConst.h"
 
 // this will eventually be reconciled with what the Geometry team is using.
-namespace CADKernel
+namespace UE::CADKernel
 {
 
 class FPoint2D;
@@ -23,7 +23,7 @@ public:
 		};
 
 		UE_DEPRECATED(all, "For internal use only")
-		double XYZ[3];
+			double XYZ[3];
 	};
 
 public:
@@ -57,16 +57,16 @@ public:
 
 	constexpr double& operator[](int32 Index)
 	{
-PRAGMA_DISABLE_DEPRECATION_WARNINGS
-		return XYZ[Index];
-PRAGMA_ENABLE_DEPRECATION_WARNINGS
+		PRAGMA_DISABLE_DEPRECATION_WARNINGS
+			return XYZ[Index];
+		PRAGMA_ENABLE_DEPRECATION_WARNINGS
 	}
 
 	constexpr double operator[](int32 Index) const
 	{
-PRAGMA_DISABLE_DEPRECATION_WARNINGS
-		return XYZ[Index];
-PRAGMA_ENABLE_DEPRECATION_WARNINGS
+		PRAGMA_DISABLE_DEPRECATION_WARNINGS
+			return XYZ[Index];
+		PRAGMA_ENABLE_DEPRECATION_WARNINGS
 	}
 
 	bool operator==(const FPoint& Point) const
@@ -218,9 +218,9 @@ PRAGMA_ENABLE_DEPRECATION_WARNINGS
 		return FVector(X, Y, Z);
 	}
 
-	operator FVector3f() const 
+	operator FVector3f() const
 	{
-		return FVector3f((float) X, (float)Y, (float)Z);
+		return FVector3f((float)X, (float)Y, (float)Z);
 	}
 
 	void SetMin(const FPoint& Point)
@@ -361,7 +361,7 @@ public:
 		};
 
 		UE_DEPRECATED(all, "For internal use only")
-		double XYZW[4];
+			double XYZW[4];
 	};
 
 public:
@@ -396,16 +396,16 @@ public:
 
 	constexpr double& operator[](int32 Index)
 	{
-PRAGMA_DISABLE_DEPRECATION_WARNINGS
-		return XYZW[Index];
-PRAGMA_ENABLE_DEPRECATION_WARNINGS
+		PRAGMA_DISABLE_DEPRECATION_WARNINGS
+			return XYZW[Index];
+		PRAGMA_ENABLE_DEPRECATION_WARNINGS
 	}
 
 	constexpr double operator[](int32 Index) const
 	{
-PRAGMA_DISABLE_DEPRECATION_WARNINGS
-		return XYZW[Index];
-PRAGMA_ENABLE_DEPRECATION_WARNINGS
+		PRAGMA_DISABLE_DEPRECATION_WARNINGS
+			return XYZW[Index];
+		PRAGMA_ENABLE_DEPRECATION_WARNINGS
 	}
 
 	FPointH operator+(const FPointH& Point) const
@@ -686,7 +686,7 @@ public:
 		};
 
 		UE_DEPRECATED(all, "For internal use only")
-		float XYZ[3];
+			float XYZ[3];
 	};
 
 	/** A zero point (0,0,0)*/
@@ -710,16 +710,16 @@ public:
 
 	constexpr float& operator[](int32 Index)
 	{
-PRAGMA_DISABLE_DEPRECATION_WARNINGS
-		return XYZ[Index];
-PRAGMA_ENABLE_DEPRECATION_WARNINGS
+		PRAGMA_DISABLE_DEPRECATION_WARNINGS
+			return XYZ[Index];
+		PRAGMA_ENABLE_DEPRECATION_WARNINGS
 	}
 
 	constexpr float operator[](int32 Index) const
 	{
-PRAGMA_DISABLE_DEPRECATION_WARNINGS
-		return XYZ[Index];
-PRAGMA_ENABLE_DEPRECATION_WARNINGS
+		PRAGMA_DISABLE_DEPRECATION_WARNINGS
+			return XYZ[Index];
+		PRAGMA_ENABLE_DEPRECATION_WARNINGS
 	}
 
 	FFPoint& operator=(const FPoint& Point)
@@ -746,17 +746,17 @@ inline FPoint2D operator*(double s, const FPoint2D& Point)
 
 }
 
-inline uint32 GetTypeHash(const CADKernel::FPoint& Point)
+inline uint32 GetTypeHash(const UE::CADKernel::FPoint& Point)
 {
 	return HashCombine(GetTypeHash(Point.X), HashCombine(GetTypeHash(Point.Y), GetTypeHash(Point.Z)));
 }
 
-inline uint32 GetTypeHash(const CADKernel::FPoint2D& Point)
+inline uint32 GetTypeHash(const UE::CADKernel::FPoint2D& Point)
 {
 	return HashCombine(GetTypeHash(Point.U), GetTypeHash(Point.V));
 }
 
-inline uint32 GetTypeHash(const CADKernel::FPointH& Point)
+inline uint32 GetTypeHash(const UE::CADKernel::FPointH& Point)
 {
 	return HashCombine(GetTypeHash(Point.X), HashCombine(GetTypeHash(Point.Y), HashCombine(GetTypeHash(Point.Z), GetTypeHash(Point.W))));
 }
diff --git a/Engine/Source/Runtime/Datasmith/CADKernel/Public/CADKernel/Mesh/Criteria/Criterion.h b/Engine/Source/Runtime/Datasmith/CADKernel/Public/CADKernel/Mesh/Criteria/Criterion.h
--- a/Engine/Source/Runtime/Datasmith/CADKernel/Public/CADKernel/Mesh/Criteria/Criterion.h
+++ b/Engine/Source/Runtime/Datasmith/CADKernel/Public/CADKernel/Mesh/Criteria/Criterion.h
@@ -8,131 +8,131 @@
 #include "CADKernel/Math/Point.h"
 #include "CADKernel/Mesh/Criteria/CriterionType.h"
 
-namespace CADKernel
+namespace UE::CADKernel
 {
-	extern const TCHAR* CriterionTypeNames[];
+extern const TCHAR* CriterionTypeNames[];
 
-	// Defined for Python purpose
-	extern const char* CriterionTypeConstNames[];
-	extern const char* CriterionTypeConstDescHelp[];
+// Defined for Python purpose
+extern const char* CriterionTypeConstNames[];
+extern const char* CriterionTypeConstDescHelp[];
 
-	class FTopologicalEdge;
-	struct FCurvePoint;
-	struct FIsoCurvature;
+class FTopologicalEdge;
+struct FCurvePoint;
+struct FIsoCurvature;
 
-	class CADKERNEL_API FCriterion : public FEntity
-	{
-		friend class FEntity;
+class CADKERNEL_API FCriterion : public FEntity
+{
+	friend class FEntity;
 
-	protected:
+protected:
 
-		FCriterion() = default;
+	FCriterion() = default;
 
-	public:
+public:
 
-		virtual void Serialize(FCADKernelArchive& Ar) override
+	virtual void Serialize(FCADKernelArchive& Ar) override
+	{
+		// Criterion's type is serialize because it is used to instantiate the correct entity on deserialization (@see Deserialize(FArchive& Archive)) 
+		if (Ar.IsSaving())
 		{
-			// Criterion's type is serialize because it is used to instantiate the correct entity on deserialization (@see Deserialize(FArchive& Archive)) 
-			if (Ar.IsSaving())
-			{
-				ECriterion CriterionType = GetCriterionType();
-				Ar << CriterionType;
-			}
-			FEntity::Serialize(Ar);
+			ECriterion CriterionType = GetCriterionType();
+			Ar << CriterionType;
 		}
+		FEntity::Serialize(Ar);
+	}
 
-		/**
-		 * Specific method for curve family to instantiate the correct derived class of FCurve
-		 */
-		static TSharedPtr<FCriterion> Deserialize(FCADKernelArchive& Archive);
+	/**
+	 * Specific method for curve family to instantiate the correct derived class of FCurve
+	 */
+	static TSharedPtr<FCriterion> Deserialize(FCADKernelArchive& Archive);
 
-		static TSharedPtr<FCriterion> CreateCriterion(ECriterion type, double value = 0.);
+	static TSharedPtr<FCriterion> CreateCriterion(ECriterion type, double value = 0.);
 
 #ifdef CADKERNEL_DEV
-		virtual FInfoEntity& GetInfo(FInfoEntity&) const override;
+	virtual FInfoEntity& GetInfo(FInfoEntity&) const override;
 #endif
 
-		virtual EEntity GetEntityType() const override
-		{
-			return EEntity::Criterion;
-		}
+	virtual EEntity GetEntityType() const override
+	{
+		return EEntity::Criterion;
+	}
 
-		/**
-		 * Sag & Angle criterion.pdf
-		 * https://docs.google.com/presentation/d/1bUnrRFWCW3sDn9ngb9ftfQS-2JxNJaUZlh783hZMMEw/edit?usp=sharing
-		 */
-		virtual void ApplyOnEdgeParameters(FTopologicalEdge& Edge, const TArray<double>& Coordinates, const TArray<FCurvePoint>& Points) const;
+	/**
+	 * Sag & Angle criterion.pdf
+	 * https://docs.google.com/presentation/d/1bUnrRFWCW3sDn9ngb9ftfQS-2JxNJaUZlh783hZMMEw/edit?usp=sharing
+	 */
+	virtual void ApplyOnEdgeParameters(FTopologicalEdge& Edge, const TArray<double>& Coordinates, const TArray<FCurvePoint>& Points) const;
 
-		virtual ECriterion GetCriterionType() const = 0;
+	virtual ECriterion GetCriterionType() const = 0;
 
-		FString GetCriterionName()
-		{
-			return CriterionTypeNames[(int32)GetCriterionType()];
-		}
+	FString GetCriterionName()
+	{
+		return CriterionTypeNames[(int32)GetCriterionType()];
+	}
 
-		virtual double Value() const = 0;
+	virtual double Value() const = 0;
 
-		static double EvaluateSag(const FPoint& PointPoint, const FPoint& PointNext, const FPoint& PointMiddle, double& Length)
-		{
-			FPoint ChordVec = PointNext - PointPoint;
-			FPoint MiddleVec = PointMiddle - PointPoint;
+	static double EvaluateSag(const FPoint& PointPoint, const FPoint& PointNext, const FPoint& PointMiddle, double& Length)
+	{
+		FPoint ChordVec = PointNext - PointPoint;
+		FPoint MiddleVec = PointMiddle - PointPoint;
+
+		double Sag = 0.0;
+		double NormSqrVec = ChordVec * ChordVec;
+		Length = sqrt(NormSqrVec);
 
-			double Sag = 0.0;
-			double NormSqrVec = ChordVec * ChordVec;
-			Length = sqrt(NormSqrVec);
+		if (NormSqrVec > DOUBLE_SMALL_NUMBER)
+		{
+			FPoint VecSag = ChordVec ^ MiddleVec;
+			double NormSqrSag = VecSag * VecSag;
+			Sag = NormSqrSag / NormSqrVec;
 
-			if (NormSqrVec > DOUBLE_SMALL_NUMBER)
+			if (Sag < SMALL_NUMBER_SQUARE)
+			{
+				Sag = 0.;
+			}
+			else
 			{
-				FPoint VecSag = ChordVec ^ MiddleVec;
-				double NormSqrSag = VecSag * VecSag;
-				Sag = NormSqrSag / NormSqrVec;
-
-				if (Sag < SMALL_NUMBER_SQUARE)
-				{
-					Sag = 0.;
-				}
-				else
-				{
-					Sag = sqrt(Sag);
-				}
+				Sag = sqrt(Sag);
 			}
-			return Sag;
 		}
+		return Sag;
+	}
 
-		virtual bool IsAppliedBetweenBreaks() const
-		{
-			return false;
-		}
+	virtual bool IsAppliedBetweenBreaks() const
+	{
+		return false;
+	}
 
-		static FString GetCriterionName(ECriterion CriterionType)
-		{
-			return CriterionTypeNames[(int32)CriterionType];
-		}
+	static FString GetCriterionName(ECriterion CriterionType)
+	{
+		return CriterionTypeNames[(int32)CriterionType];
+	}
 
 
-		static double DefaultValue(ECriterion type);
+	static double DefaultValue(ECriterion type);
 
-		virtual void UpdateDelta(double InDeltaU, double InUSag, double InDiagonalSag, double InVSag, double ChordLength, double DiagonalLength, double& OutSagDeltaUMax, double& OutSagDeltaUMin, FIsoCurvature& SurfaceCurvature) const
-		{
-			// the component according to U of the sag along the diagonal = (diagonal sag - V Sag) * U Length / diag length
-			double DiagonalSagU = FMath::Abs(InDiagonalSag - InVSag) * ChordLength / DiagonalLength;
-			InUSag = FMath::Max(InUSag, DiagonalSagU);
+	virtual void UpdateDelta(double InDeltaU, double InUSag, double InDiagonalSag, double InVSag, double ChordLength, double DiagonalLength, double& OutSagDeltaUMax, double& OutSagDeltaUMin, FIsoCurvature& SurfaceCurvature) const
+	{
+		// the component according to U of the sag along the diagonal = (diagonal sag - V Sag) * U Length / diag length
+		double DiagonalSagU = FMath::Abs(InDiagonalSag - InVSag) * ChordLength / DiagonalLength;
+		InUSag = FMath::Max(InUSag, DiagonalSagU);
 
-			if (InUSag > DOUBLE_SMALL_NUMBER)
-			{
-				double DeltaUMax = ComputeDeltaU(ChordLength, InDeltaU, InUSag);
-				OutSagDeltaUMax = (OutSagDeltaUMax < DeltaUMax) ? OutSagDeltaUMax : DeltaUMax;
-			}
+		if (InUSag > DOUBLE_SMALL_NUMBER)
+		{
+			double DeltaUMax = ComputeDeltaU(ChordLength, InDeltaU, InUSag);
+			OutSagDeltaUMax = (OutSagDeltaUMax < DeltaUMax) ? OutSagDeltaUMax : DeltaUMax;
 		}
+	}
 
-	protected:
-
-		virtual double ComputeDeltaU(double ChordLength, double DeltaU, double Sag) const
-		{
-			ensureCADKernel(false);
-			return 0.0;
-		};
+protected:
 
+	virtual double ComputeDeltaU(double ChordLength, double DeltaU, double Sag) const
+	{
+		ensureCADKernel(false);
+		return 0.0;
 	};
 
-} // namespace CADKernel
+};
+
+} // namespace UE::CADKernel
diff --git a/Engine/Source/Runtime/Datasmith/CADKernel/Public/CADKernel/Mesh/Criteria/CriterionType.h b/Engine/Source/Runtime/Datasmith/CADKernel/Public/CADKernel/Mesh/Criteria/CriterionType.h
--- a/Engine/Source/Runtime/Datasmith/CADKernel/Public/CADKernel/Mesh/Criteria/CriterionType.h
+++ b/Engine/Source/Runtime/Datasmith/CADKernel/Public/CADKernel/Mesh/Criteria/CriterionType.h
@@ -3,16 +3,16 @@
 
 #include "CADKernel/Core/Types.h"
 
-namespace CADKernel
+namespace UE::CADKernel
 {
-	enum class ECriterion : uint8
-	{
-		Size = 0,
-		MaxSize = 1,
-		MinSize = 2,
-		Angle = 3,
-		Sag = 4,
-		CADCurvature = 5,
-		None = 255
-	} ;
-} // namespace CADKernel
+enum class ECriterion : uint8
+{
+	Size = 0,
+	MaxSize = 1,
+	MinSize = 2,
+	Angle = 3,
+	Sag = 4,
+	CADCurvature = 5,
+	None = 255
+};
+} // namespace UE::CADKernel
diff --git a/Engine/Source/Runtime/Datasmith/CADKernel/Public/CADKernel/Mesh/Structure/EdgeMesh.h b/Engine/Source/Runtime/Datasmith/CADKernel/Public/CADKernel/Mesh/Structure/EdgeMesh.h
--- a/Engine/Source/Runtime/Datasmith/CADKernel/Public/CADKernel/Mesh/Structure/EdgeMesh.h
+++ b/Engine/Source/Runtime/Datasmith/CADKernel/Public/CADKernel/Mesh/Structure/EdgeMesh.h
@@ -3,42 +3,42 @@
 
 #include "CADKernel/Mesh/Structure/Mesh.h"
 
-namespace CADKernel
+namespace UE::CADKernel
 {
+/**
+ * Define the mesh of an edge. This mesh is polygon
+ */
+class CADKERNEL_API FEdgeMesh : public FMesh
+{
+public:
+
 	/**
-	 * Define the mesh of an edge. This mesh is polygon
+	 * Index of the vertices of the mesh of the TopologicalEdge in the MeshModel vertex array i.e.
+	 * EdgeVerticesIndex = {5,10,11,12,13,14,3} defines Edge(Vertices[5],Vertices[10]), Edge(Vertices[10],Vertices[11]) etc
+	 * of the mesh (Edge is side of mesh's triangle). "Vertices" is the mesh model vertices array
 	 */
-	class CADKERNEL_API FEdgeMesh : public FMesh
-	{
-	public:
+	TArray<int32> EdgeVerticesIndex;
 
-		/** 
-		 * Index of the vertices of the mesh of the TopologicalEdge in the MeshModel vertex array i.e. 
-		 * EdgeVerticesIndex = {5,10,11,12,13,14,3} defines Edge(Vertices[5],Vertices[10]), Edge(Vertices[10],Vertices[11]) etc 
-		 * of the mesh (Edge is side of mesh's triangle). "Vertices" is the mesh model vertices array 
-		 */
-		TArray<int32> EdgeVerticesIndex;
-
-		FEdgeMesh(FModelMesh& Model, FTopologicalEntity& TopologicalEntity)
-			: FMesh(Model, TopologicalEntity)
-		{
-		}
+	FEdgeMesh(FModelMesh& Model, FTopologicalEntity& TopologicalEntity)
+		: FMesh(Model, TopologicalEntity)
+	{
+	}
 
-		void Mesh(int32 StartVertexNodeIndex, int32 EndVertexNodeIndex)
+	void Mesh(int32 StartVertexNodeIndex, int32 EndVertexNodeIndex)
+	{
+		EdgeVerticesIndex.Reserve(NodeCoordinates.Num() + 2);
+		EdgeVerticesIndex.Add(StartVertexNodeIndex);
+		for (int32 Index = StartNodeId; Index < NodeCoordinates.Num() + StartNodeId; ++Index)
 		{
-			EdgeVerticesIndex.Reserve(NodeCoordinates.Num() + 2);
-			EdgeVerticesIndex.Add(StartVertexNodeIndex);
-			for (int32 Index = StartNodeId; Index < NodeCoordinates.Num() + StartNodeId; ++Index)
-			{
-				EdgeVerticesIndex.Add(Index);
-			}
-			EdgeVerticesIndex.Add(EndVertexNodeIndex);
+			EdgeVerticesIndex.Add(Index);
 		}
+		EdgeVerticesIndex.Add(EndVertexNodeIndex);
+	}
 
-		int32 GetNodeCount() const
-		{
-			return EdgeVerticesIndex.Num();
-		}
-	};
+	int32 GetNodeCount() const
+	{
+		return EdgeVerticesIndex.Num();
+	}
+};
 }
 
diff --git a/Engine/Source/Runtime/Datasmith/CADKernel/Public/CADKernel/Mesh/Structure/FaceMesh.h b/Engine/Source/Runtime/Datasmith/CADKernel/Public/CADKernel/Mesh/Structure/FaceMesh.h
--- a/Engine/Source/Runtime/Datasmith/CADKernel/Public/CADKernel/Mesh/Structure/FaceMesh.h
+++ b/Engine/Source/Runtime/Datasmith/CADKernel/Public/CADKernel/Mesh/Structure/FaceMesh.h
@@ -3,63 +3,63 @@
 
 #include "CADKernel/Mesh/Structure/Mesh.h"
 
-namespace CADKernel
+namespace UE::CADKernel
 {
-	class FModelMesh;
-	class FPoint;
-	class FTopologicalEntity;
-	class FTopologicalFace;
-
-	/**
-	 * Define the mesh of a topological face. This mesh is an array of triangle
-	 */
-	class CADKERNEL_API FFaceMesh : public FMesh
-	{
-	public:
+class FModelMesh;
+class FPoint;
+class FTopologicalEntity;
+class FTopologicalFace;
+
+/**
+ * Define the mesh of a topological face. This mesh is an array of triangle
+ */
+class CADKERNEL_API FFaceMesh : public FMesh
+{
+public:
 
-		/** Index of the 3 vertices of each triangle in the local VertexIndices i.e. Triangle(a,b,c) use VertexIndices[a], VertexIndices[b], VertexIndices[c]*/
-		TArray<int32> TrianglesVerticesIndex;
+	/** Index of the 3 vertices of each triangle in the local VertexIndices i.e. Triangle(a,b,c) use VertexIndices[a], VertexIndices[b], VertexIndices[c]*/
+	TArray<int32> TrianglesVerticesIndex;
 
-		/** Index of the surface mesh vertices in the MeshModel vertex set*/
-		TArray<int32> VerticesGlobalIndex;
+	/** Index of the surface mesh vertices in the MeshModel vertex set*/
+	TArray<int32> VerticesGlobalIndex;
 
-		/** Normals of the surface mesh vertices*/
-		TArray<FVector3f> Normals;
+	/** Normals of the surface mesh vertices*/
+	TArray<FVector3f> Normals;
 
-		/** UV coordinates of the surface mesh vertices*/
-		TArray<FVector2f> UVMap;
+	/** UV coordinates of the surface mesh vertices*/
+	TArray<FVector2f> UVMap;
 
-	public:
-		FFaceMesh(FModelMesh& InMeshModel, FTopologicalEntity& InTopologicalEntity)
-			: FMesh(InMeshModel, InTopologicalEntity)
-		{
-		}
+public:
+	FFaceMesh(FModelMesh& InMeshModel, FTopologicalEntity& InTopologicalEntity)
+		: FMesh(InMeshModel, InTopologicalEntity)
+	{
+	}
 
-		void Init(int32 TriangleNum, int32 VertexNum)
-		{
-			TrianglesVerticesIndex.Reserve(TriangleNum * 3);
+	void Init(int32 TriangleNum, int32 VertexNum)
+	{
+		TrianglesVerticesIndex.Reserve(TriangleNum * 3);
 
-			VerticesGlobalIndex.Reserve(VertexNum);
-			Normals.Reserve(VertexNum);
-			UVMap.Reserve(VertexNum);
-		}
+		VerticesGlobalIndex.Reserve(VertexNum);
+		Normals.Reserve(VertexNum);
+		UVMap.Reserve(VertexNum);
+	}
 
-		void AddTriangle(int32 IndexA, int32 IndexB, int32 IndexC)
+	void AddTriangle(int32 IndexA, int32 IndexB, int32 IndexC)
+	{
+		if (VerticesGlobalIndex[IndexA] == VerticesGlobalIndex[IndexB] || VerticesGlobalIndex[IndexA] == VerticesGlobalIndex[IndexC] || VerticesGlobalIndex[IndexB] == VerticesGlobalIndex[IndexC])
 		{
-			if (VerticesGlobalIndex[IndexA] == VerticesGlobalIndex[IndexB] || VerticesGlobalIndex[IndexA] == VerticesGlobalIndex[IndexC] || VerticesGlobalIndex[IndexB] == VerticesGlobalIndex[IndexC])
-			{
-				return;
-			}
-
-			TrianglesVerticesIndex.Add(IndexA);
-			TrianglesVerticesIndex.Add(IndexB);
-			TrianglesVerticesIndex.Add(IndexC);
+			return;
 		}
 
-		void InverseOrientation();
+		TrianglesVerticesIndex.Add(IndexA);
+		TrianglesVerticesIndex.Add(IndexB);
+		TrianglesVerticesIndex.Add(IndexC);
+	}
+
+	void InverseOrientation();
 
-		void GetNodeIdToCoordinates(TMap<int32, const FPoint*>& NodeIdToCoordinates) const;
+	void GetNodeIdToCoordinates(TMap<int32, const FPoint*>& NodeIdToCoordinates) const;
 
-	};
+};
 }
 
diff --git a/Engine/Source/Runtime/Datasmith/CADKernel/Public/CADKernel/Mesh/Structure/Mesh.h b/Engine/Source/Runtime/Datasmith/CADKernel/Public/CADKernel/Mesh/Structure/Mesh.h
--- a/Engine/Source/Runtime/Datasmith/CADKernel/Public/CADKernel/Mesh/Structure/Mesh.h
+++ b/Engine/Source/Runtime/Datasmith/CADKernel/Public/CADKernel/Mesh/Structure/Mesh.h
@@ -6,85 +6,85 @@
 #include "CADKernel/Core/Types.h"
 #include "CADKernel/Math/Point.h"
 
-namespace CADKernel
+namespace UE::CADKernel
 {
-	class FModelMesh;
-	class FNode;
-	class FTopologicalEntity;
+class FModelMesh;
+class FNode;
+class FTopologicalEntity;
 
-	class CADKERNEL_API FMesh : public FEntityGeom
-	{
-	protected:
-		FModelMesh& ModelMesh;
-		FTopologicalEntity& TopologicalEntity;
+class CADKERNEL_API FMesh : public FEntityGeom
+{
+protected:
+	FModelMesh& ModelMesh;
+	FTopologicalEntity& TopologicalEntity;
 
-		int32 StartNodeId;
-		int32 LastNodeIndex;
+	int32 StartNodeId;
+	int32 LastNodeIndex;
 
-		TArray<FPoint> NodeCoordinates;
-		int32 MeshModelIndex;
+	TArray<FPoint> NodeCoordinates;
+	int32 MeshModelIndex;
 
-	public:
+public:
 
-		FMesh(FModelMesh& InMeshModel, FTopologicalEntity& InTopologicalEntity)
-			: FEntityGeom()
-			, ModelMesh(InMeshModel)
-			, TopologicalEntity(InTopologicalEntity)
-		{
-		}
+	FMesh(FModelMesh& InMeshModel, FTopologicalEntity& InTopologicalEntity)
+		: FEntityGeom()
+		, ModelMesh(InMeshModel)
+		, TopologicalEntity(InTopologicalEntity)
+	{
+	}
 
-		virtual ~FMesh() = default;
+	virtual ~FMesh() = default;
 
 #ifdef CADKERNEL_DEV
-		virtual FInfoEntity& GetInfo(FInfoEntity&) const override;
+	virtual FInfoEntity& GetInfo(FInfoEntity&) const override;
 #endif
 
-		virtual EEntity GetEntityType() const override
-		{
-			return EEntity::Mesh;
-		}
-
-		TArray<FPoint>& GetNodeCoordinates()
-		{
-			return NodeCoordinates;
-		}
-
-		const TArray<FPoint>& GetNodeCoordinates() const
-		{
-			return NodeCoordinates;
-		}
-
-		int32 RegisterCoordinates();
-
-		const int32 GetStartVertexId() const
-		{
-			return StartNodeId;
-		}
-
-		const int32 GetLastVertexIndex() const
-		{
-			return LastNodeIndex;
-		}
-
-		const int32 GetIndexInMeshModel() const
-		{
-			return MeshModelIndex;
-		}
-
-		FModelMesh& GetMeshModel()
-		{
-			return ModelMesh;
-		}
-
-		const FModelMesh& GetMeshModel() const
-		{
-			return ModelMesh;
-		}
-
-		const FTopologicalEntity& GetGeometricEntity() const
-		{
-			return TopologicalEntity;
-		}
-	};
+	virtual EEntity GetEntityType() const override
+	{
+		return EEntity::Mesh;
+	}
+
+	TArray<FPoint>& GetNodeCoordinates()
+	{
+		return NodeCoordinates;
+	}
+
+	const TArray<FPoint>& GetNodeCoordinates() const
+	{
+		return NodeCoordinates;
+	}
+
+	int32 RegisterCoordinates();
+
+	const int32 GetStartVertexId() const
+	{
+		return StartNodeId;
+	}
+
+	const int32 GetLastVertexIndex() const
+	{
+		return LastNodeIndex;
+	}
+
+	const int32 GetIndexInMeshModel() const
+	{
+		return MeshModelIndex;
+	}
+
+	FModelMesh& GetMeshModel()
+	{
+		return ModelMesh;
+	}
+
+	const FModelMesh& GetMeshModel() const
+	{
+		return ModelMesh;
+	}
+
+	const FTopologicalEntity& GetGeometricEntity() const
+	{
+		return TopologicalEntity;
+	}
+};
 }
 
diff --git a/Engine/Source/Runtime/Datasmith/CADKernel/Public/CADKernel/Mesh/Structure/ModelMesh.h b/Engine/Source/Runtime/Datasmith/CADKernel/Public/CADKernel/Mesh/Structure/ModelMesh.h
--- a/Engine/Source/Runtime/Datasmith/CADKernel/Public/CADKernel/Mesh/Structure/ModelMesh.h
+++ b/Engine/Source/Runtime/Datasmith/CADKernel/Public/CADKernel/Mesh/Structure/ModelMesh.h
@@ -7,158 +7,158 @@
 #include "CADKernel/Core/Types.h"
 #include "CADKernel/Mesh/Structure/Mesh.h"
 
-namespace CADKernel
+namespace UE::CADKernel
 {
-	class FCriterion;
-	class FDatabase;
-	class FEdgeMesh;
-	class FFaceMesh;
-	class FMesh;
-	class FPoint;
-	class FSession;
-	class FVertexMesh;
+class FCriterion;
+class FDatabase;
+class FEdgeMesh;
+class FFaceMesh;
+class FMesh;
+class FPoint;
+class FSession;
+class FVertexMesh;
+
+class CADKERNEL_API FModelMesh : public FEntityGeom
+{
+	friend FEntity;
 
-	class CADKERNEL_API FModelMesh : public FEntityGeom
-	{
-		friend FEntity;
+	TArray<TSharedPtr<FCriterion>> Criteria;
 
-		TArray<TSharedPtr<FCriterion>> Criteria;
+	TArray<TArray<FPoint>*> GlobalPointCloud;
+	FIdent LastIdUsed;
 
-		TArray<TArray<FPoint>*> GlobalPointCloud;
-		FIdent LastIdUsed;
+	TArray<TSharedPtr<FVertexMesh>> VertexMeshes;
+	TArray<TSharedPtr<FEdgeMesh>> EdgeMeshes;
+	TArray<TSharedPtr<FFaceMesh>> FaceMeshes;
 
-		TArray<TSharedPtr<FVertexMesh>> VertexMeshes;
-		TArray<TSharedPtr<FEdgeMesh>> EdgeMeshes;
-		TArray<TSharedPtr<FFaceMesh>> FaceMeshes;
+	bool QuadAnalyse = false;
+	double MinSize = DOUBLE_SMALL_NUMBER;
+	double MaxSize = HUGE_VALUE;
+	double MaxAngle = DOUBLE_PI;
+	double Sag = HUGE_VALUE;
 
-		bool QuadAnalyse = false;
-		double MinSize = DOUBLE_SMALL_NUMBER;
-		double MaxSize = HUGE_VALUE;
-		double MaxAngle = DOUBLE_PI;
-		double Sag = HUGE_VALUE;
+	FModelMesh()
+		: FEntityGeom()
+		, LastIdUsed(0)
+	{
+	}
 
-		FModelMesh()
-			: FEntityGeom()
-			, LastIdUsed(0)
-		{
-		}
+public:
 
-	public:
-		
-		int32 GetFaceCount() const
-		{
-			return FaceMeshes.Num();
-		}
+	int32 GetFaceCount() const
+	{
+		return FaceMeshes.Num();
+	}
 
-		int32 GetVertexCount() const
-		{
-			return LastIdUsed;
-		}
+	int32 GetVertexCount() const
+	{
+		return LastIdUsed;
+	}
 
-		int32 GetTriangleCount() const;
+	int32 GetTriangleCount() const;
 
-		virtual void SpawnIdent(FDatabase& Database) override
+	virtual void SpawnIdent(FDatabase& Database) override
+	{
+		if (!FEntity::SetId(Database))
 		{
-			if (!FEntity::SetId(Database))
-			{
-				return;
-			}
+			return;
+		}
 
-			SpawnIdentOnEntities(VertexMeshes, Database);
-			SpawnIdentOnEntities(EdgeMeshes, Database);
-			SpawnIdentOnEntities(FaceMeshes, Database);
-		};
+		SpawnIdentOnEntities(VertexMeshes, Database);
+		SpawnIdentOnEntities(EdgeMeshes, Database);
+		SpawnIdentOnEntities(FaceMeshes, Database);
+	};
 
-		virtual void ResetMarkersRecursively() 
-		{
-			ResetMarkers();
-			ResetMarkersRecursivelyOnEntities(VertexMeshes);
-			ResetMarkersRecursivelyOnEntities(EdgeMeshes  );
-			ResetMarkersRecursivelyOnEntities(FaceMeshes  );
-		};
+	virtual void ResetMarkersRecursively()
+	{
+		ResetMarkers();
+		ResetMarkersRecursivelyOnEntities(VertexMeshes);
+		ResetMarkersRecursivelyOnEntities(EdgeMeshes);
+		ResetMarkersRecursivelyOnEntities(FaceMeshes);
+	};
 
 
-		virtual EEntity GetEntityType() const override
-		{
-			return EEntity::MeshModel;
-		}
+	virtual EEntity GetEntityType() const override
+	{
+		return EEntity::MeshModel;
+	}
 
-		const TArray<TSharedPtr<FCriterion>>& GetCriteria() const
-		{
-			return Criteria;
-		}
+	const TArray<TSharedPtr<FCriterion>>& GetCriteria() const
+	{
+		return Criteria;
+	}
 
-		void AddCriterion(TSharedPtr<FCriterion>& Criterion);
+	void AddCriterion(TSharedPtr<FCriterion>& Criterion);
 
-		double GetGeometricTolerance() const
-		{
-			return 0.02; //Session->GetGeometricTolerance();
-		}
+	double GetGeometricTolerance() const
+	{
+		return 0.02; //Session->GetGeometricTolerance();
+	}
 
-		double GetMinSize() const
-		{
-			return MinSize;
-		}
+	double GetMinSize() const
+	{
+		return MinSize;
+	}
 
-		double GetMaxSize() const
-		{
-			return MaxSize;
-		}
+	double GetMaxSize() const
+	{
+		return MaxSize;
+	}
 
-		double GetAngleCriteria()
-		{
-			return MaxAngle;
-		}
+	double GetAngleCriteria()
+	{
+		return MaxAngle;
+	}
 
-		double GetSag() const
-		{
-			return Sag;
-		}
+	double GetSag() const
+	{
+		return Sag;
+	}
 
-		void AddMesh(TSharedRef<FVertexMesh> Mesh)
-		{
-			VertexMeshes.Add(Mesh);
-		}
+	void AddMesh(TSharedRef<FVertexMesh> Mesh)
+	{
+		VertexMeshes.Add(Mesh);
+	}
 
-		void AddMesh(TSharedRef<FEdgeMesh> Mesh)
-		{
-			EdgeMeshes.Add(Mesh);
-		}
+	void AddMesh(TSharedRef<FEdgeMesh> Mesh)
+	{
+		EdgeMeshes.Add(Mesh);
+	}
 
-		void AddMesh(TSharedRef<FFaceMesh> Mesh)
-		{
-			FaceMeshes.Add(Mesh);
-		}
+	void AddMesh(TSharedRef<FFaceMesh> Mesh)
+	{
+		FaceMeshes.Add(Mesh);
+	}
 
-		void RegisterCoordinates(TArray<FPoint>& Coordinates, int32& OutStartVertexId, int32& OutIndex)
-		{
-			OutIndex = (int32)GlobalPointCloud.Num();
-			OutStartVertexId = LastIdUsed;
+	void RegisterCoordinates(TArray<FPoint>& Coordinates, int32& OutStartVertexId, int32& OutIndex)
+	{
+		OutIndex = (int32)GlobalPointCloud.Num();
+		OutStartVertexId = LastIdUsed;
 
-			GlobalPointCloud.Add(&Coordinates);
-			LastIdUsed += (int32)Coordinates.Num();
-		}
+		GlobalPointCloud.Add(&Coordinates);
+		LastIdUsed += (int32)Coordinates.Num();
+	}
 
-		const int32 GetIndexOfVertexFromId(const int32 Id) const;
-		const int32 GetIndexOfEdgeFromId(const int32 Id) const;
-		const int32 GetIndexOfSurfaceFromId(const int32 Id) const;
+	const int32 GetIndexOfVertexFromId(const int32 Id) const;
+	const int32 GetIndexOfEdgeFromId(const int32 Id) const;
+	const int32 GetIndexOfSurfaceFromId(const int32 Id) const;
 
-		const TSharedPtr<FVertexMesh> GetMeshOfVertexNodeId(const int32 Id) const;
+	const TSharedPtr<FVertexMesh> GetMeshOfVertexNodeId(const int32 Id) const;
 
-		void GetNodeCoordinates(TArray<FPoint>& NodeCoordinates) const;
-		void GetNodeCoordinates(TArray<FVector3f>& NodeCoordinates) const;
+	void GetNodeCoordinates(TArray<FPoint>& NodeCoordinates) const;
+	void GetNodeCoordinates(TArray<FVector3f>& NodeCoordinates) const;
 
-		const TArray<TSharedPtr<FMesh>>& GetMeshes() const;
+	const TArray<TSharedPtr<FMesh>>& GetMeshes() const;
 
-		const TArray<TSharedPtr<FFaceMesh>>& GetFaceMeshes() const 
-		{
-			return FaceMeshes;
-		}
+	const TArray<TSharedPtr<FFaceMesh>>& GetFaceMeshes() const
+	{
+		return FaceMeshes;
+	}
 
 #ifdef CADKERNEL_DEV
-		virtual FInfoEntity& GetInfo(FInfoEntity&) const override;
+	virtual FInfoEntity& GetInfo(FInfoEntity&) const override;
 #endif
-	};
+};
 
 }
 
diff --git a/Engine/Source/Runtime/Datasmith/CADKernel/Public/CADKernel/Mesh/Structure/VertexMesh.h b/Engine/Source/Runtime/Datasmith/CADKernel/Public/CADKernel/Mesh/Structure/VertexMesh.h
--- a/Engine/Source/Runtime/Datasmith/CADKernel/Public/CADKernel/Mesh/Structure/VertexMesh.h
+++ b/Engine/Source/Runtime/Datasmith/CADKernel/Public/CADKernel/Mesh/Structure/VertexMesh.h
@@ -3,24 +3,24 @@
 
 #include "CADKernel/Mesh/Structure/Mesh.h"
 
-namespace CADKernel
+namespace UE::CADKernel
 {
-	class FModelMesh;
+class FModelMesh;
 
-	class CADKERNEL_API FVertexMesh : public FMesh
-	{
-	protected:
+class CADKERNEL_API FVertexMesh : public FMesh
+{
+protected:
 
-	public:
-		FVertexMesh(FModelMesh& Model, FTopologicalEntity& TopologicalEntity)
-			: FMesh(Model, TopologicalEntity)
-		{
-		}
+public:
+	FVertexMesh(FModelMesh& Model, FTopologicalEntity& TopologicalEntity)
+		: FMesh(Model, TopologicalEntity)
+	{
+	}
 
-		int32 GetMesh() const
-		{
-			return StartNodeId;
-		}
-	};
+	int32 GetMesh() const
+	{
+		return StartNodeId;
+	}
+};
 }
 
diff --git a/Engine/Source/Runtime/Datasmith/CADKernel/Public/CADKernel/Topo/Model.h b/Engine/Source/Runtime/Datasmith/CADKernel/Public/CADKernel/Topo/Model.h
--- a/Engine/Source/Runtime/Datasmith/CADKernel/Public/CADKernel/Topo/Model.h
+++ b/Engine/Source/Runtime/Datasmith/CADKernel/Public/CADKernel/Topo/Model.h
@@ -7,204 +7,204 @@
 #include "CADKernel/Topo/TopologicalShapeEntity.h"
 #include "CoreTypes.h"
 
-namespace CADKernel
+namespace UE::CADKernel
 {
 
-	class FBody;
-	class FDatabase;
-	class FGroup;
-	class FTopologicalEdge;
-	class FTopologicalEntity;
-	class FTopologicalFace;
-	class FTopologicalVertex;
+class FBody;
+class FDatabase;
+class FGroup;
+class FTopologicalEdge;
+class FTopologicalEntity;
+class FTopologicalFace;
+class FTopologicalVertex;
 
-	class CADKERNEL_API FModel : public FTopologicalShapeEntity
-	{
-		friend FEntity;
-
-	protected:
+class CADKERNEL_API FModel : public FTopologicalShapeEntity
+{
+	friend FEntity;
 
-		TArray<TSharedPtr<FBody>> Bodies;
-		TArray<TSharedPtr<FTopologicalFace>> Faces;
+protected:
 
-		FModel()
-		{
-			Bodies.Reserve(100);
-			Faces.Reserve(100);
-		}
+	TArray<TSharedPtr<FBody>> Bodies;
+	TArray<TSharedPtr<FTopologicalFace>> Faces;
 
-	public:
+	FModel()
+	{
+		Bodies.Reserve(100);
+		Faces.Reserve(100);
+	}
 
-		virtual void Serialize(FCADKernelArchive& Ar) override
-		{
-			FTopologicalShapeEntity::Serialize(Ar);
-			SerializeIdents(Ar, (TArray<TSharedPtr<FEntity>>&) Bodies);
-			SerializeIdents(Ar, (TArray<TSharedPtr<FEntity>>&) Faces);
+public:
 
-			if(Ar.IsLoading())
-			{
-				//ensureCADKernel(Archive.ArchiveModel == nullptr);
-				if (Ar.ArchiveModel == nullptr)
-				{
-					Ar.ArchiveModel = this;
-				}
-			}
-		}
+	virtual void Serialize(FCADKernelArchive& Ar) override
+	{
+		FTopologicalShapeEntity::Serialize(Ar);
+		SerializeIdents(Ar, (TArray<TSharedPtr<FEntity>>&) Bodies);
+		SerializeIdents(Ar, (TArray<TSharedPtr<FEntity>>&) Faces);
 
-		virtual void SpawnIdent(FDatabase& Database) override
+		if (Ar.IsLoading())
 		{
-			if (!FEntity::SetId(Database))
+			//ensureCADKernel(Archive.ArchiveModel == nullptr);
+			if (Ar.ArchiveModel == nullptr)
 			{
-				return;
+				Ar.ArchiveModel = this;
 			}
-
-			SpawnIdentOnEntities(Bodies, Database);
-			SpawnIdentOnEntities(Faces, Database);
 		}
+	}
 
-		virtual void ResetMarkersRecursively() override
+	virtual void SpawnIdent(FDatabase& Database) override
+	{
+		if (!FEntity::SetId(Database))
 		{
-			ResetMarkers();
-			ResetMarkersRecursivelyOnEntities(Bodies);
-			ResetMarkersRecursivelyOnEntities(Faces);
+			return;
 		}
 
+		SpawnIdentOnEntities(Bodies, Database);
+		SpawnIdentOnEntities(Faces, Database);
+	}
+
+	virtual void ResetMarkersRecursively() override
+	{
+		ResetMarkers();
+		ResetMarkersRecursivelyOnEntities(Bodies);
+		ResetMarkersRecursivelyOnEntities(Faces);
+	}
+
 #ifdef CADKERNEL_DEV
-		virtual FInfoEntity& GetInfo(FInfoEntity&) const override;
+	virtual FInfoEntity& GetInfo(FInfoEntity&) const override;
 #endif
 
-		virtual EEntity GetEntityType() const override
-		{
-			return EEntity::Model;
-		}
+	virtual EEntity GetEntityType() const override
+	{
+		return EEntity::Model;
+	}
 
-		void AddEntity(TSharedRef<FTopologicalEntity> InEntity);
+	void AddEntity(TSharedRef<FTopologicalEntity> InEntity);
 
-		void Add(const TSharedPtr<FTopologicalFace>& InFace)
-		{
-			Faces.Add(InFace);
-		}
+	void Add(const TSharedPtr<FTopologicalFace>& InFace)
+	{
+		Faces.Add(InFace);
+	}
 
-		void Append(TArray<TSharedPtr<FTopologicalFace>>& InNewFaces)
-		{
-			Faces.Append(InNewFaces);
-		}
+	void Append(TArray<TSharedPtr<FTopologicalFace>>& InNewFaces)
+	{
+		Faces.Append(InNewFaces);
+	}
 
-		void Append(TArray<TSharedPtr<FBody>>& InNewBody)
-		{
-			Bodies.Append(InNewBody);
-		}
+	void Append(TArray<TSharedPtr<FBody>>& InNewBody)
+	{
+		Bodies.Append(InNewBody);
+	}
 
-		void Add(const TSharedPtr<FBody>& InBody)
-		{
-			Bodies.Add(InBody);
-		}
+	void Add(const TSharedPtr<FBody>& InBody)
+	{
+		Bodies.Add(InBody);
+	}
 
-		void Empty()
-		{
-			Bodies.Empty();
-			Faces.Empty();
-		}
+	void Empty()
+	{
+		Bodies.Empty();
+		Faces.Empty();
+	}
 
-		void RemoveFace(TSharedPtr<FTopologicalFace> InToplologicalFace)
-		{
-			Faces.Remove(InToplologicalFace);
-		}
+	void RemoveFace(TSharedPtr<FTopologicalFace> InToplologicalFace)
+	{
+		Faces.Remove(InToplologicalFace);
+	}
 
-		void RemoveBody(TSharedPtr<FBody> InBody)
-		{
-			Bodies.Remove(InBody);
-		}
+	void RemoveBody(TSharedPtr<FBody> InBody)
+	{
+		Bodies.Remove(InBody);
+	}
 
-		void RemoveBody(const FBody* InBody)
+	void RemoveBody(const FBody* InBody)
+	{
+		int32 Index = Bodies.IndexOfByPredicate([&](const TSharedPtr<FBody>& Body) { return (InBody == Body.Get()); });
+		if (Index != INDEX_NONE)
 		{
-			int32 Index = Bodies.IndexOfByPredicate([&](const TSharedPtr<FBody>& Body){ return (InBody == Body.Get()); });
-			if (Index != INDEX_NONE)
-			{
-				Bodies.RemoveAt(Index);
-			}
+			Bodies.RemoveAt(Index);
 		}
+	}
 
-		void PrintBodyAndShellCount();
+	void PrintBodyAndShellCount();
 
-		void RemoveEntity(TSharedPtr<FTopologicalEntity> InEntity);
+	void RemoveEntity(TSharedPtr<FTopologicalEntity> InEntity);
 
-		bool Contains(TSharedPtr<FTopologicalEntity> InEntity);
+	bool Contains(TSharedPtr<FTopologicalEntity> InEntity);
 
-		/**
-		 * Copy the body and face arrays of other model
-		 */
-		void Copy(const TSharedPtr<FModel>& OtherModel)
-		{
-			Bodies.Append(OtherModel->Bodies);
-			Faces.Append(OtherModel->Faces);
-		}
-		
-		/**
-		 * Copy the body and face arrays of other model
-		 */
-		void Copy(const FModel& OtherModel)
-		{
-			Bodies.Append(OtherModel.Bodies);
-			Faces.Append(OtherModel.Faces);
-		}
-
-		/**
-		 * Copy the body and face arrays of other model
-		 * Empty other model arrays
-		 */
-		void Merge(FModel& OtherModel)
-		{
-			Copy(OtherModel);
-			OtherModel.Bodies.Empty();
-			OtherModel.Faces.Empty();
-		}
+	/**
+	 * Copy the body and face arrays of other model
+	 */
+	void Copy(const TSharedPtr<FModel>& OtherModel)
+	{
+		Bodies.Append(OtherModel->Bodies);
+		Faces.Append(OtherModel->Faces);
+	}
+
+	/**
+	 * Copy the body and face arrays of other model
+	 */
+	void Copy(const FModel& OtherModel)
+	{
+		Bodies.Append(OtherModel.Bodies);
+		Faces.Append(OtherModel.Faces);
+	}
+
+	/**
+	 * Copy the body and face arrays of other model
+	 * Empty other model arrays
+	 */
+	void Merge(FModel& OtherModel)
+	{
+		Copy(OtherModel);
+		OtherModel.Bodies.Empty();
+		OtherModel.Faces.Empty();
+	}
 
-		int32 EntityCount() const
-		{
-			return Bodies.Num() + Faces.Num();
-		}
+	int32 EntityCount() const
+	{
+		return Bodies.Num() + Faces.Num();
+	}
 
-		const TArray<TSharedPtr<FTopologicalFace>>& GetFaces() const
-		{
-			return Faces;
-		}
+	const TArray<TSharedPtr<FTopologicalFace>>& GetFaces() const
+	{
+		return Faces;
+	}
 
-		virtual void GetFaces(TArray<FTopologicalFace*>& OutFaces) override;
+	virtual void GetFaces(TArray<FTopologicalFace*>& OutFaces) override;
 
-		virtual int32 FaceCount() const override;
+	virtual int32 FaceCount() const override;
 
-		const TArray<TSharedPtr<FBody>>& GetBodies() const
-		{
-			return Bodies;
-		}
+	const TArray<TSharedPtr<FBody>>& GetBodies() const
+	{
+		return Bodies;
+	}
 
-		virtual void SpreadBodyOrientation() override;
+	virtual void SpreadBodyOrientation() override;
 
-		// Topo functions
+	// Topo functions
 
-		/**
-		 * Check topology of each body
-		 */
-		void CheckTopology();
+	/**
+	 * Check topology of each body
+	 */
+	void CheckTopology();
 
 #ifdef CADKERNEL_DEV
-		virtual void FillTopologyReport(FTopologyReport& Report) const override;
+	virtual void FillTopologyReport(FTopologyReport& Report) const override;
 #endif
 
-		/**
-		 * Fore each body
-		 */
-		void FixModelTopology(double JoiningTolerance);
+	/**
+	 * Fore each body
+	 */
+	void FixModelTopology(double JoiningTolerance);
 
-		void MergeInto(TSharedPtr<FBody> Body, TArray<TSharedPtr<FTopologicalEntity>>& InEntities);
+	void MergeInto(TSharedPtr<FBody> Body, TArray<TSharedPtr<FTopologicalEntity>>& InEntities);
 
-		/**
-		 * Fore each shell of each body, try to stitch topological gap
-		 */
-		void HealModelTopology(double JoiningTolerance);
+	/**
+	 * Fore each shell of each body, try to stitch topological gap
+	 */
+	void HealModelTopology(double JoiningTolerance);
 
-		void Orient();
-	};
+	void Orient();
+};
 
-} // namespace CADKernel
+} // namespace UE::CADKernel
diff --git a/Engine/Source/Runtime/Datasmith/CADKernel/Public/CADKernel/Topo/TopologicalEntity.h b/Engine/Source/Runtime/Datasmith/CADKernel/Public/CADKernel/Topo/TopologicalEntity.h
--- a/Engine/Source/Runtime/Datasmith/CADKernel/Public/CADKernel/Topo/TopologicalEntity.h
+++ b/Engine/Source/Runtime/Datasmith/CADKernel/Public/CADKernel/Topo/TopologicalEntity.h
@@ -6,66 +6,66 @@
 #include "CADKernel/Core/HaveStates.h"
 #include "CADKernel/Core/Types.h"
 
-namespace CADKernel
+namespace UE::CADKernel
 {
-	class FModelMesh;
-	class FTopologicalFace;
-	class FTopologyReport;
+class FModelMesh;
+class FTopologicalFace;
+class FTopologyReport;
 
-	class CADKERNEL_API FTopologicalEntity : public FEntity
-	{
-		friend class FCoreTechBridge;
+class CADKERNEL_API FTopologicalEntity : public FEntity
+{
+	friend class FCoreTechBridge;
 
-	protected:
-		FIdent CtKioId = 0;
+protected:
+	FIdent CtKioId = 0;
 
-	public:
+public:
 
-		FIdent GetKioId() const
-		{
-			return CtKioId;
-		}
+	FIdent GetKioId() const
+	{
+		return CtKioId;
+	}
+
+	virtual void Serialize(FCADKernelArchive& Ar) override
+	{
+		FEntity::Serialize(Ar);
+		Ar << CtKioId;
+	}
 
-		virtual void Serialize(FCADKernelArchive& Ar) override
-		{
-			FEntity::Serialize(Ar);
-			Ar << CtKioId;
-		}
-		
 #ifdef CADKERNEL_DEV
-		virtual FInfoEntity& GetInfo(FInfoEntity&) const override;
+	virtual FInfoEntity& GetInfo(FInfoEntity&) const override;
 #endif
 
-		const bool IsApplyCriteria() const
-		{
-			return ((States & EHaveStates::IsApplyCriteria) == EHaveStates::IsApplyCriteria);
-		}
-
-		virtual void SetApplyCriteria() const
-		{
-			States |= EHaveStates::IsApplyCriteria;
-		}
-		
-		virtual void ResetApplyCriteria()
-		{
-			States &= ~EHaveStates::IsApplyCriteria;
-		}
-
-		bool IsMeshed() const
-		{
-			return ((States & EHaveStates::IsMeshed) == EHaveStates::IsMeshed);
-		}
-		
-		virtual void SetMeshed()
-		{
-			States |= EHaveStates::IsMeshed;
-		}
-		
-		virtual void ResetMeshed()
-		{
-			States &= ~EHaveStates::IsMeshed;
-		}
-	};
-
-} // namespace CADKernel
+	const bool IsApplyCriteria() const
+	{
+		return ((States & EHaveStates::IsApplyCriteria) == EHaveStates::IsApplyCriteria);
+	}
+
+	virtual void SetApplyCriteria() const
+	{
+		States |= EHaveStates::IsApplyCriteria;
+	}
+
+	virtual void ResetApplyCriteria()
+	{
+		States &= ~EHaveStates::IsApplyCriteria;
+	}
+
+	bool IsMeshed() const
+	{
+		return ((States & EHaveStates::IsMeshed) == EHaveStates::IsMeshed);
+	}
+
+	virtual void SetMeshed()
+	{
+		States |= EHaveStates::IsMeshed;
+	}
+
+	virtual void ResetMeshed()
+	{
+		States &= ~EHaveStates::IsMeshed;
+	}
+};
+
+} // namespace UE::CADKernel
 
diff --git a/Engine/Source/Runtime/Datasmith/CADKernel/Public/CADKernel/Topo/TopologicalShapeEntity.h b/Engine/Source/Runtime/Datasmith/CADKernel/Public/CADKernel/Topo/TopologicalShapeEntity.h
--- a/Engine/Source/Runtime/Datasmith/CADKernel/Public/CADKernel/Topo/TopologicalShapeEntity.h
+++ b/Engine/Source/Runtime/Datasmith/CADKernel/Public/CADKernel/Topo/TopologicalShapeEntity.h
@@ -4,132 +4,132 @@
 #include "CADKernel/Core/MetadataDictionary.h"
 #include "CADKernel/Topo/TopologicalEntity.h"
 
-namespace CADKernel
+namespace UE::CADKernel
 {
-	class FModelMesh;
-	class FTopologicalFace;
-	class FTopologyReport;
+class FModelMesh;
+class FTopologicalFace;
+class FTopologyReport;
 
-	class CADKERNEL_API FTopologicalShapeEntity : public FTopologicalEntity
-	{
-	private:
-		FTopologicalShapeEntity* HostedBy = nullptr;
-		FMetadataDictionary Dictionary;
+class CADKERNEL_API FTopologicalShapeEntity : public FTopologicalEntity
+{
+private:
+	FTopologicalShapeEntity* HostedBy = nullptr;
+	FMetadataDictionary Dictionary;
 
-	public:
+public:
 
-		virtual void Serialize(FCADKernelArchive& Ar) override
-		{
-			FTopologicalEntity::Serialize(Ar);
-			Dictionary.Serialize(Ar);
-			SerializeIdent(Ar, &HostedBy);
-		}
+	virtual void Serialize(FCADKernelArchive& Ar) override
+	{
+		FTopologicalEntity::Serialize(Ar);
+		Dictionary.Serialize(Ar);
+		SerializeIdent(Ar, &HostedBy);
+	}
 
-		const FMetadataDictionary& GetMetadataDictionary() const
-		{
-			return Dictionary;
-		}
+	const FMetadataDictionary& GetMetadataDictionary() const
+	{
+		return Dictionary;
+	}
 
-		void ExtractMetaData(TMap<FString, FString>& OutMetaData) const
-		{
-			Dictionary.ExtractMetaData(OutMetaData);
-		}
+	void ExtractMetaData(TMap<FString, FString>& OutMetaData) const
+	{
+		Dictionary.ExtractMetaData(OutMetaData);
+	}
 
-		void CompleteMetadata();
+	void CompleteMetadata();
 
-		virtual int32 FaceCount() const = 0;
-		virtual void GetFaces(TArray<FTopologicalFace*>& OutFaces) = 0;
+	virtual int32 FaceCount() const = 0;
+	virtual void GetFaces(TArray<FTopologicalFace*>& OutFaces) = 0;
 
-		/**
-		 * Each face of model is set by its orientation. This allow to make oriented mesh and to keep the face orientation in topological function.
-		 * Marker2 of spread face is set. It must be reset after the process
-		 */
-		virtual void SpreadBodyOrientation() = 0;
+	/**
+	 * Each face of model is set by its orientation. This allow to make oriented mesh and to keep the face orientation in topological function.
+	 * Marker2 of spread face is set. It must be reset after the process
+	 */
+	virtual void SpreadBodyOrientation() = 0;
 
 #ifdef CADKERNEL_DEV
-		virtual void FillTopologyReport(FTopologyReport& Report) const = 0;
+	virtual void FillTopologyReport(FTopologyReport& Report) const = 0;
 #endif
 
-		FTopologicalShapeEntity* GetHost()
-		{
-			return HostedBy;
-		}
-
-		void SetHost(FTopologicalShapeEntity* Body)
-		{
-			HostedBy = Body;
-		}
-
-		void ResetHost()
-		{
-			HostedBy = nullptr;
-		}
-
-		void SetHostId(const int32 InHostId)
-		{
-			Dictionary.SetHostId(InHostId);
-		}
-
-		int32 GetHostId() const
-		{
-			return Dictionary.GetHostId();
-		}
-
-		void SetLayer(const int32 InLayerId)
-		{
-			Dictionary.SetLayer(InLayerId);
-		}
-
-		void SetName(const FString& InName)
-		{
-			Dictionary.SetName(InName);
-		}
-
-		bool HasName() const
-		{
-			return Dictionary.HasName();
-		}
-
-		const TCHAR* GetName() const
-		{
-			return Dictionary.GetName();
-		}
-
-		void SetColorId(const uint32& InColorId)
-		{
-			Dictionary.SetColorId(InColorId);
-		}
-
-		uint32 GetColorId() const
-		{
-			return Dictionary.GetColorId();
-		}
-
-		void SetMaterialId(const uint32& InMaterialId)
-		{
-			Dictionary.SetMaterialId(InMaterialId);
-		}
-
-		uint32 GetMaterialId() const
-		{
-			return Dictionary.GetMaterialId();
-		}
-
-		void SetPatchId(int32 InPatchId)
-		{
-			Dictionary.SetPatchId(InPatchId);
-		}
-
-		int32 GetPatchId() const 
-		{
-			return Dictionary.GetPatchId();
-		}
+	FTopologicalShapeEntity* GetHost()
+	{
+		return HostedBy;
+	}
+
+	void SetHost(FTopologicalShapeEntity* Body)
+	{
+		HostedBy = Body;
+	}
+
+	void ResetHost()
+	{
+		HostedBy = nullptr;
+	}
+
+	void SetHostId(const int32 InHostId)
+	{
+		Dictionary.SetHostId(InHostId);
+	}
+
+	int32 GetHostId() const
+	{
+		return Dictionary.GetHostId();
+	}
+
+	void SetLayer(const int32 InLayerId)
+	{
+		Dictionary.SetLayer(InLayerId);
+	}
+
+	void SetName(const FString& InName)
+	{
+		Dictionary.SetName(InName);
+	}
+
+	bool HasName() const
+	{
+		return Dictionary.HasName();
+	}
+
+	const TCHAR* GetName() const
+	{
+		return Dictionary.GetName();
+	}
+
+	void SetColorId(const uint32& InColorId)
+	{
+		Dictionary.SetColorId(InColorId);
+	}
+
+	uint32 GetColorId() const
+	{
+		return Dictionary.GetColorId();
+	}
+
+	void SetMaterialId(const uint32& InMaterialId)
+	{
+		Dictionary.SetMaterialId(InMaterialId);
+	}
+
+	uint32 GetMaterialId() const
+	{
+		return Dictionary.GetMaterialId();
+	}
+
+	void SetPatchId(int32 InPatchId)
+	{
+		Dictionary.SetPatchId(InPatchId);
+	}
+
+	int32 GetPatchId() const
+	{
+		return Dictionary.GetPatchId();
+	}
 
 #ifdef CADKERNEL_DEV
-		virtual FInfoEntity& GetInfo(FInfoEntity&) const override;
+	virtual FInfoEntity& GetInfo(FInfoEntity&) const override;
 #endif
 
-	};
+};
 
-} // namespace CADKernel
+} // namespace UE::CADKernel
 
