Rename CADKernel namespace into UE::CADKernel

diff --git a/Engine/Source/Runtime/Datasmith/CADKernel/Public/CADKernel/Topo/TopologicalVertex.h b/Engine/Source/Runtime/Datasmith/CADKernel/Public/CADKernel/Topo/TopologicalVertex.h
--- a/Engine/Source/Runtime/Datasmith/CADKernel/Public/CADKernel/Topo/TopologicalVertex.h
+++ b/Engine/Source/Runtime/Datasmith/CADKernel/Public/CADKernel/Topo/TopologicalVertex.h
@@ -9,221 +9,221 @@
 #include "CADKernel/Topo/Linkable.h"
 #include "CADKernel/Topo/TopologicalLink.h"
 
-namespace CADKernel
+namespace UE::CADKernel
 {
-	class FDatabase;
-	class FModelMesh;
-	class FTopologicalEdge;
-	class FVertexMesh;
+class FDatabase;
+class FModelMesh;
+class FTopologicalEdge;
+class FVertexMesh;
 
-	class CADKERNEL_API FTopologicalVertex : public TLinkable<FTopologicalVertex, FVertexLink>
-	{
-		friend class FEntity;
-		friend class FVertexLink;
+class CADKERNEL_API FTopologicalVertex : public TLinkable<FTopologicalVertex, FVertexLink>
+{
+	friend class FEntity;
+	friend class FVertexLink;
 
-	protected:
+protected:
 
-		TArray<FTopologicalEdge*> ConnectedEdges;
-		FPoint Coordinates;
-		TSharedPtr<FVertexMesh> Mesh;
+	TArray<FTopologicalEdge*> ConnectedEdges;
+	FPoint Coordinates;
+	TSharedPtr<FVertexMesh> Mesh;
 
-		FTopologicalVertex(const FPoint& InCoordinates)
-			: Coordinates(InCoordinates)
-		{
-		}
+	FTopologicalVertex(const FPoint& InCoordinates)
+		: Coordinates(InCoordinates)
+	{
+	}
 
-		FTopologicalVertex() = default;
+	FTopologicalVertex() = default;
 
-	public:
+public:
 
-		static TSharedRef<FTopologicalVertex> Make(const FPoint& InCoordinate)
-		{
-			TSharedRef<FTopologicalVertex> Vertex = FEntity::MakeShared<FTopologicalVertex>(InCoordinate);
-			Vertex->Finalize();
-			return Vertex;
-		}
+	static TSharedRef<FTopologicalVertex> Make(const FPoint& InCoordinate)
+	{
+		TSharedRef<FTopologicalVertex> Vertex = FEntity::MakeShared<FTopologicalVertex>(InCoordinate);
+		Vertex->Finalize();
+		return Vertex;
+	}
 
-		virtual void Serialize(FCADKernelArchive& Ar) override
+	virtual void Serialize(FCADKernelArchive& Ar) override
+	{
+		if (Ar.IsSaving())
 		{
-			if(Ar.IsSaving())
-			{
-				ensureCADKernel(ConnectedEdges.Num());
-			}
-
-			TLinkable<FTopologicalVertex, FVertexLink>::Serialize(Ar);
-			Ar.Serialize(Coordinates);
-			SerializeIdents(Ar, ConnectedEdges);
+			ensureCADKernel(ConnectedEdges.Num());
 		}
 
-		virtual void SpawnIdent(FDatabase& Database) override;
+		TLinkable<FTopologicalVertex, FVertexLink>::Serialize(Ar);
+		Ar.Serialize(Coordinates);
+		SerializeIdents(Ar, ConnectedEdges);
+	}
 
-		virtual void ResetMarkersRecursively() override
-		{
-			ResetMarkers();
-		}
+	virtual void SpawnIdent(FDatabase& Database) override;
+
+	virtual void ResetMarkersRecursively() override
+	{
+		ResetMarkers();
+	}
 
 
 #ifdef CADKERNEL_DEV
-		virtual FInfoEntity& GetInfo(FInfoEntity&) const override;
+	virtual FInfoEntity& GetInfo(FInfoEntity&) const override;
 #endif
 
-		virtual EEntity GetEntityType() const override
-		{
-			return EEntity::TopologicalVertex;
-		}
+	virtual EEntity GetEntityType() const override
+	{
+		return EEntity::TopologicalVertex;
+	}
 
-		/**
-		 * @return the 3d coordinate of the barycenter of the twin vertices
-		 */
-		inline const FPoint& GetBarycenter() const
+	/**
+	 * @return the 3d coordinate of the barycenter of the twin vertices
+	 */
+	inline const FPoint& GetBarycenter() const
+	{
+		if (TopologicalLink.IsValid() && TopologicalLink->GetTwinEntityNum() > 1)
 		{
-			if (TopologicalLink.IsValid() && TopologicalLink->GetTwinEntityNum() > 1)
-			{
-				return TopologicalLink->GetBarycenter();
-			}
-			return Coordinates;
+			return TopologicalLink->GetBarycenter();
 		}
+		return Coordinates;
+	}
 
-		/**
-		 * @return the 3d coordinates of the vertex (prefere GetBarycenter())
-		 */
-		const FPoint& GetCoordinates() const
-		{
-			return Coordinates;
-		}
+	/**
+	 * @return the 3d coordinates of the vertex (prefere GetBarycenter())
+	 */
+	const FPoint& GetCoordinates() const
+	{
+		return Coordinates;
+	}
 
-		void SetCoordinates(const FPoint& NewCoordinates)
+	void SetCoordinates(const FPoint& NewCoordinates)
+	{
+		if (GetLink()->GetTwinEntityNum() > 1)
 		{
-			if (GetLink()->GetTwinEntityNum() > 1)
-			{
-				// Update barycenter
-				FPoint BaryCenter = GetLink()->GetBarycenter() * (double)GetLink()->GetTwinEntityNum();
-				BaryCenter -= Coordinates;
-				BaryCenter += NewCoordinates;
-				BaryCenter /= (double)GetLink()->GetTwinEntityNum();
-				GetLink()->SetBarycenter(BaryCenter);
-			}
-			else
-			{
-				GetLink()->SetBarycenter(NewCoordinates);
-			}
-			Coordinates = NewCoordinates;
+			// Update barycenter
+			FPoint BaryCenter = GetLink()->GetBarycenter() * (double)GetLink()->GetTwinEntityNum();
+			BaryCenter -= Coordinates;
+			BaryCenter += NewCoordinates;
+			BaryCenter /= (double)GetLink()->GetTwinEntityNum();
+			GetLink()->SetBarycenter(BaryCenter);
 		}
-
-		double Distance(const FTopologicalVertex& OtherVertex) const
+		else
 		{
-			return Coordinates.Distance(OtherVertex.Coordinates);
+			GetLink()->SetBarycenter(NewCoordinates);
 		}
+		Coordinates = NewCoordinates;
+	}
 
-		double SquareDistance(const FTopologicalVertex& OtherVertex) const
-		{
-			return Coordinates.SquareDistance(OtherVertex.Coordinates);
-		}
+	double Distance(const FTopologicalVertex& OtherVertex) const
+	{
+		return Coordinates.Distance(OtherVertex.Coordinates);
+	}
 
-		double SquareDistance(const FPoint& Point) const
-		{
-			return Coordinates.SquareDistance(Point);
-		}
+	double SquareDistance(const FTopologicalVertex& OtherVertex) const
+	{
+		return Coordinates.SquareDistance(OtherVertex.Coordinates);
+	}
+
+	double SquareDistance(const FPoint& Point) const
+	{
+		return Coordinates.SquareDistance(Point);
+	}
 
-		TSharedRef<FVertexMesh> GetOrCreateMesh(FModelMesh& MeshModel);
+	TSharedRef<FVertexMesh> GetOrCreateMesh(FModelMesh& MeshModel);
 
-		const TSharedRef<FVertexMesh> GetMesh() const
+	const TSharedRef<FVertexMesh> GetMesh() const
+	{
+		if (!IsActiveEntity())
 		{
-			if (!IsActiveEntity())
-			{
-				return GetLinkActiveEntity()->GetMesh();
-			}
-			return Mesh.ToSharedRef();
+			return GetLinkActiveEntity()->GetMesh();
 		}
+		return Mesh.ToSharedRef();
+	}
 
-		void Link(FTopologicalVertex& InEntity);
+	void Link(FTopologicalVertex& InEntity);
 
-		void UnlinkTo(FTopologicalVertex& Entity);
+	void UnlinkTo(FTopologicalVertex& Entity);
 
-		virtual void RemoveFromLink() override
+	virtual void RemoveFromLink() override
+	{
+		if (TopologicalLink.IsValid())
 		{
-			if (TopologicalLink.IsValid())
-			{
-				TopologicalLink->RemoveEntity(*this);
-				TopologicalLink->ComputeBarycenter();
-				TopologicalLink.Reset();
-			}
+			TopologicalLink->RemoveEntity(*this);
+			TopologicalLink->ComputeBarycenter();
+			TopologicalLink.Reset();
 		}
+	}
 
-		void DeleteIfIsolated()
+	void DeleteIfIsolated()
+	{
+		if (ConnectedEdges.Num() == 0)
 		{
-			if (ConnectedEdges.Num() == 0)
+			if (TopologicalLink.IsValid())
 			{
-				if (TopologicalLink.IsValid())
+				TopologicalLink->RemoveEntity(*this);
+				if (!TopologicalLink->IsDeleted())
 				{
-					TopologicalLink->RemoveEntity(*this);
-					if(!TopologicalLink->IsDeleted())
-					{
-						TopologicalLink->ComputeBarycenter();
-					}
-					TopologicalLink.Reset();
+					TopologicalLink->ComputeBarycenter();
 				}
-				SetDeleted();
+				TopologicalLink.Reset();
 			}
-		}
+			SetDeleted();
+		}
+	}
+
+	bool IsBorderVertex();
+
+	void AddConnectedEdge(FTopologicalEdge& Edge);
+	void RemoveConnectedEdge(FTopologicalEdge& Edge);
+
+	/**
+	 * Mandatory: to browse all the connected edges, you have to browse the connected edges of all the twin vertices
+	 * for (TWeakPtr<FTopologicalVertex> TwinVertex : Vertex->GetTwinsEntities())
+	 * {
+	 *    for (TWeakPtr<FTopologicalEdge> ConnectedEdge : TwinVertex.Pin()->GetDirectConnectedEdges())
+	 *    {
+	 *       ...
+	 *    }
+	 *  }
+	 */
+	const TArray<FTopologicalEdge*>& GetDirectConnectedEdges() const
+	{
+		return ConnectedEdges;
+	}
 
-		bool IsBorderVertex();
-
-		void AddConnectedEdge(FTopologicalEdge& Edge);
-		void RemoveConnectedEdge(FTopologicalEdge& Edge);
-
-		/**
-		 * Mandatory: to browse all the connected edges, you have to browse the connected edges of all the twin vertices
-		 * for (TWeakPtr<FTopologicalVertex> TwinVertex : Vertex->GetTwinsEntities())
-		 * {
-		 *    for (TWeakPtr<FTopologicalEdge> ConnectedEdge : TwinVertex.Pin()->GetDirectConnectedEdges())
-		 *    {
-		 *       ...
-		 *    }
-		 *  }
-		 */
-		const TArray<FTopologicalEdge*>& GetDirectConnectedEdges() const 
+	void GetConnectedEdges(TArray<FTopologicalEdge*>& OutConnectedEdges) const
+	{
+		if (!TopologicalLink.IsValid())
 		{
-			return ConnectedEdges;
+			OutConnectedEdges = ConnectedEdges;
 		}
-
-		void GetConnectedEdges(TArray<FTopologicalEdge*>& OutConnectedEdges) const
+		else
 		{
-			if (!TopologicalLink.IsValid())
-			{
-				OutConnectedEdges = ConnectedEdges;
-			}
-			else
+			OutConnectedEdges.Reserve(100);
+			for (const FTopologicalVertex* Vertex : GetLink()->GetTwinEntities())
 			{
-				OutConnectedEdges.Reserve(100);
-				for (const FTopologicalVertex* Vertex : GetLink()->GetTwinEntities())
-				{
-					OutConnectedEdges.Append(Vertex->ConnectedEdges);
-				}
+				OutConnectedEdges.Append(Vertex->ConnectedEdges);
 			}
 		}
+	}
 
-		const int32 ConnectedEdgeCount()
+	const int32 ConnectedEdgeCount()
+	{
+		if (!TopologicalLink.IsValid())
 		{
-			if (!TopologicalLink.IsValid())
-			{
-				return ConnectedEdges.Num();
-			}
-			else
+			return ConnectedEdges.Num();
+		}
+		else
+		{
+			int32 Count = 0;
+			for (const FTopologicalVertex* Vertex : GetLink()->GetTwinEntities())
 			{
-				int32 Count = 0;
-				for (const FTopologicalVertex* Vertex : GetLink()->GetTwinEntities())
-				{
-					Count += Vertex->ConnectedEdges.Num();
-				}
-				return Count;
+				Count += Vertex->ConnectedEdges.Num();
 			}
+			return Count;
 		}
+	}
 
-		/**
-		 * 
-		 */
-		void GetConnectedEdges(const FTopologicalVertex& OtherVertex, TArray<FTopologicalEdge*>& Edges) const;
-	};
+	/**
+	 *
+	 */
+	void GetConnectedEdges(const FTopologicalVertex& OtherVertex, TArray<FTopologicalEdge*>& Edges) const;
+};
 
-} // namespace CADKernel
+} // namespace UE::CADKernel
diff --git a/Engine/Source/Runtime/Datasmith/CADKernel/Public/CADKernel/UI/Console.h b/Engine/Source/Runtime/Datasmith/CADKernel/Public/CADKernel/UI/Console.h
--- a/Engine/Source/Runtime/Datasmith/CADKernel/Public/CADKernel/UI/Console.h
+++ b/Engine/Source/Runtime/Datasmith/CADKernel/Public/CADKernel/UI/Console.h
@@ -5,20 +5,20 @@
 
 #include "CADKernel/UI/Message.h"
 
-namespace CADKernel
+namespace UE::CADKernel
 {
-	class CADKERNEL_API FConsole
-	{
-	public:
+class CADKERNEL_API FConsole
+{
+public:
 
-		virtual ~FConsole()
-		{
-		}
+	virtual ~FConsole()
+	{
+	}
 
-		virtual void Print(const TCHAR* Texte, EVerboseLevel VerboseLevel = EVerboseLevel::NoVerbose)
-		{}
+	virtual void Print(const TCHAR* Texte, EVerboseLevel VerboseLevel = EVerboseLevel::NoVerbose)
+	{}
 
-	};
+};
 
-} // namespace CADKernel
+} // namespace UE::CADKernel
 
diff --git a/Engine/Source/Runtime/Datasmith/CADKernel/Public/CADKernel/UI/Display.h b/Engine/Source/Runtime/Datasmith/CADKernel/Public/CADKernel/UI/Display.h
--- a/Engine/Source/Runtime/Datasmith/CADKernel/Public/CADKernel/UI/Display.h
+++ b/Engine/Source/Runtime/Datasmith/CADKernel/Public/CADKernel/UI/Display.h
@@ -9,7 +9,7 @@
 
 class FString;
 
-namespace CADKernel
+namespace UE::CADKernel
 {
 
 class FAABB2D;
@@ -226,7 +226,7 @@ void DisplayOrientedPolyline(const TArray<TPoint>& Points, EVisuProperty Propert
 	Open3DDebugSegment(0);
 	for (int32 Index = 1; Index < Points.Num(); ++Index)
 	{
-		DisplaySegment(Points[Index-1], Points[Index], Index, Property, true);
+		DisplaySegment(Points[Index - 1], Points[Index], Index, Property, true);
 	}
 	Close3DDebugSegment();
 }
@@ -266,5 +266,5 @@ void DrawSegmentOrientation(const TPoint& Point1, const TPoint& Point2, EVisuPro
 CADKERNEL_API void DrawIsoCurves(const FTopologicalFace& Face);
 
 
-} // namespace CADKernel
+} // namespace UE::CADKernel
 
diff --git a/Engine/Source/Runtime/Datasmith/CADKernel/Public/CADKernel/UI/Progress.h b/Engine/Source/Runtime/Datasmith/CADKernel/Public/CADKernel/UI/Progress.h
--- a/Engine/Source/Runtime/Datasmith/CADKernel/Public/CADKernel/UI/Progress.h
+++ b/Engine/Source/Runtime/Datasmith/CADKernel/Public/CADKernel/UI/Progress.h
@@ -3,120 +3,120 @@
 
 #include "CADKernel/Core/Types.h"
 
-namespace CADKernel
+namespace UE::CADKernel
 {
-	class FProgress;
-	class FProgressBase;
+class FProgress;
+class FProgressBase;
 
-	class CADKERNEL_API FProgressManager
-	{
-		friend class FProgress;
+class CADKERNEL_API FProgressManager
+{
+	friend class FProgress;
 
-	protected:
+protected:
 
-		FProgress* RootProgress;
-		FProgress* CurrentProgress;
+	FProgress* RootProgress;
+	FProgress* CurrentProgress;
 
-		FProgress* GetRoot() const
-		{
-			return RootProgress;
-		}
+	FProgress* GetRoot() const
+	{
+		return RootProgress;
+	}
 
-		virtual void ActivateProgressBar(bool bActive)
-		{}
+	virtual void ActivateProgressBar(bool bActive)
+	{}
 
-		/**
-		 * Update the progress bar 
-		 */
-		virtual void Update()
-		{}
+	/**
+	 * Update the progress bar
+	 */
+	virtual void Update()
+	{}
 
-		double GetProgression() const;
+	double GetProgression() const;
 
-		/**
-		 * @return a string describing the current step hierarchy 
-		 */
-		virtual FString GetCurrentStep() const;
+	/**
+	 * @return a string describing the current step hierarchy
+	 */
+	virtual FString GetCurrentStep() const;
 
-		virtual void SetCurrentProgress(FProgress* InProgress);
+	virtual void SetCurrentProgress(FProgress* InProgress);
 
-		FProgress* GetCurrentProgress() const
-		{
-			return CurrentProgress;
-		}
+	FProgress* GetCurrentProgress() const
+	{
+		return CurrentProgress;
+	}
 
-	public:
-		FProgressManager()
-			: RootProgress(nullptr)
-			, CurrentProgress(nullptr)
-		{}
-	};
+public:
+	FProgressManager()
+		: RootProgress(nullptr)
+		, CurrentProgress(nullptr)
+	{}
+};
 
-	class CADKERNEL_API FProgress
-	{
-		friend class FProgressManager;
-	private:
-		FString Name;
+class CADKERNEL_API FProgress
+{
+	friend class FProgressManager;
+private:
+	FString Name;
 
-		FProgress* Parent = nullptr;
-		FProgress* UnderlyingProgress = nullptr;
+	FProgress* Parent = nullptr;
+	FProgress* UnderlyingProgress = nullptr;
 
-		int32 StepCount;
-		int32 Progression = 0;
+	int32 StepCount;
+	int32 Progression = 0;
 
-	public:
-		FProgress(int32 InStepCount, const FString& InStepName = FString());
+public:
+	FProgress(int32 InStepCount, const FString& InStepName = FString());
 
-		FProgress(const FString& StepName = FString())
-			: FProgress(1, StepName)
-		{
-		}
+	FProgress(const FString& StepName = FString())
+		: FProgress(1, StepName)
+	{
+	}
 
-		virtual ~FProgress();
+	virtual ~FProgress();
 
-		const FString& GetName()
-		{
-			return Name;
-		}
+	const FString& GetName()
+	{
+		return Name;
+	}
 
-		FProgress* GetParent() const
-		{
-			return Parent;
-		}
+	FProgress* GetParent() const
+	{
+		return Parent;
+	}
 
-		/**
-		 * Increase the advancement of of one step
-		 */
-		void Increase(int32 StepSize = 1);
+	/**
+	 * Increase the advancement of of one step
+	 */
+	void Increase(int32 StepSize = 1);
 
-	protected:
+protected:
 
-		bool IsRoot() const
-		{
-			return Parent == nullptr;
-		}
+	bool IsRoot() const
+	{
+		return Parent == nullptr;
+	}
 
-		void AddUnderlying(FProgress* Progress)
-		{
-			UnderlyingProgress = Progress;
-		}
+	void AddUnderlying(FProgress* Progress)
+	{
+		UnderlyingProgress = Progress;
+	}
 
-		void UnderlyingFinished(FProgress* Progress)
-		{
-			Progression++;
-			UnderlyingProgress = nullptr;
-		}
+	void UnderlyingFinished(FProgress* Progress)
+	{
+		Progression++;
+		UnderlyingProgress = nullptr;
+	}
 
-		double GetProgression() const
-		{
-			double Ratio = ((double)Progression / (double)StepCount);
+	double GetProgression() const
+	{
+		double Ratio = ((double)Progression / (double)StepCount);
 
-			if (UnderlyingProgress)
-			{
-				Ratio += UnderlyingProgress->GetProgression() / (double)StepCount;
-			}
-			return Ratio;
+		if (UnderlyingProgress)
+		{
+			Ratio += UnderlyingProgress->GetProgression() / (double)StepCount;
 		}
-	};
+		return Ratio;
+	}
+};
 }
 
diff --git a/Engine/Source/Runtime/Datasmith/CADKernel/Public/CADKernel/UI/Visu.h b/Engine/Source/Runtime/Datasmith/CADKernel/Public/CADKernel/UI/Visu.h
--- a/Engine/Source/Runtime/Datasmith/CADKernel/Public/CADKernel/UI/Visu.h
+++ b/Engine/Source/Runtime/Datasmith/CADKernel/Public/CADKernel/UI/Visu.h
@@ -6,162 +6,162 @@
 #include "CADKernel/Core/Types.h"
 #include "CADKernel/Math/Point.h"
 
-namespace CADKernel
+namespace UE::CADKernel
 {
-	enum EVisuProperty
-	{
-		BBox = 0,          // Green see CADKernelUI::properties    
-		Point,             // Blue
-		Element,           // Brown
-		Curve,             // Blue
-		Iso,               // Red
-		ControlLine,       // Purple
-		EdgeMesh,          // blue
-		NodeMesh,          // blue
-		BorderEdge,        // Green
-		NonManifoldEdge,   // Red          
-		BorderTopo,        // Dark brown
-		ControlPoint,      // Purple		 
-		YellowPoint,			
-		YellowCurve,			
-		BluePoint,				
-		BlueCurve,				
-		RedPoint,				 
-		RedCurve,				 
-		PurplePoint,			
-		PurpleCurve,			
-		GreenPoint,				
-		GreenCurve,				
-		PinkPoint,
-		PinkCurve,
-		OrangePoint,
-		OrangeCurve,
-		Last,
-	} ;
-
-	extern const TCHAR* VisuPropertyNames[];
-
-	class FMesh;
-
-	class CADKERNEL_API FVisuParameters : public FParameters
-	{
-	public:
-
-		FParameter ChordError;
-
-		/**
-		 *	Number of isos along U to display on a surface or topological face
-		 */ 
-		FParameter IsoUNumber;
-
-		/**
-		 *	Number of isos along V to display on a surface or topological face
-		 */
-		FParameter IsoVNumber;
-
-		FParameter bDisplayCADOrient;
-		FParameter bDisplayMeshOrient;
-		FParameter bDisplayAxis;
-		FParameter bDisplayNormals;
-
-		FParameter NormalLength;
-
-		FVisuParameters()
-			: FParameters(10)
-			, ChordError(TEXT("DiscretizationError"), 0.02, *this)
-			, IsoUNumber(TEXT("IsoUNumber"), 3, *this)
-			, IsoVNumber(TEXT("IsoVNumber"), 3, *this)
-			, bDisplayCADOrient(TEXT("DisplayCadOrient"), false, *this)
-			, bDisplayMeshOrient(TEXT("DisplayMeshOrient"), true, *this)
-			, bDisplayAxis(TEXT("DisplayAxis"), false, *this)
-			, bDisplayNormals(TEXT("DisplayNormals"), false, *this)
-			, NormalLength(TEXT("NormalLength"), 10., *this)
-		{}
-	};
-
-	class CADKERNEL_API FVisu
-	{
-	protected:
-		FVisuParameters Parameters;
+enum EVisuProperty
+{
+	BBox = 0,          // Green see CADKernelUI::properties    
+	Point,             // Blue
+	Element,           // Brown
+	Curve,             // Blue
+	Iso,               // Red
+	ControlLine,       // Purple
+	EdgeMesh,          // blue
+	NodeMesh,          // blue
+	BorderEdge,        // Green
+	NonManifoldEdge,   // Red          
+	BorderTopo,        // Dark brown
+	ControlPoint,      // Purple		 
+	YellowPoint,
+	YellowCurve,
+	BluePoint,
+	BlueCurve,
+	RedPoint,
+	RedCurve,
+	PurplePoint,
+	PurpleCurve,
+	GreenPoint,
+	GreenCurve,
+	PinkPoint,
+	PinkCurve,
+	OrangePoint,
+	OrangeCurve,
+	Last,
+};
+
+extern const TCHAR* VisuPropertyNames[];
+
+class FMesh;
+
+class CADKERNEL_API FVisuParameters : public FParameters
+{
+public:
+
+	FParameter ChordError;
+
+	/**
+	 *	Number of isos along U to display on a surface or topological face
+	 */
+	FParameter IsoUNumber;
+
+	/**
+	 *	Number of isos along V to display on a surface or topological face
+	 */
+	FParameter IsoVNumber;
+
+	FParameter bDisplayCADOrient;
+	FParameter bDisplayMeshOrient;
+	FParameter bDisplayAxis;
+	FParameter bDisplayNormals;
+
+	FParameter NormalLength;
+
+	FVisuParameters()
+		: FParameters(10)
+		, ChordError(TEXT("DiscretizationError"), 0.02, *this)
+		, IsoUNumber(TEXT("IsoUNumber"), 3, *this)
+		, IsoVNumber(TEXT("IsoVNumber"), 3, *this)
+		, bDisplayCADOrient(TEXT("DisplayCadOrient"), false, *this)
+		, bDisplayMeshOrient(TEXT("DisplayMeshOrient"), true, *this)
+		, bDisplayAxis(TEXT("DisplayAxis"), false, *this)
+		, bDisplayNormals(TEXT("DisplayNormals"), false, *this)
+		, NormalLength(TEXT("NormalLength"), 10., *this)
+	{}
+};
+
+class CADKERNEL_API FVisu
+{
+protected:
+	FVisuParameters Parameters;
 
-		/**
-		 * Number of graphic session opened
-		 */ 
-		int32 SessionNum;
+	/**
+	 * Number of graphic session opened
+	 */
+	int32 SessionNum;
 
-		void NewSession()
-		{
-			ensureCADKernel(SessionNum >= 0);
-			++SessionNum;
-		}
+	void NewSession()
+	{
+		ensureCADKernel(SessionNum >= 0);
+		++SessionNum;
+	}
 
-		void EndSession()
+	void EndSession()
+	{
+		--SessionNum;
+		if (SessionNum < 0)
 		{
-			--SessionNum;
-			if (SessionNum < 0)
-			{
-				SessionNum = 0;
-			}
+			SessionNum = 0;
 		}
+	}
 
-	public:
+public:
 
-		FVisu()
-		{
-			SessionNum = 0;
-		}
+	FVisu()
+	{
+		SessionNum = 0;
+	}
 
-		virtual ~FVisu() = default;
+	virtual ~FVisu() = default;
 
-		int32 GetPropertyCount() 
-		{
-			return EVisuProperty::Last;
-		}
+	int32 GetPropertyCount()
+	{
+		return EVisuProperty::Last;
+	}
 
-		FVisuParameters* GetParameters() 
-		{ 
-			return &Parameters;
-		}
+	FVisuParameters* GetParameters()
+	{
+		return &Parameters;
+	}
 
-		virtual void NewDB(const TCHAR* InName) 
-		{}
+	virtual void NewDB(const TCHAR* InName)
+	{}
 
-		virtual void Open3DDebugSession(const TCHAR* SessionName, const TArray<FIdent>& Ids) 
-		{}
+	virtual void Open3DDebugSession(const TCHAR* SessionName, const TArray<FIdent>& Ids)
+	{}
 
-		virtual void Close3DDebugSession() 
-		{}
+	virtual void Close3DDebugSession()
+	{}
 
-		virtual void Open3DDebugSegment(FIdent Ident) 
-		{}
+	virtual void Open3DDebugSegment(FIdent Ident)
+	{}
 
-		virtual void Close3DDebugSegment() 
-		{}
+	virtual void Close3DDebugSegment()
+	{}
 
-		virtual void UpdateViewer() 
-		{}
+	virtual void UpdateViewer()
+	{}
 
-		virtual void DrawPoint(const FPoint& Point, EVisuProperty InProperty = EVisuProperty::BluePoint)
-		{}
+	virtual void DrawPoint(const FPoint& Point, EVisuProperty InProperty = EVisuProperty::BluePoint)
+	{}
 
-		virtual void DrawPoint(const FPoint2D& Point, EVisuProperty InProperty = EVisuProperty::BluePoint)
-		{}
+	virtual void DrawPoint(const FPoint2D& Point, EVisuProperty InProperty = EVisuProperty::BluePoint)
+	{}
 
-		virtual void DrawPoint(const FPointH& Point, EVisuProperty InProperty = EVisuProperty::BluePoint)
-		{}
+	virtual void DrawPoint(const FPointH& Point, EVisuProperty InProperty = EVisuProperty::BluePoint)
+	{}
 
-		virtual void DrawElement(int32 Dimension, TArray<FPoint>& Points, EVisuProperty InProperty = EVisuProperty::Element)
-		{}
+	virtual void DrawElement(int32 Dimension, TArray<FPoint>& Points, EVisuProperty InProperty = EVisuProperty::Element)
+	{}
 
-		virtual void DrawPolyline(const TArray<FPoint>& Points, EVisuProperty InProperty = EVisuProperty::BlueCurve)
-		{}
+	virtual void DrawPolyline(const TArray<FPoint>& Points, EVisuProperty InProperty = EVisuProperty::BlueCurve)
+	{}
 
-		virtual void DrawPolyline(const TArray<FPoint2D>& Points, EVisuProperty InProperty = EVisuProperty::BlueCurve)
-		{}
+	virtual void DrawPolyline(const TArray<FPoint2D>& Points, EVisuProperty InProperty = EVisuProperty::BlueCurve)
+	{}
 
-		virtual void DrawMesh(FIdent MeshId) 
-		{}
-	};
+	virtual void DrawMesh(FIdent MeshId)
+	{}
+};
 
-} // namespace CADKernel
+} // namespace UE::CADKernel
 
diff --git a/Engine/Source/Runtime/Datasmith/CADKernel/Public/CADKernel/Utils/Cache.h b/Engine/Source/Runtime/Datasmith/CADKernel/Public/CADKernel/Utils/Cache.h
--- a/Engine/Source/Runtime/Datasmith/CADKernel/Public/CADKernel/Utils/Cache.h
+++ b/Engine/Source/Runtime/Datasmith/CADKernel/Public/CADKernel/Utils/Cache.h
@@ -1,89 +1,89 @@
 // Copyright Epic Games, Inc. All Rights Reserved.
 #pragma once
 
-namespace CADKernel
+namespace UE::CADKernel
 {
-	template<class ObjectType>
-	class CADKERNEL_API TCache
-	{
-	protected:
-		ObjectType Value;
-		bool bReady;
+template<class ObjectType>
+class CADKERNEL_API TCache
+{
+protected:
+	ObjectType Value;
+	bool bReady;
 
-	public:
-		TCache()
-			: Value()
-			, bReady(false)
-		{
-		}
+public:
+	TCache()
+		: Value()
+		, bReady(false)
+	{
+	}
 
-		template<typename... InArgTypes>
-		TCache(InArgTypes&&... Args)
-			: Value(Forward<InArgTypes>(Args)...)
-			, bReady(true)
-		{
-		}
+	template<typename... InArgTypes>
+	TCache(InArgTypes&&... Args)
+		: Value(Forward<InArgTypes>(Args)...)
+		, bReady(true)
+	{
+	}
 
-		friend FArchive& operator<<(FArchive& Ar, TCache<ObjectType>& Cache)
-		{
-			Ar << Cache.Value;
-			Ar << Cache.bReady;
-			return Ar;
-		}
+	friend FArchive& operator<<(FArchive& Ar, TCache<ObjectType>& Cache)
+	{
+		Ar << Cache.Value;
+		Ar << Cache.bReady;
+		return Ar;
+	}
 
-		void Set(const ObjectType& NewValue)
-		{
-			Value = NewValue;
-			bReady = true;
-		}
+	void Set(const ObjectType& NewValue)
+	{
+		Value = NewValue;
+		bReady = true;
+	}
 
-		template<typename... InArgTypes>
-		void Set(InArgTypes&&... Args)
-		{
-			Value.Set(Forward<InArgTypes>(Args)...);
-			bReady = true;
-		}
+	template<typename... InArgTypes>
+	void Set(InArgTypes&&... Args)
+	{
+		Value.Set(Forward<InArgTypes>(Args)...);
+		bReady = true;
+	}
 
-		void operator=(const ObjectType& NewValue)
-		{
-			Value = NewValue;
-			bReady = true;
-		}
+	void operator=(const ObjectType& NewValue)
+	{
+		Value = NewValue;
+		bReady = true;
+	}
 
-		void operator+=(const ObjectType& NewValue)
-		{
-			ensureCADKernel(bReady);
-			Value += NewValue;
-		}
+	void operator+=(const ObjectType& NewValue)
+	{
+		ensureCADKernel(bReady);
+		Value += NewValue;
+	}
 
-		ObjectType* operator->()
-		{
-			return &Value;
-		}
+	ObjectType* operator->()
+	{
+		return &Value;
+	}
 
-		operator ObjectType& ()
-		{
-			return Value;
-		}
+	operator ObjectType& ()
+	{
+		return Value;
+	}
 
-		const ObjectType& operator*() const
-		{
-			return Value;
-		}
+	const ObjectType& operator*() const
+	{
+		return Value;
+	}
 
-		void SetReady()
-		{
-			bReady = true;
-		}
+	void SetReady()
+	{
+		bReady = true;
+	}
 
-		bool IsValid() const
-		{
-			return bReady;
-		}
+	bool IsValid() const
+	{
+		return bReady;
+	}
 
-		void Empty()
-		{
-			bReady = false;
-		}
-	};
+	void Empty()
+	{
+		bReady = false;
+	}
+};
 }
diff --git a/Engine/Source/Runtime/Datasmith/CADKernel/Public/CADKernel/Utils/IndexOfCoordinateFinder.h b/Engine/Source/Runtime/Datasmith/CADKernel/Public/CADKernel/Utils/IndexOfCoordinateFinder.h
--- a/Engine/Source/Runtime/Datasmith/CADKernel/Public/CADKernel/Utils/IndexOfCoordinateFinder.h
+++ b/Engine/Source/Runtime/Datasmith/CADKernel/Public/CADKernel/Utils/IndexOfCoordinateFinder.h
@@ -3,145 +3,145 @@
 
 #include "CADKernel/Core/Types.h"
 
-namespace CADKernel
+namespace UE::CADKernel
 {
-	/**
-	 * Finder of the index of the segment of the polyline coordinates containing the input coordinate
-	 * This abstract class return the lower index of the segment i.e.
-	 * InPolylineCoordinates[OutIndex] <= InCoordinate <= InPolylineCoordinates[OutIndex + 1]
-	 * As input coordinates are increasing, the algorithme is optimize by restarting form the last coordinate found.
-	 *
-	 * Class derived from it allows to optimized polyline functions needing the index of the segment containing the input coordinate (@see ApproximatePoint).
-	 * @see FLinearFinder to optimize the search when the next points to search is nearest the previous.
-	 * @see FFinderByDichotomy to optimize the search when there is no simple rule, so a dichotomy algorithm is used.
-	 * This class is dedicated for a unique or mess up set of input coordinates
-	 */
-	class CADKERNEL_API FIndexOfCoordinateFinder
+/**
+ * Finder of the index of the segment of the polyline coordinates containing the input coordinate
+ * This abstract class return the lower index of the segment i.e.
+ * InPolylineCoordinates[OutIndex] <= InCoordinate <= InPolylineCoordinates[OutIndex + 1]
+ * As input coordinates are increasing, the algorithme is optimize by restarting form the last coordinate found.
+ *
+ * Class derived from it allows to optimized polyline functions needing the index of the segment containing the input coordinate (@see ApproximatePoint).
+ * @see FLinearFinder to optimize the search when the next points to search is nearest the previous.
+ * @see FFinderByDichotomy to optimize the search when there is no simple rule, so a dichotomy algorithm is used.
+ * This class is dedicated for a unique or mess up set of input coordinates
+ */
+class CADKERNEL_API FIndexOfCoordinateFinder
+{
+protected:
+	const TArray<double>& Coordinates;
+public:
+	FIndexOfCoordinateFinder(const TArray<double>& InPolylineCoordinates)
+		: Coordinates(InPolylineCoordinates)
 	{
-	protected:
-		const TArray<double>& Coordinates;
-	public:
-		FIndexOfCoordinateFinder(const TArray<double>& InPolylineCoordinates)
-			: Coordinates(InPolylineCoordinates)
-		{
-		}
+	}
 
-		virtual int32 Find(const double InCoordinate) = 0;
-	};
+	virtual int32 Find(const double InCoordinate) = 0;
+};
 
-	class CADKERNEL_API FLinearFinder : public FIndexOfCoordinateFinder
+class CADKERNEL_API FLinearFinder : public FIndexOfCoordinateFinder
+{
+	int32& Index;
+
+public:
+	FLinearFinder(const TArray<double>& InPolylineCoordinates, int32& InOutIndex)
+		: FIndexOfCoordinateFinder(InPolylineCoordinates)
+		, Index(InOutIndex)
 	{
-		int32& Index;
+	}
 
-	public:
-		FLinearFinder(const TArray<double>& InPolylineCoordinates, int32& InOutIndex)
-			: FIndexOfCoordinateFinder(InPolylineCoordinates)
-			, Index(InOutIndex)
+	virtual int32 Find(const double InCoordinate) override
+	{
+		if ((Coordinates.Num() < 3) || (InCoordinate < Coordinates[0]))
 		{
+			Index = 0;
+			return Index;
 		}
 
-		virtual int32 Find(const double InCoordinate) override
+		if (InCoordinate > Coordinates.Last())
 		{
-			if ((Coordinates.Num() < 3) || (InCoordinate < Coordinates[0]))
-			{
-				Index = 0;
-				return Index;
-			}
-
-			if (InCoordinate > Coordinates.Last())
-			{
-				Index = Coordinates.Num() - 2;
-				return Index;
-			}
+			Index = Coordinates.Num() - 2;
+			return Index;
+		}
 
-			if (InCoordinate < Coordinates[Index])
+		if (InCoordinate < Coordinates[Index])
+		{
+			// find the previous segment containing the coordinate
+			while (Index > 0)
 			{
-				// find the previous segment containing the coordinate
-				while(Index > 0)
+				--Index;
+				if (InCoordinate + DOUBLE_SMALL_NUMBER >= Coordinates[Index])
 				{
-					--Index;
-					if (InCoordinate + DOUBLE_SMALL_NUMBER >= Coordinates[Index])
-					{
-						return Index;
-					}
+					return Index;
 				}
-				return Index;
 			}
+			return Index;
+		}
 
-			// find the next segment containing the coordinate
-			for (; Index < Coordinates.Num() - 2; ++Index)
+		// find the next segment containing the coordinate
+		for (; Index < Coordinates.Num() - 2; ++Index)
+		{
+			if (InCoordinate <= Coordinates[Index + 1])
 			{
-				if ( InCoordinate <= Coordinates[Index + 1])
-				{
-					return Index;
-				}
+				return Index;
 			}
-			Index = Coordinates.Num() - 2;
-			return Index;
 		}
-	};
+		Index = Coordinates.Num() - 2;
+		return Index;
+	}
+};
+
+class CADKERNEL_API FDichotomyFinder : public FIndexOfCoordinateFinder
+{
+
+public:
+	FDichotomyFinder(const TArray<double>& InPolylineCoordinates)
+		: FIndexOfCoordinateFinder(InPolylineCoordinates)
+		, StartLower(0)
+		, StartUpper(InPolylineCoordinates.Num() - 2)
+	{
+	}
 
-	class CADKERNEL_API FDichotomyFinder : public FIndexOfCoordinateFinder
+	FDichotomyFinder(const TArray<double>& InPolylineCoordinates, int32 InStartLowerBound, int32 InStartUpperBound)
+		: FIndexOfCoordinateFinder(InPolylineCoordinates)
+		, StartLower(InStartLowerBound)
+		, StartUpper(InStartUpperBound)
 	{
+	}
 
-	public:
-		FDichotomyFinder(const TArray<double>& InPolylineCoordinates)
-			: FIndexOfCoordinateFinder(InPolylineCoordinates)
-			, StartLower(0)
-			, StartUpper(InPolylineCoordinates.Num() - 2)
+	int32 StartLower;
+	int32 StartUpper;
+
+	virtual int32 Find(const double InCoordinate) override
+	{
+		if (StartUpper < 0)
 		{
+			StartUpper = Coordinates.Num() - 2;
 		}
 
-		FDichotomyFinder(const TArray<double>& InPolylineCoordinates, int32 InStartLowerBound, int32 InStartUpperBound)
-			: FIndexOfCoordinateFinder(InPolylineCoordinates)
-			, StartLower(InStartLowerBound)
-			, StartUpper(InStartUpperBound)
+		if ((StartUpper - StartLower < 1) || (InCoordinate < Coordinates[StartLower + 1]))
 		{
+			return StartLower;
 		}
 
-		int32 StartLower;
-		int32 StartUpper;
-
-		virtual int32 Find(const double InCoordinate) override
+		if (InCoordinate > Coordinates[StartUpper])
 		{
-			if (StartUpper < 0)
-			{
-				StartUpper = Coordinates.Num() - 2;
-			}
+			StartLower = StartUpper;
+			return StartUpper;
+		}
+
+		int32 LowerBound = StartLower;
+		int32 UpperBound = StartUpper;
 
-			if ((StartUpper - StartLower < 1) || (InCoordinate < Coordinates[StartLower + 1]))
+		while (LowerBound <= UpperBound)
+		{
+			int32 Center = (LowerBound + UpperBound) >> 1;
+			if (InCoordinate >= Coordinates[Center] && InCoordinate <= Coordinates[Center + 1])
 			{
-				return StartLower;
+				StartLower = Center;
+				return Center;
 			}
-
-			if (InCoordinate > Coordinates[StartUpper])
+			else if (InCoordinate < Coordinates[Center])
 			{
-				StartLower = StartUpper;
-				return StartUpper;
+				UpperBound = Center;
 			}
-
-			int32 LowerBound = StartLower;
-			int32 UpperBound = StartUpper;
-
-			while (LowerBound <= UpperBound)
+			else
 			{
-				int32 Center = (LowerBound + UpperBound) >> 1;
-				if (InCoordinate >= Coordinates[Center] && InCoordinate <= Coordinates[Center + 1])
-				{
-					StartLower = Center;
-					return Center;
-				}
-				else if (InCoordinate < Coordinates[Center])
-				{
-					UpperBound = Center;
-				}
-				else
-				{
-					LowerBound = Center + 1;
-				}
+				LowerBound = Center + 1;
 			}
-			return StartUpper;
 		}
-	};
-} // namespace CADKernel
+		return StartUpper;
+	}
+};
+} // namespace UE::CADKernel
 
diff --git a/Engine/Source/Runtime/Datasmith/CADKernel/Public/CADKernel/Utils/StringUtil.h b/Engine/Source/Runtime/Datasmith/CADKernel/Public/CADKernel/Utils/StringUtil.h
--- a/Engine/Source/Runtime/Datasmith/CADKernel/Public/CADKernel/Utils/StringUtil.h
+++ b/Engine/Source/Runtime/Datasmith/CADKernel/Public/CADKernel/Utils/StringUtil.h
@@ -14,7 +14,7 @@
 #include <stdio.h>
 #include <stdarg.h>
 
-namespace CADKernel::Utils
+namespace UE::CADKernel::Utils
 {
 
 inline std::string ToString(const TCHAR* WideString)
