Fixed reader's issue

diff --git a/Engine/Source/Runtime/Datasmith/DatasmithCore/Private/DatasmithSceneXmlReader.cpp b/Engine/Source/Runtime/Datasmith/DatasmithCore/Private/DatasmithSceneXmlReader.cpp
--- a/Engine/Source/Runtime/Datasmith/DatasmithCore/Private/DatasmithSceneXmlReader.cpp
+++ b/Engine/Source/Runtime/Datasmith/DatasmithCore/Private/DatasmithSceneXmlReader.cpp
@@ -1615,7 +1615,7 @@ void FDatasmithSceneXmlReader::ParseMaterialInstance(FXmlNode* InNode, TSharedPt
 
 	for ( const FXmlAttribute& Attribute : InNode->GetAttributes() )
 	{
-		if (Attribute.GetTag() == DATASMITH_MATERIALINSTANCENAME)
+		if (Attribute.GetTag() == DATASMITH_MATERIALINSTANCETYPE)
 		{
 			EDatasmithReferenceMaterialType MaterialType = (EDatasmithReferenceMaterialType)FMath::Clamp( ValueFromString<int32>( Attribute.GetValue() ), 0, (int32)EDatasmithReferenceMaterialType::Count - 1 );
 
