Datasmith DirectLink Decal import fix

diff --git a/Engine/Source/Runtime/Datasmith/DatasmithCore/Private/DatasmithSceneGraphSharedState.h b/Engine/Source/Runtime/Datasmith/DatasmithCore/Private/DatasmithSceneGraphSharedState.h
--- a/Engine/Source/Runtime/Datasmith/DatasmithCore/Private/DatasmithSceneGraphSharedState.h
+++ b/Engine/Source/Runtime/Datasmith/DatasmithCore/Private/DatasmithSceneGraphSharedState.h
@@ -80,6 +80,7 @@ public:
 
 	virtual void SetNodes(DirectLink::IReferenceResolutionProvider& ResolutionProvider, const TArray<DirectLink::FSceneGraphId>& NodeIds) override
 	{
+		Inner.Reset(NodeIds.Num());
 		Algo::Transform(NodeIds, Inner, [&](DirectLink::FSceneGraphId NodeId)
 		{
 			// #ue_directlink_quality validate type through property "type" ?
