DatasmithSDK support for pattern parameters

diff --git a/Engine/Source/Programs/Enterprise/Datasmith/DatasmithSDKSample/Scenes/Cloth.cpp b/Engine/Source/Programs/Enterprise/Datasmith/DatasmithSDKSample/Scenes/Cloth.cpp
--- a/Engine/Source/Programs/Enterprise/Datasmith/DatasmithSDKSample/Scenes/Cloth.cpp
+++ b/Engine/Source/Programs/Enterprise/Datasmith/DatasmithSDKSample/Scenes/Cloth.cpp
@@ -98,14 +98,12 @@ TSharedPtr<IDatasmithScene> FClothScene::Export()
 			}
 		}
 
-#if 0 // require tvariant serial
+		// #ue_ds_sdk_todo Proper data API
+		// eg. SetVertices, SetVertexParameter...
 		FParameterData& Param = Pattern.Parameters.AddDefaulted_GetRef();
+		Param.Name = TEXT("Test");
 		Param.Data.Set<TArray<double>>({1,2,3});
 
-		FParameterData& Paramf = Pattern.Parameters.AddDefaulted_GetRef();
-		Paramf.Data.Set<TArray<float>>({1,2,3});
-#endif
-
 		FDatasmithMeshExporter MeshExporter;
 		FString ClothFilePath = FString(DatasmithSceneExporter.GetAssetsOutputPath()) / ClothElement->GetName();
 
diff --git a/Engine/Source/Runtime/Datasmith/DatasmithCore/Private/DatasmithCloth.cpp b/Engine/Source/Runtime/Datasmith/DatasmithCore/Private/DatasmithCloth.cpp
--- a/Engine/Source/Runtime/Datasmith/DatasmithCore/Private/DatasmithCloth.cpp
+++ b/Engine/Source/Runtime/Datasmith/DatasmithCore/Private/DatasmithCloth.cpp
@@ -9,6 +9,7 @@
 enum EDatasmithClothSerializationVersion
 {
     EDCSV_Base = 0,
+    EDCSV_WithPatternParameters = 1,
 
 	// -----<new versions can be added before this line>-------------------------------------------------
     _EDCSV_LastPlusOne,
@@ -30,6 +31,38 @@ const FGuid FDatasmithClothSerializationVersion::GUID(0x28B01036, 0x66B4498F, 0x
 FCustomVersionRegistration GRegisterDatasmithClothCustomVersion(FDatasmithClothSerializationVersion::GUID, _EDCSV_Last, TEXT("DatasmithCloth"));
 
 
+FArchive& operator<<(FArchive& Ar, FParameterData& ParameterData)
+{
+	Ar.UsingCustomVersion(FDatasmithClothSerializationVersion::GUID);
+
+	static_assert(TVariantSize<decltype(ParameterData.Data)>::Value == 2, "Serialization code not synced with structure");
+	static_assert(ParameterData.Data.IndexOfType<TArray<float>>() == 0, "Serialization relies on this specific order");
+	static_assert(ParameterData.Data.IndexOfType<TArray<double>>() == 1, "Serialization relies on this specific order");
+
+	Ar << ParameterData.Target;
+	Ar << ParameterData.Data;
+
+	return Ar;
+}
+
+FArchive& operator<<(FArchive& Ar, FDatasmithClothPattern& Pattern)
+{
+	Ar.UsingCustomVersion(FDatasmithClothSerializationVersion::GUID);
+	int32 ClothSerialVersion = Ar.CustomVer(FDatasmithClothSerializationVersion::GUID);
+
+	Ar << Pattern.SimPosition;
+	Ar << Pattern.SimRestPosition;
+	Ar << Pattern.SimTriangleIndices;
+
+	if (ClothSerialVersion >= EDCSV_WithPatternParameters)
+	{
+		Ar << Pattern.Parameters;
+	}
+
+	return Ar;
+}
+
+
 FArchive& operator<<(FArchive& Ar, FDatasmithClothPresetProperty& Property)
 {
 	Ar.UsingCustomVersion(FDatasmithClothSerializationVersion::GUID);
@@ -59,15 +92,3 @@ FArchive& operator<<(FArchive& Ar, FDatasmithCloth& Cloth)
 
 	return Ar;
 }
-
-FArchive& operator<<(FArchive& Ar, FDatasmithClothPattern& Pattern)
-{
-	Ar.UsingCustomVersion(FDatasmithClothSerializationVersion::GUID);
-
-	Ar << Pattern.SimPosition;
-	Ar << Pattern.SimRestPosition;
-	Ar << Pattern.SimTriangleIndices;
-
-	return Ar;
-}
-
diff --git a/Engine/Source/Runtime/Datasmith/DatasmithCore/Public/DatasmithCloth.h b/Engine/Source/Runtime/Datasmith/DatasmithCore/Public/DatasmithCloth.h
--- a/Engine/Source/Runtime/Datasmith/DatasmithCore/Public/DatasmithCloth.h
+++ b/Engine/Source/Runtime/Datasmith/DatasmithCore/Public/DatasmithCloth.h
@@ -7,11 +7,25 @@
 #include "Math/UnrealMathSSE.h"
 #include "Math/Vector.h"
 #include "Math/Vector2D.h"
+#include "Misc/TVariant.h"
 #include "UObject/NameTypes.h"
 
 class FDatasmithMesh;
 
 
+struct FParameterData
+{
+	FString Name;
+
+	enum class ETarget { Vertex };
+	ETarget Target = ETarget::Vertex; // (also drives the expected number of values)
+
+	TVariant<TArray<float>, TArray<double>> Data;
+
+public:
+	friend FArchive& operator<<(FArchive& Ar, FParameterData& Pattern);
+};
+
 
 class DATASMITHCORE_API FDatasmithClothPattern
 {
@@ -20,6 +34,8 @@ public:
 	TArray<FVector3f> SimRestPosition;
 	TArray<uint32> SimTriangleIndices;
 
+	TArray<FParameterData> Parameters;
+
 public:
 	bool IsValid() { return SimRestPosition.Num() == SimPosition.Num() && SimTriangleIndices.Num() % 3 == 0 && SimTriangleIndices.Num(); }
 	friend FArchive& operator<<(FArchive& Ar, FDatasmithClothPattern& Pattern);
