Build fix for clang

diff --git a/Engine/Source/Runtime/Datasmith/DatasmithCore/Private/DatasmithCloth.cpp b/Engine/Source/Runtime/Datasmith/DatasmithCore/Private/DatasmithCloth.cpp
--- a/Engine/Source/Runtime/Datasmith/DatasmithCore/Private/DatasmithCloth.cpp
+++ b/Engine/Source/Runtime/Datasmith/DatasmithCore/Private/DatasmithCloth.cpp
@@ -36,8 +36,8 @@ FArchive& operator<<(FArchive& Ar, FParameterData& ParameterData)
 	Ar.UsingCustomVersion(FDatasmithClothSerializationVersion::GUID);
 
 	static_assert(TVariantSize<decltype(ParameterData.Data)>::Value == 2, "Serialization code not synced with structure");
-	static_assert(ParameterData.Data.IndexOfType<TArray<float>>() == 0, "Serialization relies on this specific order");
-	static_assert(ParameterData.Data.IndexOfType<TArray<double>>() == 1, "Serialization relies on this specific order");
+	static_assert(decltype(ParameterData.Data)::IndexOfType<TArray<float>>() == 0, "Serialization relies on this specific order");
+	static_assert(decltype(ParameterData.Data)::IndexOfType<TArray<double>>() == 1, "Serialization relies on this specific order");
 
 	Ar << ParameterData.Target;
 	Ar << ParameterData.Data;
