Enabled old serialization for forward compatibility with 5.0

diff --git a/Engine/Source/Runtime/Datasmith/DatasmithCore/Private/DatasmithMeshSerialization.cpp b/Engine/Source/Runtime/Datasmith/DatasmithCore/Private/DatasmithMeshSerialization.cpp
--- a/Engine/Source/Runtime/Datasmith/DatasmithCore/Private/DatasmithMeshSerialization.cpp
+++ b/Engine/Source/Runtime/Datasmith/DatasmithCore/Private/DatasmithMeshSerialization.cpp
@@ -173,7 +173,8 @@ FMD5Hash FDatasmithPackedMeshes::Serialize(FArchive& Ar, bool bCompressed)
 		return {};
 	}
 
-	uint32 SerialVersion = 1;
+	// Setting SerialVersion to 0 for forward compatibility between 5.1 and 5.0.
+	uint32 SerialVersion = 0;
 	Ar << SerialVersion;
 
 	enum EBufferType{ RawMeshDescription, CompressedMeshDescription };
