Fixed wrong CL merged to Main: Reverting back to revision #9

diff --git a/Engine/Source/Runtime/Datasmith/DatasmithCore/Private/DatasmithMeshSerialization.cpp b/Engine/Source/Runtime/Datasmith/DatasmithCore/Private/DatasmithMeshSerialization.cpp
--- a/Engine/Source/Runtime/Datasmith/DatasmithCore/Private/DatasmithMeshSerialization.cpp
+++ b/Engine/Source/Runtime/Datasmith/DatasmithCore/Private/DatasmithMeshSerialization.cpp
@@ -173,8 +173,7 @@ FMD5Hash FDatasmithPackedMeshes::Serialize(FArchive& Ar, bool bCompressed)
 		return {};
 	}
 
-	// Setting SerialVersion to 0 for forward compatibility between 5.1 and 5.0.
-	uint32 SerialVersion = 0;
+	uint32 SerialVersion = 1;
 	Ar << SerialVersion;
 
 	enum EBufferType{ RawMeshDescription, CompressedMeshDescription };
