[CADKernel] Improve topological welding:
- to try to fix large crack, a second step could be perform with a larger tolerance
- to prevent the meshing failure of thin surface,  and consequently the presence of a crack in the mesh, thin surface can be removed and replaced by a topological weld.
3 new CVar are added to drive this new process:
- ds.CADTranslator.Stitching.ForceSew
- ds.CADTranslator.Stitching.ForceFactor
- ds.CADTranslator.Stitching.RemoveThinFaces

Fix crashes, fix HD leak, fix RAM usage

diff --git a/Engine/Config/ConsoleVariables.ini b/Engine/Config/ConsoleVariables.ini
--- a/Engine/Config/ConsoleVariables.ini
+++ b/Engine/Config/ConsoleVariables.ini
@@ -119,7 +119,7 @@ fx.SkipVectorVMBackendOptimizations=1
 
 ; Uncomment to overwrite any existing cache associated with the file being imported. 
 ; Default is false
-; ds.CADTranslator.OverwriteCache=true
+;ds.CADTranslator.OverwriteCache=true
 
 ; Uncomment to enable/disable the timer that kill the worker if the import time is unusually long. With this time control, the load of the corrupted file 
 ; is canceled but the rest of the scene is imported. Default is enable
@@ -133,6 +133,23 @@ fx.SkipVectorVMBackendOptimizations=1
 ; Default value of StitchingTolerance is 0.001 cm
 ;ds.CADTranslator.StitchingTolerance=0.001
 
+; Uncomment to enable/disable the second welding step performed to try to fix last cracks with a largest threshold. 
+; The threshold of this second step is increased by ds.CADTranslator.Stitching.ForceFactor.\n\
+; Only available with CADKernel.
+; Default value of ForceSew is true
+;ds.CADTranslator.Stitching.ForceSew=false
+
+; Factor to increase the threshold during the welding process (largest threshold).
+; Only available with CADKernel.
+; Default value of ForceFactor is 5.
+;ds.CADTranslator.Stitching.ForceFactor=5.
+
+; During the welding process, Thin faces are remove before the stiching process. 
+; The width of the thin face is equal to the force sew tolerance.\n\
+; Only available with CADKernel.
+; Default value of RemoveThinFaces is true
+;ds.CADTranslator.Stitching.RemoveThinFaces=false
+
 ; Perform a welding of the mesh to try to stitch mesh cracks
 ; This welding is performed respecting the ds.CADTranslator.StitchingTolerance
 ; Default value is true
diff --git a/Engine/Plugins/Enterprise/DatasmithCADImporter/Source/CADInterfaces/Private/CADFileReader.cpp b/Engine/Plugins/Enterprise/DatasmithCADImporter/Source/CADInterfaces/Private/CADFileReader.cpp
--- a/Engine/Plugins/Enterprise/DatasmithCADImporter/Source/CADInterfaces/Private/CADFileReader.cpp
+++ b/Engine/Plugins/Enterprise/DatasmithCADImporter/Source/CADInterfaces/Private/CADFileReader.cpp
@@ -110,16 +110,16 @@ namespace CADLibrary
 			bool bNeedToProceed = true;
 
 			FString CADFileCachePath = CADFileData.GetCADCachePath();
-			if (!FImportParameters::bGOverwriteCache && IFileManager::Get().FileExists(*CADFileCachePath))
+			if (IFileManager::Get().FileExists(*CADFileCachePath))
 			{
-				FString MeshArchiveFilePath = CADFileData.GetMeshArchiveFilePath();
-				if (IFileManager::Get().FileExists(*MeshArchiveFilePath)) // the file has been proceed with same meshing parameters
+				CADFileData.GetCADFileDescription().SetCacheFile(CADFileCachePath);
+				if (!FImportParameters::bGOverwriteCache)
 				{
-					bNeedToProceed = false;
-				}
-				else // the file has been converted into TechSoft file (".prc") but meshed with different parameters
-				{
-					CADFileData.GetCADFileDescription().SetCacheFile(CADFileCachePath);
+					FString MeshArchiveFilePath = CADFileData.GetMeshArchiveFilePath();
+					if (IFileManager::Get().FileExists(*MeshArchiveFilePath)) // the file has been proceed with same meshing parameters
+					{
+						bNeedToProceed = false;
+					}
 				}
 			}
 
@@ -127,6 +127,10 @@ namespace CADLibrary
 			{
 				// The file has been yet proceed, get ExternalRef
 				CADFileData.LoadSceneGraphArchive();
+				// update the ExternalRef path according to the current file path
+				// Indeed, the loaded file can be in another folder with differents other files
+				// Without the update, this is the files in the other folder that will be proceed 
+				CADFileData.UpdateExternalRefPath();
 				return ECADParsingResult::ProcessOk;
 			}
 		}
diff --git a/Engine/Plugins/Enterprise/DatasmithCADImporter/Source/CADInterfaces/Private/TechSoftBridge.cpp b/Engine/Plugins/Enterprise/DatasmithCADImporter/Source/CADInterfaces/Private/TechSoftBridge.cpp
--- a/Engine/Plugins/Enterprise/DatasmithCADImporter/Source/CADInterfaces/Private/TechSoftBridge.cpp
+++ b/Engine/Plugins/Enterprise/DatasmithCADImporter/Source/CADInterfaces/Private/TechSoftBridge.cpp
@@ -182,13 +182,14 @@ UE::CADKernel::FLinearBoundary GetLinearBoundary(const A3DCrvBase* A3DCurve)
 	return Domain;
 }
 
-} // ns TechSoftUtils
+} // namespace TechSoftUtils
 
 FTechSoftBridge::FTechSoftBridge(FTechSoftFileParser& InParser, UE::CADKernel::FSession& InSession)
 	: Parser(InParser)
 	, Session(InSession)
 	, Model(InSession.GetModel())
 	, GeometricTolerance(Session.GetGeometricTolerance())
+	, EdgeLengthTolerance(Session.GetGeometricTolerance() * 2.)
 	, SquareGeometricTolerance(FMath::Square(Session.GetGeometricTolerance()))
 	, SquareJoiningVertexTolerance(SquareGeometricTolerance * 2)
 {
@@ -206,10 +207,10 @@ const A3DRiBrepModel* FTechSoftBridge::GetA3DBody(UE::CADKernel::FBody* BRepMode
 
 UE::CADKernel::FBody* FTechSoftBridge::GetBody(A3DRiBrepModel* A3DBRepModel)
 {
-	UE::CADKernel::FBody** BodyPtr = TechSoftToCADKernel.Find(A3DBRepModel);
+	TSharedPtr<UE::CADKernel::FBody>* BodyPtr = TechSoftToCADKernel.Find(A3DBRepModel);
 	if (BodyPtr &&!(*BodyPtr)->IsDeleted())
 	{
-		return *BodyPtr;
+		return (*BodyPtr).Get();
 	}
 	return nullptr;
 }
@@ -231,14 +232,9 @@ UE::CADKernel::FBody* FTechSoftBridge::AddBody(A3DRiBrepModel* A3DBRepModel, TMa
 		BRepMetaData.MetaData.FindOrAdd(TEXT("Name")) = *Name;
 	}
 	
-	UE::CADKernel::FBody** BodyPtr = TechSoftToCADKernel.Find(A3DBRepModel);
-	if (BodyPtr)
+	if (TSharedPtr<UE::CADKernel::FBody>* BodyPtr = TechSoftToCADKernel.Find(A3DBRepModel))
 	{
-		if((*BodyPtr)->IsDeleted())
-		{
-			return nullptr;
-		}
-		return *BodyPtr;
+		return (*BodyPtr)->IsDeleted() ? nullptr : (*BodyPtr).Get();
 	}
 
 	TSharedRef<UE::CADKernel::FBody> Body = UE::CADKernel::FEntity::MakeShared<UE::CADKernel::FBody>();
@@ -257,7 +253,7 @@ UE::CADKernel::FBody* FTechSoftBridge::AddBody(A3DRiBrepModel* A3DBRepModel, TMa
 	}
 
 	Model.Add(Body);
-	TechSoftToCADKernel.Add(A3DBRepModel, &*Body);
+	TechSoftToCADKernel.Add(A3DBRepModel, Body);
 	CADKernelToTechSoft.Add(&*Body, A3DBRepModel);
 
 	return &*Body;
@@ -657,30 +653,6 @@ TSharedPtr<UE::CADKernel::FTopologicalEdge> FTechSoftBridge::AddEdge(const A3DTo
 
 	OutOrientation = CoEdgeData->m_ucOrientationUVWithLoop > 0 ? UE::CADKernel::EOrientation::Front : UE::CADKernel::EOrientation::Back;
 
-	// Link edges
-	if (CoEdgeData->m_pNeighbor)
-	{
-		const A3DTopoCoEdge* Neighbor = CoEdgeData->m_pNeighbor;
-		while (Neighbor && Neighbor != A3DCoedge)
-		{
-			TSharedPtr<UE::CADKernel::FTopologicalEdge>* TwinEdge = A3DEdgeToEdge.Find(Neighbor);
-			if (TwinEdge != nullptr)
-			{
-				Edge->Link(*TwinEdge->Get(), SquareJoiningVertexTolerance);
-			}
-
-			TUniqueTSObj<A3DTopoCoEdgeData> NeighborData(Neighbor);
-			if (NeighborData.IsValid())
-			{
-				Neighbor = NeighborData->m_pNeighbor;
-			}
-			else
-			{
-				break;
-			}
-		}
-	}
-
 	return Edge;
 }
 
@@ -727,7 +699,50 @@ TSharedPtr<UE::CADKernel::FTopologicalLoop> FTechSoftBridge::AddLoop(const A3DTo
 		return TSharedPtr<UE::CADKernel::FTopologicalLoop>();
 	}
 
-	return UE::CADKernel::FTopologicalLoop::Make(Edges, Directions, bIsExternalLoop, GeometricTolerance);
+	TSharedPtr<UE::CADKernel::FTopologicalLoop> Loop = UE::CADKernel::FTopologicalLoop::Make(Edges, Directions, bIsExternalLoop, GeometricTolerance);
+
+	// Link the edges of the loop with their neighbors if possible
+	for (A3DUns32 Index = 0; Index < TopoLoopData->m_uiCoEdgeSize; ++Index)
+	{
+		const A3DTopoCoEdge* A3DCoedge = TopoLoopData->m_ppCoEdges[Index];
+		TSharedPtr<UE::CADKernel::FTopologicalEdge>* Edge = A3DEdgeToEdge.Find(A3DCoedge);
+		if (Edge == nullptr || !Edge->IsValid() || (*Edge)->IsDeleted())
+		{
+			continue;
+		}
+
+		TUniqueTSObj<A3DTopoCoEdgeData> CoEdgeData(A3DCoedge);
+		if (!CoEdgeData.IsValid())
+		{
+			continue;
+		}
+
+		if (CoEdgeData->m_pNeighbor)
+		{
+			const A3DTopoCoEdge* Neighbor = CoEdgeData->m_pNeighbor;
+			while (Neighbor && Neighbor != A3DCoedge)
+			{
+				TSharedPtr<UE::CADKernel::FTopologicalEdge>* TwinEdge = A3DEdgeToEdge.Find(Neighbor);
+				if (TwinEdge != nullptr && TwinEdge->IsValid() && !(*TwinEdge)->IsDeleted())
+				{
+					(*Edge)->LinkIfCoincident(*TwinEdge->Get(), EdgeLengthTolerance, SquareJoiningVertexTolerance);
+				}
+
+				// Next
+				TUniqueTSObj<A3DTopoCoEdgeData> NeighborData(Neighbor);
+				if (NeighborData.IsValid())
+				{
+					Neighbor = NeighborData->m_pNeighbor;
+				}
+				else
+				{
+					break;
+				}
+			}
+		}
+	}
+	 
+	return Loop;
 }
 
 void FTechSoftBridge::AddFace(const A3DTopoFace* A3DFace, UE::CADKernel::EOrientation Orientation, TSharedRef<UE::CADKernel::FShell>& Shell, uint32 ShellIndex)
diff --git a/Engine/Plugins/Enterprise/DatasmithCADImporter/Source/CADInterfaces/Private/TechSoftBridge.h b/Engine/Plugins/Enterprise/DatasmithCADImporter/Source/CADInterfaces/Private/TechSoftBridge.h
--- a/Engine/Plugins/Enterprise/DatasmithCADImporter/Source/CADInterfaces/Private/TechSoftBridge.h
+++ b/Engine/Plugins/Enterprise/DatasmithCADImporter/Source/CADInterfaces/Private/TechSoftBridge.h
@@ -175,10 +175,11 @@ private:
 
 
 	const double GeometricTolerance;
+	const double EdgeLengthTolerance;
 	const double SquareGeometricTolerance;
 	const double SquareJoiningVertexTolerance;
 
-	TMap<const A3DEntity*, UE::CADKernel::FBody*> TechSoftToCADKernel;
+	TMap<const A3DEntity*, TSharedPtr<UE::CADKernel::FBody>> TechSoftToCADKernel;
 	TMap<UE::CADKernel::FBody*, const A3DEntity*> CADKernelToTechSoft;
 	TMap<const A3DTopoCoEdge*, TSharedPtr<UE::CADKernel::FTopologicalEdge>> A3DEdgeToEdge;
 
diff --git a/Engine/Plugins/Enterprise/DatasmithCADImporter/Source/CADInterfaces/Private/TechSoftFileParserCADKernelTessellator.cpp b/Engine/Plugins/Enterprise/DatasmithCADImporter/Source/CADInterfaces/Private/TechSoftFileParserCADKernelTessellator.cpp
--- a/Engine/Plugins/Enterprise/DatasmithCADImporter/Source/CADInterfaces/Private/TechSoftFileParserCADKernelTessellator.cpp
+++ b/Engine/Plugins/Enterprise/DatasmithCADImporter/Source/CADInterfaces/Private/TechSoftFileParserCADKernelTessellator.cpp
@@ -4,6 +4,7 @@
 
 
 #include "CADFileData.h"
+#include "CADOptions.h"
 #include "CADKernelTools.h"
 #include "TechSoftUtils.h"
 #include "TechSoftBridge.h"
@@ -27,6 +28,8 @@ namespace CADLibrary
 FTechSoftFileParserCADKernelTessellator::FTechSoftFileParserCADKernelTessellator(FCADFileData& InCADData, const FString& EnginePluginsPath)
 	: FTechSoftFileParser(InCADData, EnginePluginsPath)
 	, LastHostIdUsed(1 << 30)
+	, GeometricTolerance(FImportParameters::GStitchingTolerance * 10) // cm to mm
+	, ForceFactor(FImportParameters::GStitchingForceFactor)
 {
 }
 
@@ -34,6 +37,7 @@ FTechSoftFileParserCADKernelTessellator::FTechSoftFileParserCADKernelTessellator
 
 void FTechSoftFileParserCADKernelTessellator::GenerateBodyMeshes()
 {
+	TRACE_CPUPROFILER_EVENT_SCOPE(FTechSoftFileParserCADKernelTessellator::GenerateBodyMeshes);
 	if (bForceSew || CADFileData.GetImportParameters().GetStitchingTechnique() == EStitchingTechnique::StitchingSew)
 	{
 		SewAndGenerateBodyMeshes();
@@ -47,6 +51,8 @@ void FTechSoftFileParserCADKernelTessellator::GenerateBodyMeshes()
 
 void FTechSoftFileParserCADKernelTessellator::SewAndGenerateBodyMeshes()
 {
+	TRACE_CPUPROFILER_EVENT_SCOPE(FTechSoftFileParserCADKernelTessellator::SewAndGenerateBodyMeshes);
+
 	int32 RepresentationCount = RepresentationItemsCache.Num();
 	TMap<FCadId, TArray<A3DRiRepresentationItem*>> OccurenceIdToRepresentations;
 	for (TPair<A3DRiRepresentationItem*, FCadId>& Entry : RepresentationItemsCache)
@@ -72,9 +78,12 @@ void FTechSoftFileParserCADKernelTessellator::SewAndGenerateBodyMeshes()
 
 void FTechSoftFileParserCADKernelTessellator::SewAndMesh(TArray<A3DRiRepresentationItem*>& Representations)
 {
+	TRACE_CPUPROFILER_EVENT_SCOPE(FTechSoftFileParserCADKernelTessellator::SewAndMesh);
+
 	using namespace UE::CADKernel;
 
-	const double GeometricTolerance = 0.01; // mm
+	const bool bTryToImproveSew = FImportParameters::bGStitchingForceSew;
+	const bool bRemoveThinSurface = FImportParameters::bGStitchingRemoveThinFaces;
 
 	FSession CADKernelSession(GeometricTolerance);
 	CADKernelSession.SetFirstNewHostId(LastHostIdUsed);
@@ -95,8 +104,8 @@ void FTechSoftFileParserCADKernelTessellator::SewAndMesh(TArray<A3DRiRepresentat
 	}
 
 	// Sew if needed
-	FTopomaker Topomaker(CADKernelSession, GeometricTolerance);
-	Topomaker.Sew();
+	FTopomaker Topomaker(CADKernelSession, GeometricTolerance, ForceFactor);
+	Topomaker.Sew(bTryToImproveSew, bRemoveThinSurface);
 	Topomaker.SplitIntoConnectedShells();
 	Topomaker.OrientShells();
 
@@ -106,7 +115,6 @@ void FTechSoftFileParserCADKernelTessellator::SewAndMesh(TArray<A3DRiRepresentat
 	// The new bodies will be linked to FArchiveBody of deleted bodies but the metadata of these FArchiveBody havbe to be cleaned
 
 	int32 BodyCount = CADKernelModel.GetBodies().Num();
-	int32 DeletedBodyCount = Representations.Num() - BodyCount;
 
 	TArray<FBody*> NewBodies;
 	TArray<FBody*> ExistingBodies;
@@ -133,21 +141,24 @@ void FTechSoftFileParserCADKernelTessellator::SewAndMesh(TArray<A3DRiRepresentat
 	}
 
 	// find Representation of deleted bodies to find unused FArchiveBody
-	//int32 DeletedBodyCount = Representations.Num() - BodyCount;
 	TArray<FCadId> ArchiveBodyIdOfDeletedRepresentation;
-	ArchiveBodyIdOfDeletedRepresentation.Reserve(DeletedBodyCount);
-
-	for (A3DRiRepresentationItem* Representation : Representations)
+	int32 DeletedBodyCount = Representations.Num() - BodyCount;
+	if (DeletedBodyCount > 0)
 	{
-		FBody* Body = TechSoftBridge.GetBody(Representation);
-		if (Body == nullptr)
+		ArchiveBodyIdOfDeletedRepresentation.Reserve(DeletedBodyCount);
+
+		for (A3DRiRepresentationItem* Representation : Representations)
 		{
-			FCadId* ArchiveBodyId = RepresentationItemsCache.Find(Representation);
-			if (ArchiveBodyId == nullptr)
+			FBody* Body = TechSoftBridge.GetBody(Representation);
+			if (Body == nullptr)
 			{
-				continue; // should not append
+				FCadId* ArchiveBodyId = RepresentationItemsCache.Find(Representation);
+				if (ArchiveBodyId == nullptr)
+				{
+					continue; // should not happen
+				}
+				ArchiveBodyIdOfDeletedRepresentation.Add(*ArchiveBodyId);
 			}
-			ArchiveBodyIdOfDeletedRepresentation.Add(*ArchiveBodyId);
 		}
 	}
 
@@ -166,30 +177,33 @@ void FTechSoftFileParserCADKernelTessellator::SewAndMesh(TArray<A3DRiRepresentat
 		MeshAndGetTessellation(CADKernelSession, ArchiveBody, *Body);
 	}
 
-	// Process new bodies
-	int32 Index = 0;
-	for (FBody* Body : NewBodies)
+	if (ArchiveBodyIdOfDeletedRepresentation.Num())
 	{
-		FCadId ArchiveBodyId = ArchiveBodyIdOfDeletedRepresentation[Index++];
-		FArchiveBody& ArchiveBody = SceneGraph.GetBody(ArchiveBodyId);
-		MeshAndGetTessellation(CADKernelSession, ArchiveBody, *Body);
-	}
+		// Process new bodies
+		int32 Index = 0;
+		for (FBody* Body : NewBodies)
+		{
+			FCadId ArchiveBodyId = ArchiveBodyIdOfDeletedRepresentation[Index++];
+			FArchiveBody& ArchiveBody = SceneGraph.GetBody(ArchiveBodyId);
+			MeshAndGetTessellation(CADKernelSession, ArchiveBody, *Body);
+		}
 
-	// Delete unused FArchiveBody
-	for (; Index < ArchiveBodyIdOfDeletedRepresentation.Num(); ++Index)
-	{
-		FCadId ArchiveBodyId = ArchiveBodyIdOfDeletedRepresentation[Index++];
-		FArchiveBody& ArchiveBody = SceneGraph.GetBody(ArchiveBodyId);
-		ArchiveBody.Delete();
+		// Delete unused FArchiveBody
+		for (; Index < ArchiveBodyIdOfDeletedRepresentation.Num(); ++Index)
+		{
+			FCadId ArchiveBodyId = ArchiveBodyIdOfDeletedRepresentation[Index++];
+			FArchiveBody& ArchiveBody = SceneGraph.GetBody(ArchiveBodyId);
+			ArchiveBody.Delete();
+
+		}
 	}
 }
 
 void FTechSoftFileParserCADKernelTessellator::GenerateBodyMesh(A3DRiRepresentationItem* Representation, FArchiveBody& ArchiveBody)
 {
+	TRACE_CPUPROFILER_EVENT_SCOPE(FTechSoftFileParserCADKernelTessellator::GenerateBodyMesh);
 	using namespace UE::CADKernel;
 
-	const double GeometricTolerance = 0.01; // mm
-
 	FSession CADKernelSession(GeometricTolerance);
 	CADKernelSession.SetFirstNewHostId(LastHostIdUsed);
 	FModel& CADKernelModel = CADKernelSession.GetModel();
@@ -204,9 +218,11 @@ void FTechSoftFileParserCADKernelTessellator::GenerateBodyMesh(A3DRiRepresentati
 
 	if (CADFileData.GetImportParameters().GetStitchingTechnique() == StitchingHeal)
 	{
-		const double Tolerance = 0.1; // mm
-		FTopomaker Topomaker(CADKernelSession, Tolerance);
-		Topomaker.Sew();
+		const bool bTryToImproveSew = FImportParameters::bGStitchingForceSew;
+		const bool bRemoveThinSurface = FImportParameters::bGStitchingRemoveThinFaces;
+
+		FTopomaker Topomaker(CADKernelSession, GeometricTolerance, ForceFactor);
+		Topomaker.Sew(bTryToImproveSew, bRemoveThinSurface);
 		Topomaker.OrientShells();
 	}
 
@@ -215,6 +231,7 @@ void FTechSoftFileParserCADKernelTessellator::GenerateBodyMesh(A3DRiRepresentati
 
 void FTechSoftFileParserCADKernelTessellator::MeshAndGetTessellation(UE::CADKernel::FSession& CADKernelSession, FArchiveBody& ArchiveBody, UE::CADKernel::FBody& CADKernelBody)
 {
+	TRACE_CPUPROFILER_EVENT_SCOPE(FTechSoftFileParserCADKernelTessellator::MeshAndGetTessellation);
 	using namespace UE::CADKernel;
 
 	FBodyMesh& BodyMesh = CADFileData.AddBodyMesh(ArchiveBody.Id, ArchiveBody);
@@ -225,7 +242,7 @@ void FTechSoftFileParserCADKernelTessellator::MeshAndGetTessellation(UE::CADKern
 	if (CADFileData.IsCacheDefined())
 	{
 		FString BodyFilePath = CADFileData.GetBodyCachePath(ArchiveBody.MeshActorUId);
-		CADKernelSession.SaveDatabase(*BodyFilePath);
+		CADKernelSession.SaveDatabase(*BodyFilePath, CADKernelBody);
 	}
 
 	// Tessellate the body
@@ -244,6 +261,8 @@ void FTechSoftFileParserCADKernelTessellator::MeshAndGetTessellation(UE::CADKern
 
 A3DStatus FTechSoftFileParserCADKernelTessellator::AdaptBRepModel()
 {
+	TRACE_CPUPROFILER_EVENT_SCOPE(FTechSoftFileParserCADKernelTessellator::AdaptBRepModel);
+
 	const A3DUns32 ValidSurfaceCount = 9;
 	A3DUns32 AcceptedSurfaces[ValidSurfaceCount] = {
 		//kA3DTypeSurfBlend01,
diff --git a/Engine/Plugins/Enterprise/DatasmithCADImporter/Source/CADInterfaces/Private/TechSoftFileParserCADKernelTessellator.h b/Engine/Plugins/Enterprise/DatasmithCADImporter/Source/CADInterfaces/Private/TechSoftFileParserCADKernelTessellator.h
--- a/Engine/Plugins/Enterprise/DatasmithCADImporter/Source/CADInterfaces/Private/TechSoftFileParserCADKernelTessellator.h
+++ b/Engine/Plugins/Enterprise/DatasmithCADImporter/Source/CADInterfaces/Private/TechSoftFileParserCADKernelTessellator.h
@@ -50,6 +50,9 @@ private:
 	FCadId LastEntityId = 1;
 	int32 LastHostIdUsed = 0;
 
+	const double GeometricTolerance = 0.01; // mm
+	const double ForceFactor = 5.;
+
 };
 
 } // ns CADLibrary
\ No newline at end of file
diff --git a/Engine/Plugins/Enterprise/DatasmithCADImporter/Source/CADInterfaces/Public/CADFileData.h b/Engine/Plugins/Enterprise/DatasmithCADImporter/Source/CADInterfaces/Public/CADFileData.h
--- a/Engine/Plugins/Enterprise/DatasmithCADImporter/Source/CADInterfaces/Public/CADFileData.h
+++ b/Engine/Plugins/Enterprise/DatasmithCADImporter/Source/CADInterfaces/Public/CADFileData.h
@@ -39,7 +39,11 @@ public:
 			GeomFileHash = GetSceneFileHash();
 			GeomFileHash = HashCombine(GeomFileHash, GetTypeHash(ImportParameters));
 			GeomFileHash = HashCombine(GeomFileHash, ::GetTypeHash(FImportParameters::bGSewMeshIfNeeded));
+			GeomFileHash = HashCombine(GeomFileHash, ::GetTypeHash(FImportParameters::bGRemoveDuplicatedTriangle));
 			GeomFileHash = HashCombine(GeomFileHash, ::GetTypeHash(FImportParameters::GStitchingTolerance));
+			GeomFileHash = HashCombine(GeomFileHash, ::GetTypeHash(FImportParameters::bGStitchingForceSew));
+			GeomFileHash = HashCombine(GeomFileHash, ::GetTypeHash(FImportParameters::bGStitchingRemoveThinFaces));
+			GeomFileHash = HashCombine(GeomFileHash, ::GetTypeHash(FImportParameters::GStitchingForceFactor));
 		}
 		return GeomFileHash;
 	}
@@ -80,7 +84,7 @@ public:
 	{
 		if (IsCacheDefined())
 		{
-			return CADLibrary::BuildCacheFilePath(*CachePath, TEXT("cad"), FileDescription.GetDescriptorHash());
+			return CADLibrary::BuildCadCachePath(*CachePath, FileDescription.GetDescriptorHash());
 		}
 		return FString();
 	}
@@ -232,6 +236,16 @@ public:
 		return ImportParameters;
 	}
 
+	void UpdateExternalRefPath()
+	{
+		const FString Root = FileDescription.GetRootFolder();
+		const FString ArchiveRoot = FPaths::GetPath(SceneGraphArchive.FullPath);
+		for (FFileDescriptor& File : SceneGraphArchive.ExternalReferenceFiles)
+		{
+			File.ChangePath(ArchiveRoot, Root);
+		}
+	}
+
 private:
 	const FImportParameters ImportParameters;
 	const FString CachePath;
diff --git a/Engine/Plugins/Enterprise/DatasmithCADImporter/Source/CADKernelSurface/Public/CADModelToCADKernelConverterBase.h b/Engine/Plugins/Enterprise/DatasmithCADImporter/Source/CADKernelSurface/Public/CADModelToCADKernelConverterBase.h
--- a/Engine/Plugins/Enterprise/DatasmithCADImporter/Source/CADKernelSurface/Public/CADModelToCADKernelConverterBase.h
+++ b/Engine/Plugins/Enterprise/DatasmithCADImporter/Source/CADKernelSurface/Public/CADModelToCADKernelConverterBase.h
@@ -26,6 +26,7 @@ public:
 		, ImportParameters(InImportParameters)
 		, GeometricTolerance(0.01)
 		, SquareTolerance(GeometricTolerance* GeometricTolerance)
+		, EdgeLengthTolerance(2 * GeometricTolerance)
 	{
 	}
 
@@ -36,13 +37,16 @@ public:
 
 	virtual bool RepairTopology() override
 	{
+		using namespace CADLibrary;
 		// Apply stitching if applicable
-		if(ImportParameters.GetStitchingTechnique() != CADLibrary::StitchingNone)
+		if(ImportParameters.GetStitchingTechnique() != StitchingNone)
 		{
-			// the joining tolerance is set to 0.1 mm until the user can specify it
-			const double JoiningTolerance = 0.1;
-			UE::CADKernel::FTopomaker Topomaker(CADKernelSession, JoiningTolerance);
-			Topomaker.Sew();
+			const double JoiningTolerance = FImportParameters::GStitchingTolerance * 10.; //CM to MM
+			const double ForceFactor = FImportParameters::GStitchingForceFactor;
+			bool bForceJoining = FImportParameters::bGStitchingForceSew;
+			bool bRemoveThinFaces = FImportParameters::bGStitchingRemoveThinFaces;
+			UE::CADKernel::FTopomaker Topomaker(CADKernelSession, JoiningTolerance, ForceFactor);
+			Topomaker.Sew(bForceJoining, bRemoveThinFaces);
 			Topomaker.OrientShells();
 		}
 
@@ -85,6 +89,6 @@ protected:
 	CADLibrary::FImportParameters ImportParameters;
 	double GeometricTolerance;
 	double SquareTolerance;
-
+	double EdgeLengthTolerance;
 };
 
diff --git a/Engine/Plugins/Enterprise/DatasmithCADImporter/Source/CADLibrary/Private/CADKernelTools.cpp b/Engine/Plugins/Enterprise/DatasmithCADImporter/Source/CADLibrary/Private/CADKernelTools.cpp
--- a/Engine/Plugins/Enterprise/DatasmithCADImporter/Source/CADLibrary/Private/CADKernelTools.cpp
+++ b/Engine/Plugins/Enterprise/DatasmithCADImporter/Source/CADLibrary/Private/CADKernelTools.cpp
@@ -286,7 +286,8 @@ namespace CADLibrary
 		TSharedRef<FModelMesh> CADKernelModelMesh = FEntity::MakeShared<FModelMesh>();
 		FParametricMesher Mesher(*CADKernelModelMesh);
 
-		DefineMeshCriteria(*CADKernelModelMesh, ImportParameters, 0.01);
+		constexpr double GeometricTolerance = 0.01;
+		DefineMeshCriteria(*CADKernelModelMesh, ImportParameters, GeometricTolerance);
 		Mesher.MeshEntity(CADTopologicalEntity);
 
 		FMeshConversionContext Context(ImportParameters, MeshParameters);
@@ -406,13 +407,15 @@ namespace CADLibrary
 
 		if (ImportParameters.GetMaxEdgeLength() > SMALL_NUMBER)
 		{
-			TSharedPtr<UE::CADKernel::FCriterion> MaxSizeCriterion = UE::CADKernel::FCriterion::CreateCriterion(UE::CADKernel::ECriterion::MaxSize, ImportParameters.GetMaxEdgeLength());
+			const double MaxSize = ImportParameters.GetMaxEdgeLength() * 10.; // cm to mm
+			TSharedPtr<UE::CADKernel::FCriterion> MaxSizeCriterion = UE::CADKernel::FCriterion::CreateCriterion(UE::CADKernel::ECriterion::MaxSize, MaxSize);
 			MeshModel.AddCriterion(MaxSizeCriterion);
 		}
 
 		if (ImportParameters.GetChordTolerance() > SMALL_NUMBER)
 		{
-			TSharedPtr<UE::CADKernel::FCriterion> ChordCriterion = UE::CADKernel::FCriterion::CreateCriterion(UE::CADKernel::ECriterion::Sag, ImportParameters.GetChordTolerance());
+			double MaxChord = ImportParameters.GetChordTolerance() * 10.; // cm to mm
+			TSharedPtr<UE::CADKernel::FCriterion> ChordCriterion = UE::CADKernel::FCriterion::CreateCriterion(UE::CADKernel::ECriterion::Sag, MaxChord);
 			MeshModel.AddCriterion(ChordCriterion);
 		}
 
diff --git a/Engine/Plugins/Enterprise/DatasmithCADImporter/Source/CADTools/Private/CADOption.cpp b/Engine/Plugins/Enterprise/DatasmithCADImporter/Source/CADTools/Private/CADOption.cpp
--- a/Engine/Plugins/Enterprise/DatasmithCADImporter/Source/CADTools/Private/CADOption.cpp
+++ b/Engine/Plugins/Enterprise/DatasmithCADImporter/Source/CADTools/Private/CADOption.cpp
@@ -62,6 +62,34 @@ FAutoConsoleVariableRef GCADTranslatorStitchingTolerance(
 Default value of StitchingTolerance is 0.001 cm\n"),
 ECVF_Default);
 
+bool FImportParameters::bGStitchingForceSew = true;
+FAutoConsoleVariableRef GCADTranslatorStitchingForceSew(
+	TEXT("ds.CADTranslator.Stitching.ForceSew"),
+	FImportParameters::bGStitchingForceSew,
+	TEXT("During the welding process, a second step is performed to try to fix last cracks with a largest threshold.\n\
+For this second step, the threshold is increased by ds.CADTranslator.Stitching.ForceFactor.\n\
+Only available with CADKernel.\n\
+Default value of ForceSew is true\n"),
+ECVF_Default);
+
+float FImportParameters::GStitchingForceFactor = 5.f;
+FAutoConsoleVariableRef GCADTranslatorStitchingForceFactor(
+	TEXT("ds.CADTranslator.Stitching.ForceFactor"),
+	FImportParameters::GStitchingForceFactor,
+	TEXT("Factor to increase the tolerance during the welding process.\n\
+Only available with CADKernel.\n\
+Default value of ForceFactor is 5.\n"),
+ECVF_Default);
+
+bool FImportParameters::bGStitchingRemoveThinFaces = true;
+FAutoConsoleVariableRef GCADTranslatorStitchingRemoveThinFaces(
+	TEXT("ds.CADTranslator.Stitching.RemoveThinFaces"),
+	FImportParameters::bGStitchingRemoveThinFaces,
+	TEXT("During the welding process, Thin faces are removed before the stiching process. The width of the thin faces is equal to the force sew tolerance.\n\
+Only available with CADKernel.\n\
+Default value of RemoveThinFaces is true\n"),
+ECVF_Default);
+
 float FImportParameters::GUnitScale = 1.f;
 FAutoConsoleVariableRef GCADTranslatorUnitScale(
 	TEXT("ds.CADTranslator.UnitScale"),
diff --git a/Engine/Plugins/Enterprise/DatasmithCADImporter/Source/CADTools/Public/CADData.h b/Engine/Plugins/Enterprise/DatasmithCADImporter/Source/CADTools/Public/CADData.h
--- a/Engine/Plugins/Enterprise/DatasmithCADImporter/Source/CADTools/Public/CADData.h
+++ b/Engine/Plugins/Enterprise/DatasmithCADImporter/Source/CADTools/Public/CADData.h
@@ -270,6 +270,18 @@ public:
 		return Name;
 	}
 
+	void ChangePath(const FString& OldPath, const FString& NewPath)
+	{
+		int32 OldPathLength = OldPath.Len();
+		FString SourceFileOldPath = SourceFilePath.Left(OldPathLength);
+		FString SourceFileEndPath = SourceFilePath.Right(SourceFilePath.Len() - OldPathLength - 1);
+		if (OldPath == SourceFileOldPath)
+		{
+			SourceFilePath = FPaths::Combine(NewPath, SourceFileEndPath);
+		}
+		RootFolder = NewPath;
+	}
+
 private:
 
 	FString SourceFilePath; // e.g. d:/folder/content.jt
diff --git a/Engine/Plugins/Enterprise/DatasmithCADImporter/Source/CADTools/Public/CADOptions.h b/Engine/Plugins/Enterprise/DatasmithCADImporter/Source/CADTools/Public/CADOptions.h
--- a/Engine/Plugins/Enterprise/DatasmithCADImporter/Source/CADTools/Public/CADOptions.h
+++ b/Engine/Plugins/Enterprise/DatasmithCADImporter/Source/CADTools/Public/CADOptions.h
@@ -37,6 +37,9 @@ namespace CADLibrary
 		CADTOOLS_API static bool bGSewMeshIfNeeded;
 		CADTOOLS_API static bool bGRemoveDuplicatedTriangle;
 		CADTOOLS_API static float GStitchingTolerance;
+		CADTOOLS_API static bool bGStitchingForceSew;
+		CADTOOLS_API static bool bGStitchingRemoveThinFaces;
+		CADTOOLS_API static float GStitchingForceFactor;
 		CADTOOLS_API static float GUnitScale;
 		CADTOOLS_API static float GMeshingParameterFactor;
 		CADTOOLS_API static int32 GMaxMaterialCountPerMesh;
@@ -62,7 +65,7 @@ namespace CADLibrary
 		uint32 GetHash() const
 		{
 			uint32 Hash = 0;
-			for (double Param : {ChordTolerance, MaxEdgeLength, MaxNormalAngle})
+			for (double Param : {ChordTolerance, MaxEdgeLength, MaxNormalAngle, (double) GStitchingForceFactor})
 			{
 				Hash = HashCombine(Hash, ::GetTypeHash(Param));
 			}
@@ -70,7 +73,10 @@ namespace CADLibrary
 			{
 				Hash = HashCombine(Hash, ::GetTypeHash(Param));
 			}
-			Hash = HashCombine(Hash, ::GetTypeHash(bGPreferJtFileEmbeddedTessellation));
+			for (bool Param : {bGStitchingForceSew, bGStitchingRemoveThinFaces, bGDisableCADKernelTessellation, bGPreferJtFileEmbeddedTessellation})
+			{
+				Hash = HashCombine(Hash, ::GetTypeHash(Param));
+			}
 			return Hash;
 		}
 
@@ -90,6 +96,9 @@ namespace CADLibrary
 			Ar << ImportParameters.bGEnableCADCache;
 			Ar << ImportParameters.bGPreferJtFileEmbeddedTessellation;
 			Ar << ImportParameters.GStitchingTolerance;
+			Ar << ImportParameters.bGStitchingForceSew;
+			Ar << ImportParameters.bGStitchingRemoveThinFaces;
+			Ar << ImportParameters.GStitchingForceFactor;
 			Ar << ImportParameters.GMaxMaterialCountPerMesh;
 			return Ar;
 		}
@@ -127,6 +136,12 @@ namespace CADLibrary
 		CADTOOLS_API friend uint32 GetTypeHash(const FImportParameters& ImportParameters);
 	};
 
+	inline FString BuildCadCachePath(const TCHAR* CachePath, uint32 FileHash)
+	{
+		FString FileName = FString::Printf(TEXT("UEx%08x"), FileHash) + TEXT(".prc");
+		return FPaths::Combine(CachePath, TEXT("cad"), FileName);
+	}
+
 	inline FString BuildCacheFilePath(const TCHAR* CachePath, const TCHAR* Folder, uint32 BodyHash)
 	{
 		FString BodyFileName = FString::Printf(TEXT("UEx%08x"), BodyHash);
diff --git a/Engine/Plugins/Enterprise/DatasmithCADImporter/Source/DatasmithCADTranslator/Private/DatasmithCADTranslator.cpp b/Engine/Plugins/Enterprise/DatasmithCADImporter/Source/DatasmithCADTranslator/Private/DatasmithCADTranslator.cpp
--- a/Engine/Plugins/Enterprise/DatasmithCADImporter/Source/DatasmithCADTranslator/Private/DatasmithCADTranslator.cpp
+++ b/Engine/Plugins/Enterprise/DatasmithCADImporter/Source/DatasmithCADTranslator/Private/DatasmithCADTranslator.cpp
@@ -125,17 +125,18 @@ bool FDatasmithCADTranslator::IsSourceSupported(const FDatasmithSceneSource& Sou
 
 bool FDatasmithCADTranslator::LoadScene(TSharedRef<IDatasmithScene> DatasmithScene)
 {
+	using namespace CADLibrary;
 	const FDatasmithTessellationOptions& TesselationOptions = GetCommonTessellationOptions();
 	CADLibrary::FFileDescriptor FileDescriptor(*FPaths::ConvertRelativePathToFull(GetSource().GetSourceFile()));
 
 	UE_LOG(LogCADTranslator, Display, TEXT("CAD translation [%s]."), *FileDescriptor.GetSourcePath());
 	UE_LOG(LogCADTranslator, Display, TEXT(" - Parsing Library:      %s"), TEXT("TechSoft"));
 	UE_LOG(LogCADTranslator, Display, TEXT(" - Tessellation Library: %s")
-		, CADLibrary::FImportParameters::bGDisableCADKernelTessellation ? TEXT("TechSoft") : TEXT("CADKernel"));
+		, FImportParameters::bGDisableCADKernelTessellation ? TEXT("TechSoft") : TEXT("CADKernel"));
 	UE_LOG(LogCADTranslator, Display, TEXT(" - Cache mode:           %s")
-		, CADLibrary::FImportParameters::bGEnableCADCache ? (CADLibrary::FImportParameters::bGOverwriteCache ? TEXT("Override") : TEXT("Enabled")) : TEXT("Disabled"));
+		, FImportParameters::bGEnableCADCache ? (FImportParameters::bGOverwriteCache ? TEXT("Override") : TEXT("Enabled")) : TEXT("Disabled"));
 	UE_LOG(LogCADTranslator, Display, TEXT(" - Processing:           %s")
-		, CADLibrary::FImportParameters::bGEnableCADCache ? (CADLibrary::GMaxImportThreads == 1 ? TEXT("Sequencial") : TEXT("Parallel")) : TEXT("Sequencial"));
+		, FImportParameters::bGEnableCADCache ? (GMaxImportThreads == 1 ? TEXT("Sequencial") : TEXT("Parallel")) : TEXT("Sequencial"));
 
 	TFunction<double(double, double, const TCHAR*)> CheckParameterValue = [](double Value, double MinValue, const TCHAR* ParameterName) -> double
 	{
@@ -150,9 +151,35 @@ bool FDatasmithCADTranslator::LoadScene(TSharedRef<IDatasmithScene> DatasmithSce
 	ImportParameters.SetTesselationParameters(CheckParameterValue(TesselationOptions.ChordTolerance, UE::DatasmithTessellation::MinTessellationChord, TEXT("Chord tolerance")),
 		FMath::IsNearlyZero(TesselationOptions.MaxEdgeLength) ? 0. : CheckParameterValue(TesselationOptions.MaxEdgeLength, UE::DatasmithTessellation::MinTessellationEdgeLength, TEXT("Max Edge Length")),
 		CheckParameterValue(TesselationOptions.NormalTolerance, UE::DatasmithTessellation::MinTessellationAngle, TEXT("Max Angle")),
-		(CADLibrary::EStitchingTechnique)TesselationOptions.StitchingTechnique);
+		(EStitchingTechnique)TesselationOptions.StitchingTechnique);
 	ImportParameters.SetModelCoordinateSystem(FDatasmithUtils::EModelCoordSystem::ZUp_RightHanded);
 
+	UE_LOG(LogCADTranslator, Display, TEXT(" - Import parameters:"));
+	UE_LOG(LogCADTranslator, Display, TEXT("     - ChordTolerance:     %f"), ImportParameters.GetChordTolerance());
+	UE_LOG(LogCADTranslator, Display, TEXT("     - MaxEdgeLength:      %f"), ImportParameters.GetMaxEdgeLength());
+	UE_LOG(LogCADTranslator, Display, TEXT("     - MaxNormalAngle:     %f"), ImportParameters.GetMaxNormalAngle());
+	FString StitchingTechnique;
+	switch(ImportParameters.GetStitchingTechnique())
+	{
+		case EStitchingTechnique::StitchingHeal:
+			StitchingTechnique = TEXT("Heal");
+			break;
+		case EStitchingTechnique::StitchingSew:
+			StitchingTechnique = TEXT("Sew");
+			break;
+		default:
+			StitchingTechnique = TEXT("None");
+			break;
+	}
+	UE_LOG(LogCADTranslator, Display, TEXT("     - StitchingTechnique: %s"), *StitchingTechnique);
+	if (!FImportParameters::bGDisableCADKernelTessellation)
+	{
+		UE_LOG(LogCADTranslator, Display, TEXT("     - Stitching Options:"));
+		UE_LOG(LogCADTranslator, Display, TEXT("         - ForceSew:        %s"), ImportParameters.bGStitchingForceSew ? TEXT("True") : TEXT("False"));
+		UE_LOG(LogCADTranslator, Display, TEXT("         - RemoveThinFaces: %s"), ImportParameters.bGStitchingRemoveThinFaces ? TEXT("True") : TEXT("False"));
+		UE_LOG(LogCADTranslator, Display, TEXT("         - ForceFactor:     %f"), ImportParameters.GStitchingForceFactor);
+	}
+
 	switch (FileDescriptor.GetFileFormat())
 	{
 	case CADLibrary::ECADFormat::N_X:
@@ -160,25 +187,25 @@ bool FDatasmithCADTranslator::LoadScene(TSharedRef<IDatasmithScene> DatasmithSce
 		break;
 	}
 
-	case CADLibrary::ECADFormat::SOLIDWORKS:
+	case ECADFormat::SOLIDWORKS:
 	{
 		ImportParameters.SetModelCoordinateSystem(FDatasmithUtils::EModelCoordSystem::YUp_RightHanded);
 		break;
 	}
 
-	case CADLibrary::ECADFormat::INVENTOR:
-	case CADLibrary::ECADFormat::CREO:
+	case ECADFormat::INVENTOR:
+	case ECADFormat::CREO:
 	{
 		ImportParameters.SetModelCoordinateSystem(FDatasmithUtils::EModelCoordSystem::YUp_RightHanded);
 		break;
 	}
 
-	case CADLibrary::ECADFormat::DWG:
+	case ECADFormat::DWG:
 	{
 		break;
 	}
 
-	case CADLibrary::ECADFormat::IFC:
+	case ECADFormat::IFC:
 	{
 		ImportParameters.SetModelCoordinateSystem(FDatasmithUtils::EModelCoordSystem::ZUp_RightHanded_FBXLegacy);
 		break;
@@ -195,16 +222,16 @@ bool FDatasmithCADTranslator::LoadScene(TSharedRef<IDatasmithScene> DatasmithSce
 	}
 
 	// Use sequential translation (multi-processed or not)
-	if (CADLibrary::FImportParameters::bGEnableCADCache)
+	if (FImportParameters::bGEnableCADCache)
 	{
 		TMap<uint32, FString> CADFileToUEFileMap;
 		{
 			int32 NumCores = FPlatformMisc::NumberOfCores();
-			if (CADLibrary::GMaxImportThreads > 1)
+			if (GMaxImportThreads > 1)
 			{
 				if (FileDescriptor.CanReferenceOtherFiles())
 				{
-					NumCores = FMath::Min(CADLibrary::GMaxImportThreads, NumCores);
+					NumCores = FMath::Min(GMaxImportThreads, NumCores);
 				}
 				else 
 				{
@@ -214,7 +241,7 @@ bool FDatasmithCADTranslator::LoadScene(TSharedRef<IDatasmithScene> DatasmithSce
 			DatasmithDispatcher::FDatasmithDispatcher Dispatcher(ImportParameters, CachePath, NumCores, CADFileToUEFileMap, CADFileToUEGeomMap);
 			Dispatcher.AddTask(FileDescriptor);
 
-			Dispatcher.Process(CADLibrary::GMaxImportThreads != 1);
+			Dispatcher.Process(GMaxImportThreads != 1);
 		}
 
 		FDatasmithSceneGraphBuilder SceneGraphBuilder(CADFileToUEFileMap, CachePath, DatasmithScene, GetSource(), ImportParameters);
@@ -225,13 +252,13 @@ bool FDatasmithCADTranslator::LoadScene(TSharedRef<IDatasmithScene> DatasmithSce
 		return true;
 	}
 
-	CADLibrary::FCADFileReader FileReader(ImportParameters, FileDescriptor, *FPaths::EnginePluginsDir(), CachePath);
-	if (FileReader.ProcessFile() != CADLibrary::ECADParsingResult::ProcessOk)
+	FCADFileReader FileReader(ImportParameters, FileDescriptor, *FPaths::EnginePluginsDir(), CachePath);
+	if (FileReader.ProcessFile() != ECADParsingResult::ProcessOk)
 	{
 		return false;
 	}
 
-	CADLibrary::FCADFileData& CADFileData = FileReader.GetCADFileData();
+	FCADFileData& CADFileData = FileReader.GetCADFileData();
 	FDatasmithSceneBaseGraphBuilder SceneGraphBuilder(&CADFileData.GetSceneGraphArchive(), CachePath, DatasmithScene, GetSource(), ImportParameters);
 	SceneGraphBuilder.Build();
 
diff --git a/Engine/Plugins/Enterprise/DatasmithCADImporter/Source/DatasmithOpenNurbsTranslator/Private/OpenNurbsBRepToCADKernelConverter.cpp b/Engine/Plugins/Enterprise/DatasmithCADImporter/Source/DatasmithOpenNurbsTranslator/Private/OpenNurbsBRepToCADKernelConverter.cpp
--- a/Engine/Plugins/Enterprise/DatasmithCADImporter/Source/DatasmithOpenNurbsTranslator/Private/OpenNurbsBRepToCADKernelConverter.cpp
+++ b/Engine/Plugins/Enterprise/DatasmithCADImporter/Source/DatasmithOpenNurbsTranslator/Private/OpenNurbsBRepToCADKernelConverter.cpp
@@ -169,7 +169,7 @@ void FOpenNurbsBRepToCADKernelConverter::LinkEdgesLoop(const ON_BrepLoop& OpenNu
 			TSharedPtr<UE::CADKernel::FTopologicalEdge>* TwinEdge = OpenNurbsTrimId2CADKernelEdge.Find(LinkedEdgeId);
 			if (TwinEdge != nullptr && TwinEdge->IsValid() && !(*TwinEdge)->IsDeleted() && !(*TwinEdge)->IsDegenerated())
 			{
-				(*Edge)->Link(**TwinEdge, SquareTolerance);
+				(*Edge)->LinkIfCoincident(**TwinEdge, EdgeLengthTolerance, SquareTolerance);
 				break;
 			}
 		}
diff --git a/Engine/Plugins/Enterprise/DatasmithCADImporter/Source/DatasmithWireTranslator/Private/AliasModelToCADKernelConverter.cpp b/Engine/Plugins/Enterprise/DatasmithCADImporter/Source/DatasmithWireTranslator/Private/AliasModelToCADKernelConverter.cpp
--- a/Engine/Plugins/Enterprise/DatasmithCADImporter/Source/DatasmithWireTranslator/Private/AliasModelToCADKernelConverter.cpp
+++ b/Engine/Plugins/Enterprise/DatasmithCADImporter/Source/DatasmithWireTranslator/Private/AliasModelToCADKernelConverter.cpp
@@ -197,7 +197,7 @@ void FAliasModelToCADKernelConverter::LinkEdgesLoop(const AlTrimBoundary& TrimBo
 			{
 				if (TwinEdge->IsValid() && !(*TwinEdge)->IsDeleted() && !(*TwinEdge)->IsDegenerated())
 				{
-					(*Edge)->Link(**TwinEdge, SquareTolerance);
+					(*Edge)->LinkIfCoincident(**TwinEdge, EdgeLengthTolerance, SquareTolerance);
 				}
 			}
 		}
diff --git a/Engine/Source/Programs/Enterprise/Datasmith/DatasmithCADWorker/Private/DatasmithCADWorkerImpl.cpp b/Engine/Source/Programs/Enterprise/Datasmith/DatasmithCADWorker/Private/DatasmithCADWorkerImpl.cpp
--- a/Engine/Source/Programs/Enterprise/Datasmith/DatasmithCADWorker/Private/DatasmithCADWorkerImpl.cpp
+++ b/Engine/Source/Programs/Enterprise/Datasmith/DatasmithCADWorker/Private/DatasmithCADWorkerImpl.cpp
@@ -50,32 +50,32 @@ bool FDatasmithCADWorkerImpl::Run()
 	{
 		if (TSharedPtr<ICommand> Command = CommandIO.GetNextCommand(1.0))
 		{
-			switch(Command->GetType())
+			switch (Command->GetType())
 			{
-				case ECommandId::Ping:
-					ProcessCommand(*StaticCast<FPingCommand*>(Command.Get()));
-					break;
-
-				case ECommandId::BackPing:
-					ProcessCommand(*StaticCast<FBackPingCommand*>(Command.Get()));
-					break;
-
-				case ECommandId::RunTask:
-					ProcessCommand(*StaticCast<FRunTaskCommand*>(Command.Get()));
-					break;
-
-				case ECommandId::ImportParams:
-					ProcessCommand(*StaticCast<FImportParametersCommand*>(Command.Get()));
-					break;
-
-				case ECommandId::Terminate:
-					UE_LOG(LogDatasmithCADWorker, Verbose, TEXT("Terminate command received. Exiting."));
-					bIsRunning = false;
-					break;
-
-				case ECommandId::NotifyEndTask:
-				default:
-					break;
+			case ECommandId::Ping:
+				ProcessCommand(*StaticCast<FPingCommand*>(Command.Get()));
+				break;
+
+			case ECommandId::BackPing:
+				ProcessCommand(*StaticCast<FBackPingCommand*>(Command.Get()));
+				break;
+
+			case ECommandId::RunTask:
+				ProcessCommand(*StaticCast<FRunTaskCommand*>(Command.Get()));
+				break;
+
+			case ECommandId::ImportParams:
+				ProcessCommand(*StaticCast<FImportParametersCommand*>(Command.Get()));
+				break;
+
+			case ECommandId::Terminate:
+				UE_LOG(LogDatasmithCADWorker, Verbose, TEXT("Terminate command received. Exiting."));
+				bIsRunning = false;
+				break;
+
+			case ECommandId::NotifyEndTask:
+			default:
+				break;
 			}
 		}
 		else
@@ -157,14 +157,16 @@ void FDatasmithCADWorkerImpl::ProcessCommand(const FRunTaskCommand& RunTaskComma
 
 	bProcessIsRunning = true;
 	int64 MaxDuration = DefineMaximumAllowedDuration(FileToProcess);
-	UE::Tasks::FTask TimeChecker = UE::Tasks::Launch(TEXT("TimeChecker"), [&]() { CheckDuration(FileToProcess, MaxDuration); });
+	TArray<UE::Tasks::FTask> Checkers;
+	Checkers.Emplace(UE::Tasks::Launch(TEXT("TimeChecker"), [&]() { CheckDuration(FileToProcess, MaxDuration); }));
+	Checkers.Emplace(UE::Tasks::Launch(TEXT("MemoryChecker"), [&]() { CheckMemory(FileToProcess); }));
 
 	CADLibrary::FCADFileReader FileReader(ImportParameters, FileToProcess, EnginePluginsPath, CachePath);
 	ETaskState ProcessResult = FileReader.ProcessFile();
 
 	bProcessIsRunning = false;
-	TimeChecker.Wait();
-	
+	UE::Tasks::Wait(Checkers);
+
 	CompletedTask.ProcessResult = ProcessResult;
 
 	if (CompletedTask.ProcessResult == ETaskState::ProcessOk)
@@ -174,11 +176,27 @@ void FDatasmithCADWorkerImpl::ProcessCommand(const FRunTaskCommand& RunTaskComma
 		CompletedTask.SceneGraphFileName = CADFileData.GetSceneGraphFileName();
 		CompletedTask.GeomFileName = CADFileData.GetMeshFileName();
 		CompletedTask.WarningMessages = CADFileData.GetWarningMessages();
+
+		UE_LOG(LogDatasmithCADWorker, Verbose, TEXT("  Process %s %s saved into %s%s and %s%s."), *FileToProcess.GetFileName(), *FileToProcess.GetConfiguration(), *CompletedTask.SceneGraphFileName, TEXT(".sg"), *CompletedTask.GeomFileName, TEXT(".gm"));
+		UE_LOG(LogDatasmithCADWorker, Verbose, TEXT("     It generates %d bodies"), CADFileData.GetBodyMeshes().Num());
+		for (const CADLibrary::FBodyMesh& BodyMesh : CADFileData.GetBodyMeshes())
+		{
+			FString BodyFileName = FString::Printf(TEXT("UEx%08x"), BodyMesh.MeshActorUId);
+			UE_LOG(LogDatasmithCADWorker, Verbose, TEXT("     - Body %s"), *BodyFileName);
+		}
+	}
+	else
+	{
+		UE_LOG(LogDatasmithCADWorker, Verbose, TEXT("  Process %s %s failed"), *FileToProcess.GetFileName(), *FileToProcess.GetConfiguration());
 	}
 
 	CommandIO.SendCommand(CompletedTask, Config::SendCommandTimeout_s);
 
-	UE_LOG(LogDatasmithCADWorker, Verbose, TEXT("End of Process %s %s saved in %s"), *FileToProcess.GetFileName(), *FileToProcess.GetConfiguration(), *CompletedTask.GeomFileName);
+	UE_LOG(LogDatasmithCADWorker, Verbose, TEXT("End of Process %s"), *FileToProcess.GetFileName());
+	if (bRequestExit)
+	{
+		FPlatformMisc::RequestExit(true);
+	}
 }
 
 void FDatasmithCADWorkerImpl::CheckDuration(const CADLibrary::FFileDescriptor& FileToProcess, const int64 MaxDuration)
@@ -201,5 +219,37 @@ void FDatasmithCADWorkerImpl::CheckDuration(const CADLibrary::FFileDescriptor& F
 		}
 	}
 	double Duration = (FPlatformTime::Cycles64() - StartTime) * FPlatformTime::GetSecondsPerCycle64();
-	UE_LOG(LogDatasmithCADWorker, Verbose, TEXT("    Processing Time: %f s"), Duration);
+	UE_LOG(LogDatasmithCADWorker, Verbose, TEXT("    - Processing Time: %f s"), Duration);
+}
+
+void FDatasmithCADWorkerImpl::CheckMemory(const CADLibrary::FFileDescriptor& FileToProcess, bool& bOutRequestExit)
+{
+	constexpr uint64 OneMegaBit = 1024 * 1024;
+	constexpr uint64 GigaBit = 1024 * 1024 * 1024;
+
+	uint64 MaxMemoryUsed = FPlatformMemory::GetStats().UsedPhysical;
+	UE_LOG(LogDatasmithCADWorker, Verbose, TEXT("    - Start Ram used %llu Mb"), MaxMemoryUsed / OneMegaBit);
+
+	uint64 StartTime = FPlatformTime::Cycles64();
+	while (bProcessIsRunning)
+	{
+		FPlatformProcess::Sleep(0.1);
+		double Duration = (FPlatformTime::Cycles64() - StartTime) * FPlatformTime::GetSecondsPerCycle64();
+
+		const uint64 MemoryUsed = FPlatformMemory::GetStats().UsedPhysical;
+		if (MaxMemoryUsed < MemoryUsed)
+		{
+			MaxMemoryUsed = MemoryUsed;
+		}
+	}
+
+
+	uint64 EndMemoryUsed = FPlatformMemory::GetStats().UsedPhysical;
+	UE_LOG(LogDatasmithCADWorker, Verbose, TEXT("    - End Ram used %llu Mb"), EndMemoryUsed / OneMegaBit);
+	UE_LOG(LogDatasmithCADWorker, Verbose, TEXT("    - Max Ram used %llu Mb"), MaxMemoryUsed / OneMegaBit);
+	if(EndMemoryUsed > GigaBit)
+	{
+		UE_LOG(LogDatasmithCADWorker, Verbose, TEXT("    - End Ram used %llu Mb exceeded the limit to restart"), EndMemoryUsed / OneMegaBit);
+		bOutRequestExit = true;
+	}
 }
diff --git a/Engine/Source/Programs/Enterprise/Datasmith/DatasmithCADWorker/Private/DatasmithCADWorkerImpl.h b/Engine/Source/Programs/Enterprise/Datasmith/DatasmithCADWorker/Private/DatasmithCADWorkerImpl.h
--- a/Engine/Source/Programs/Enterprise/Datasmith/DatasmithCADWorker/Private/DatasmithCADWorkerImpl.h
+++ b/Engine/Source/Programs/Enterprise/Datasmith/DatasmithCADWorker/Private/DatasmithCADWorkerImpl.h
@@ -31,6 +31,7 @@ private:
 	 * @param MaxDuration is defined in seconds.
 	 */
 	void CheckDuration(const CADLibrary::FFileDescriptor& FileToProcess, const int64 MaxDuration); 
+	void CheckMemory(const CADLibrary::FFileDescriptor& FileToProcess);
 
 private:
 	DatasmithDispatcher::FNetworkClientNode NetworkInterface;
@@ -43,4 +44,6 @@ private:
 	CADLibrary::FImportParameters ImportParameters;
 	uint64 PingStartCycle;
 	TAtomic<bool> bProcessIsRunning;
+
+	bool bRequestExit = false;
 };
diff --git a/Engine/Source/Runtime/Datasmith/CADKernel/Private/CADKernel/Topo/FaceAnalyzer.cpp b/Engine/Source/Runtime/Datasmith/CADKernel/Private/CADKernel/Topo/FaceAnalyzer.cpp
--- a/Engine/Source/Runtime/Datasmith/CADKernel/Private/CADKernel/Topo/FaceAnalyzer.cpp
+++ b/Engine/Source/Runtime/Datasmith/CADKernel/Private/CADKernel/Topo/FaceAnalyzer.cpp
@@ -35,7 +35,8 @@ void FFaceAnalyzer::FindClosedSegments(Topo::FThinFaceContext& Context)
 
 		for (Topo::FEdgeSegment* Candidate : Context.LoopSegments)
 		{
-			if (Candidate->GetEdge() == Edge)
+			// to avoid to define cylinder or cone as a thin surface
+			if (Candidate->GetEdge()->IsLinkedTo(*Edge))
 			{
 				continue;
 			}
@@ -102,22 +103,23 @@ void FFaceAnalyzer::Analyze(Topo::FThinFaceContext& Context)
 		}
 		else
 		{
-			MedSquareDistance = -1;
-			MaxSquareDistance = -1;
+			MedSquareDistance = DOUBLE_BIG_NUMBER;
+			MaxSquareDistance = DOUBLE_BIG_NUMBER;
 		}
 
-		if (MaxSquareDistance < 2*FMath::Square(Edge->Length()))
+		if (MaxSquareDistance < 2 * SquareTolerance)
 		{
 			Context.EdgeSquareDistance.Add(MedSquareDistance);
 			Context.EdgeMaxSquareDistance.Add(MaxSquareDistance);
 
 			Context.MaxSquareDistance = FMath::Max(Context.MaxSquareDistance, MaxSquareDistance);
-			Context.SideEdgeLength += Edge->Length();
+			Context.ThinSideEdgeLength += Edge->Length();
 		}
 		else
 		{
 			Context.EdgeSquareDistance.Add(-1.);
 			Context.EdgeMaxSquareDistance.Add(-1.);
+			Context.OppositSideEdgeLength += Edge->Length();
 		}
 		Context.ExternalLoopLength += Edge->Length();
 	};
@@ -155,7 +157,6 @@ void FFaceAnalyzer::Analyze(Topo::FThinFaceContext& Context)
 void FFaceAnalyzer::BuildLoopSegments(Topo::FThinFaceContext& Context)
 {
 	double Length = 0;
-	double SegmentLength = Tolerance / 5.;
 
 	const TArray<FOrientedEdge>& Edges = Context.Loop.GetEdges();
 
@@ -216,12 +217,13 @@ bool FFaceAnalyzer::IsThinFace(double& OutGapSize)
 	Analyze(Context);
 	Chronos.AnalyzeClosedSegmentTime = FChrono::Elapse(StartTime);
 
-	if (Context.MaxSquareDistance < SquareTolerance)
+	if (Context.ThinSideEdgeLength > Context.OppositSideEdgeLength && Context.MaxSquareDistance < SquareTolerance )
 	{
 		OutGapSize = sqrt(Context.MaxSquareDistance);
 		return true;
 	}
 
+	OutGapSize = DOUBLE_BIG_NUMBER;
 	return false;
 }
 
diff --git a/Engine/Source/Runtime/Datasmith/CADKernel/Private/CADKernel/Topo/FaceAnalyzer.h b/Engine/Source/Runtime/Datasmith/CADKernel/Private/CADKernel/Topo/FaceAnalyzer.h
--- a/Engine/Source/Runtime/Datasmith/CADKernel/Private/CADKernel/Topo/FaceAnalyzer.h
+++ b/Engine/Source/Runtime/Datasmith/CADKernel/Private/CADKernel/Topo/FaceAnalyzer.h
@@ -156,7 +156,8 @@ struct FThinFaceContext
 	TArray<double> EdgeMaxSquareDistance;
 
 	double MaxSquareDistance = 0;
-	double SideEdgeLength = 0;
+	double ThinSideEdgeLength = 0;
+	double OppositSideEdgeLength = 0;
 
 	FThinFaceContext(const FTopologicalLoop& InLoop)
 		: Loop(InLoop)
diff --git a/Engine/Source/Runtime/Datasmith/CADKernel/Private/CADKernel/Topo/TopologicalEdge.cpp b/Engine/Source/Runtime/Datasmith/CADKernel/Private/CADKernel/Topo/TopologicalEdge.cpp
--- a/Engine/Source/Runtime/Datasmith/CADKernel/Private/CADKernel/Topo/TopologicalEdge.cpp
+++ b/Engine/Source/Runtime/Datasmith/CADKernel/Private/CADKernel/Topo/TopologicalEdge.cpp
@@ -98,12 +98,7 @@ bool FTopologicalEdge::CheckVertices()
 	return CheckExtremityGap(EndVertex, Points[1]);
 }
 
-double FTopologicalEdge::GetTolerance3D()
-{
-	return GetCurve()->GetCarrierSurface()->Get3DTolerance();
-}
-
-bool FTopologicalEdge::CheckIfDegenerated()
+bool FTopologicalEdge::CheckIfDegenerated() const
 {
 	bool bDegeneration2D = false;
 	bool bDegeneration3D = false;
@@ -167,21 +162,118 @@ TSharedPtr<FTopologicalEdge> FTopologicalEdge::ReturnIfValid(TSharedRef<FTopolog
 	return InEdge;
 }
 
-void FTopologicalEdge::Link(FTopologicalEdge& Twin, double SquareJoiningTolerance)
+bool FTopologicalEdge::HasSameLengthAs(const FTopologicalEdge& Edge, double EdgeLengthTolerance) const
 {
-	// Degenerated twin edges are not linked
-	if (IsDegenerated() || Twin.IsDegenerated())
+	double Min;
+	double Max;
+	if (Length() < Edge.Length())
 	{
-		SetAsDegenerated();
-		Twin.SetAsDegenerated();
-		return;
+		Min = Length();
+		Max = Edge.Length();
+	}
+	else
+	{
+		Min = Edge.Length();
+		Max = Length();
+	}
+
+	if (Min / Max > 0.95) // 95 %
+	{
+		return true;
+	}
+
+	if ((Max - Min) < EdgeLengthTolerance)
+	{
+		return true;
+	}
+
+	return false;
+};
+
+bool FTopologicalEdge::IsTangentAtExtremitiesWith(const FTopologicalEdge& Edge) const
+{
+	TFunction<bool(EOrientation)> IsTangentAtExtremities = [&](EOrientation Orientation) ->  bool
+	{
+		FPoint EdgeStartTangent;
+		FPoint EdgeEndTangent;
+		FPoint StartTangent;
+		FPoint EndTangent;
+		GetTangentsAtExtremities(StartTangent, EndTangent, Orientation);
+		Edge.GetTangentsAtExtremities(EdgeStartTangent, EdgeEndTangent, EOrientation::Front);
+
+		double StartAngle = EdgeStartTangent.ComputeCosinus(StartTangent);
+		double EndAngle = EdgeEndTangent.ComputeCosinus(EndTangent);
+
+		if (StartAngle >= UE_DOUBLE_HALF_SQRT_3 && EndAngle >= UE_DOUBLE_HALF_SQRT_3)
+		{
+			return true;
+		}
+		else
+		{
+			return false;
+		}
+	};
+
+	const FTopologicalVertex* ActiveEdgeVertex1 = &*Edge.GetStartVertex()->GetLinkActiveEntity();
+	const FTopologicalVertex* ActiveEdgeVertex2 = &*Edge.GetEndVertex()->GetLinkActiveEntity();
+	const FTopologicalVertex* ActiveVertex1 = &*GetStartVertex()->GetLinkActiveEntity();
+	const FTopologicalVertex* ActiveVertex2 = &*GetEndVertex()->GetLinkActiveEntity();
+
+	if ((ActiveVertex1 == ActiveEdgeVertex1) && (ActiveVertex2 == ActiveEdgeVertex2))
+	{
+		if (!IsTangentAtExtremities(EOrientation::Front))
+		{
+			if (ActiveVertex1 != ActiveEdgeVertex2)
+			{
+				return false;
+			}
+
+			// Self connected case
+			return IsTangentAtExtremities(EOrientation::Back);
+		}
+		return true;
+	}
+	else if ((ActiveVertex1 == ActiveEdgeVertex2) && (ActiveVertex2 == ActiveEdgeVertex1))
+	{
+		return IsTangentAtExtremities(EOrientation::Back);
+	}
+	return false;
+}
+
+bool FTopologicalEdge::IsLinkableTo(const FTopologicalEdge& Edge, double EdgeLengthTolerance) const
+{
+	if (IsDeleted() || Edge.IsDeleted() ||
+		IsDegenerated() || Edge.IsDegenerated())
+	{
+		return false;
 	}
 
+	if (!HasSameLengthAs(Edge, EdgeLengthTolerance))
+	{
+		return false;
+	}
+
+	return IsTangentAtExtremitiesWith(Edge);
+}
+
+void FTopologicalEdge::LinkIfCoincident(FTopologicalEdge& Twin, double EdgeLengthTolerance, double SquareJoiningTolerance)
+{
 	if(IsDeleted() || Twin.IsDeleted())
 	{
 		return;
 	}
 
+	// Degenerated twin edges are not linked
+	if (IsDegenerated() || Twin.IsDegenerated())
+	{
+		if(HasSameLengthAs(Twin, EdgeLengthTolerance))
+		{
+			SetAsDegenerated();
+			Twin.SetAsDegenerated();
+		}
+		return;
+	}
+
 	const FPoint& Edge1Vertex1 = GetStartVertex()->GetBarycenter();
 	const FPoint& Edge1Vertex2 = GetEndVertex()->GetBarycenter();
 	const FPoint& Edge2Vertex1 = Twin.GetStartVertex()->GetBarycenter();
@@ -194,7 +286,6 @@ void FTopologicalEdge::Link(FTopologicalEdge& Twin, double SquareJoiningToleranc
 	const double SquareDistanceE1V1_E2V2 = GetStartVertex()->IsLinkedTo(Twin.GetEndVertex()) ? 0. : Edge1Vertex1.SquareDistance(Edge2Vertex2);
 	const double SquareDistanceE1V2_E2V1 = GetEndVertex()->IsLinkedTo(Twin.GetStartVertex()) ? 0. : Edge1Vertex2.SquareDistance(Edge2Vertex1);
 
-	bool bCanMergeEdge = true;
 	const double SquareDistanceSameOrientation = SquareDistanceE1V1_E2V1 + SquareDistanceE1V2_E2V2;
 	const double SquareDistanceReverseOrientation = SquareDistanceE1V1_E2V2 + SquareDistanceE1V2_E2V1;
 	if (SquareDistanceSameOrientation < SquareDistanceReverseOrientation)
@@ -206,7 +297,7 @@ void FTopologicalEdge::Link(FTopologicalEdge& Twin, double SquareJoiningToleranc
 		else
 		{
 			FMessage::Printf(Log, TEXT("Edge %d and Edge %d are to far (%f) to be connected\n"), GetId(), Twin.GetId(), sqrt(SquareDistanceE1V1_E2V1));
-			bCanMergeEdge = false;
+			return;
 		}
 
 		if (SquareDistanceE1V2_E2V2 < SquareJoiningTolerance)
@@ -216,7 +307,7 @@ void FTopologicalEdge::Link(FTopologicalEdge& Twin, double SquareJoiningToleranc
 		else
 		{
 			FMessage::Printf(Log, TEXT("Edge %d and Edge %d are to far (%f) to be connected\n"), GetId(), Twin.GetId(), sqrt(SquareDistanceE1V2_E2V2));
-			bCanMergeEdge = false;
+			return;
 		}
 	}
 	else
@@ -228,7 +319,7 @@ void FTopologicalEdge::Link(FTopologicalEdge& Twin, double SquareJoiningToleranc
 		else
 		{
 			FMessage::Printf(Log, TEXT("Edge %d and Edge %d are to far (%f) to be connected\n"), GetId(), Twin.GetId(), sqrt(SquareDistanceE1V1_E2V2));
-			bCanMergeEdge = false;
+			return;
 		}
 
 		if (SquareDistanceE1V2_E2V1 < SquareJoiningTolerance)
@@ -238,16 +329,33 @@ void FTopologicalEdge::Link(FTopologicalEdge& Twin, double SquareJoiningToleranc
 		else
 		{
 			FMessage::Printf(Log, TEXT("Edge %d and Edge %d are to far (%f) to be connected\n"), GetId(), Twin.GetId(), sqrt(SquareDistanceE1V2_E2V1));
-			bCanMergeEdge = false;
+			return;
 		}
 	}
 
-	if (bCanMergeEdge)
+	if (IsLinkableTo(Twin, EdgeLengthTolerance))
 	{
 		MakeLink(Twin);
 	}
 }
 
+void FTopologicalEdge::Link(FTopologicalEdge& Twin)
+{
+	if (IsDegenerated() || Twin.IsDegenerated())
+	{
+		SetAsDegenerated();
+		Twin.SetAsDegenerated();
+		return;
+	}
+
+	if (IsDeleted() || Twin.IsDeleted())
+	{
+		return;
+	}
+
+	MakeLink(Twin);
+}
+
 void FTopologicalEdge::Delete()
 {
 	StartVertex->RemoveConnectedEdge(*this);
@@ -604,7 +712,7 @@ TSharedPtr<FTopologicalEdge> FTopologicalEdge::CreateEdgeByMergingEdges(TArray<F
 	//const double Tolerance2D = Edges[0].Entity->GetCurve()->Get2DCurve()->GetDomainTolerance();
 	TSharedRef<FSurface> CarrierSurface = Edges[0].Entity->GetCurve()->GetCarrierSurface();
 
-	// check if all curve a 2D NURBS
+	// check if all curves are 2D NURBS
 	bool bAreNurbs = true;
 	int32 NurbsMaxDegree = 0;
 
@@ -827,10 +935,9 @@ TSharedPtr<FTopologicalEdge> FTopologicalEdge::CreateEdgeByMergingEdges(TArray<F
 	TSharedPtr<FTopologicalEdge> NewEdge = Make(RestrictionCurve, StartVertex, EndVertex);
 	if (!NewEdge.IsValid())
 	{
-		printf("SSS");
+		return TSharedPtr<FTopologicalEdge>();
 	}
 
-
 	FTopologicalLoop* Loop = Edges[0].Entity->GetLoop();
 	ensureCADKernel(Loop != nullptr);
 	Loop->ReplaceEdges(Edges, NewEdge);
@@ -994,7 +1101,6 @@ TSharedPtr<FTopologicalVertex> FTopologicalEdge::SplitAt(double SplittingCoordin
 	if (GetTwinEntityCount() > 1)
 	{
 		return TSharedPtr<FTopologicalVertex>();
-		ensureCADKernel(false);
 		// TODO
 	}
 
@@ -1028,13 +1134,13 @@ TSharedPtr<FTopologicalVertex> FTopologicalEdge::SplitAt(double SplittingCoordin
 		Boundary.Min = SplittingCoordinate;
 	}
 	MiddelVertex->AddConnectedEdge(*this);
-	Length3D = -1.;
+	Length3D = Curve->ApproximateLength(Boundary);
 
 	Loop->SplitEdge(*this, NewEdge, bKeepStartVertexConnectivity);
 	return MiddelVertex;
 }
 
-bool FTopologicalEdge::IsSharpEdge()
+bool FTopologicalEdge::IsSharpEdge() const
 {
 	double EdgeLength = Length();
 	double Step = Boundary.Length() / 7;
diff --git a/Engine/Source/Runtime/Datasmith/CADKernel/Private/CADKernel/Topo/TopologicalFace.cpp b/Engine/Source/Runtime/Datasmith/CADKernel/Private/CADKernel/Topo/TopologicalFace.cpp
--- a/Engine/Source/Runtime/Datasmith/CADKernel/Private/CADKernel/Topo/TopologicalFace.cpp
+++ b/Engine/Source/Runtime/Datasmith/CADKernel/Private/CADKernel/Topo/TopologicalFace.cpp
@@ -211,6 +211,11 @@ void FTopologicalFace::ApplyNaturalLoops(const FSurfacicBoundary& Boundaries)
 	EndPoint.Set(Boundaries[EIso::IsoU].Min, Boundaries[EIso::IsoV].Min);
 	BuildEdge(StartPoint, EndPoint);
 
+	if (Edges.IsEmpty())
+	{
+		return;
+	}
+
 	TSharedPtr<FTopologicalEdge> PreviousEdge = Edges.Last();
 	for (TSharedPtr<FTopologicalEdge>& Edge : Edges)
 	{
diff --git a/Engine/Source/Runtime/Datasmith/CADKernel/Private/CADKernel/Topo/Topomaker.cpp b/Engine/Source/Runtime/Datasmith/CADKernel/Private/CADKernel/Topo/Topomaker.cpp
--- a/Engine/Source/Runtime/Datasmith/CADKernel/Private/CADKernel/Topo/Topomaker.cpp
+++ b/Engine/Source/Runtime/Datasmith/CADKernel/Private/CADKernel/Topo/Topomaker.cpp
@@ -22,196 +22,20 @@
 namespace UE::CADKernel
 {
 
-FTopomaker::FTopomaker(FSession& InSession, double InTolerance)
-	: Session(InSession)
+namespace TopomakerTools
 {
-	SetTolerance(InTolerance);
-
-	int32 ShellCount = 0;
-	for (const TSharedPtr<FBody>& Body : Session.GetModel().GetBodies())
-	{
-		ShellCount += Body->GetShells().Num();
-	}
-	Shells.Reserve(ShellCount);
 
-	for (const TSharedPtr<FBody>& Body : Session.GetModel().GetBodies())
-	{
-		for (const TSharedPtr<FShell>& Shell : Body->GetShells())
-		{
-			Shells.Add(Shell.Get());
-		}
-	}
-
-	InitFaces();
-}
-
-FTopomaker::FTopomaker(FSession& InSession, const TArray<TSharedPtr<FTopologicalFace>>& InFaces, double InTolerance)
-	: Session(InSession)
-	, Faces(InFaces)
+/**
+ * Merge Border Vertices with other vertices.
+ * @param Vertices: the initial array of active vertices to process, this array is updated at the end of the process
+ */
+void MergeCoincidentVertices(TArray<TSharedPtr<FTopologicalVertex>>& VerticesToMerge, double Tolerance)
 {
-	SetTolerance(InTolerance);
-}
+	double SquareTolerance = FMath::Square(Tolerance);
 
-FTopomaker::FTopomaker(FSession& InSession, const TArray<TSharedPtr<FShell>>& InShells, double InTolerance)
-	: Session(InSession)
-{
-	SetTolerance(InTolerance);
-
-	Shells.Reserve(InShells.Num());
-	for (const TSharedPtr<FShell>& Shell : InShells)
-	{
-		Shells.Add(Shell.Get());
-	}
-	InitFaces();
-}
-
-void FTopomaker::InitFaces()
-{
-	int32 FaceCount = 0;
-	for (FShell* Shell : Shells)
-	{
-		Shell->CompleteMetadata();
-		FaceCount += Shell->FaceCount();
-	}
-	Faces.Reserve(FaceCount);
-
-	for (FShell* Shell : Shells)
-	{
-		Shell->SpreadBodyOrientation();
-		for (const FOrientedFace& Face : Shell->GetFaces())
-		{
-			Faces.Add(Face.Entity);
-		}
-		for (const FOrientedFace& Face : Shell->GetFaces())
-		{
-			Face.Entity->CompleteMetadata();
-		}
-	}
-
-	for (const TSharedPtr<FTopologicalFace>& Face : Faces)
-	{
-		Face->ResetMarker2();
-	}
-}
-
-void FTopomaker::RemoveFacesFromShell()
-{
-	// remove faces from their shells
-	TSet<FShell*> ShellSet;
-	for (const TSharedPtr<FTopologicalFace>& Face : Faces)
-	{
-		if (Face->GetHost() != nullptr)
-		{
-			ShellSet.Add((FShell*)Face->GetHost());
-			Face->ResetHost();
-		}
-	}
-
-	for (FShell* Shell : ShellSet)
-	{
-		bool bIsOutter = Shell->IsOutter();
-
-		TArray<FOrientedFace> ShellFace;
-		ShellFace.Reserve(Shell->FaceCount());
-		for (const FOrientedFace& Face : Shell->GetFaces())
-		{
-			if (!Face.Entity->GetHost())
-			{
-				if (bIsOutter != (Face.Direction == EOrientation::Front))
-				{
-					Face.Entity->SetBackOriented();
-				}
-			}
-			else
-			{
-				ShellFace.Emplace(Face);
-			}
-		}
-		Shell->ReplaceFaces(ShellFace);
-	}
-}
-
-void FTopomaker::EmptyShells()
-{
-	for (FShell* Shell : Shells)
-	{
-		Shell->Empty();
-	}
-}
-
-void FTopomaker::Sew()
-{
-	FTimePoint StartJoinTime = FChrono::Now();
-
-	TArray<FTopologicalVertex*> BorderVertices;
-	GetBorderVertices(BorderVertices);
-	MergeCoincidentVertices(BorderVertices);
-
-	CheckSelfConnectedEdge();
-
-	MergeCoincidentEdges(BorderVertices);
-
-	MergeUnconnectedAdjacentEdges();
-
-	StitchParallelEdges(BorderVertices);
-
-#ifdef CADKERNEL_DEV
-	Report.SewDuration = FChrono::Elapse(StartJoinTime);
-	FChrono::PrintClockElapse(EVerboseLevel::Log, TEXT(""), TEXT("Sew"), Report.SewDuration);
-#endif
-}
-
-void FTopomaker::GetVertices(TArray<FTopologicalVertex*>& Vertices)
-{
-	Vertices.Empty(10 * Faces.Num());
-
-	for (TSharedPtr<FTopologicalFace> Face : Faces)
-	{
-		for (TSharedPtr<FTopologicalLoop> Loop : Face->GetLoops())
-		{
-			for (const FOrientedEdge& OrientedEdge : Loop->GetEdges())
-			{
-				const TSharedPtr<FTopologicalEdge>& Edge = OrientedEdge.Entity;
-				if (!Edge->GetStartVertex()->GetLinkActiveEntity()->HasMarker1())
-				{
-					Edge->GetStartVertex()->GetLinkActiveEntity()->SetMarker1();
-					Vertices.Add(&*Edge->GetStartVertex()->GetLinkActiveEntity());
-				}
-				if (!Edge->GetEndVertex()->GetLinkActiveEntity()->HasMarker1())
-				{
-					Edge->GetEndVertex()->GetLinkActiveEntity()->SetMarker1();
-					Vertices.Add(&*Edge->GetEndVertex()->GetLinkActiveEntity());
-				}
-			}
-		}
-	}
-	for (FTopologicalVertex* Vertex : Vertices)
-	{
-		Vertex->ResetMarker1();
-	}
-}
-
-void FTopomaker::GetBorderVertices(TArray<FTopologicalVertex*>& BorderVertices)
-{
-	TArray<FTopologicalVertex*> Vertices;
-	GetVertices(Vertices);
-
-	BorderVertices.Empty(Vertices.Num());
-
-	for (FTopologicalVertex* Vertex : Vertices)
-	{
-		if (Vertex->IsBorderVertex())
-		{
-			BorderVertices.Add(Vertex);
-		}
-	}
-}
-
-void FTopomaker::MergeCoincidentVertices(TArray<FTopologicalVertex*>& VerticesToMerge)
-{
 	FTimePoint StartTime = FChrono::Now();
 
-	const double WeigthTolerance = 3 * SewTolerance;
+	const double WeigthTolerance = 3 * Tolerance;
 
 	int32 VertexNum = (int32)VerticesToMerge.Num();
 
@@ -221,7 +45,7 @@ void FTopomaker::MergeCoincidentVertices(TArray<FTopologicalVertex*>& VerticesTo
 	TArray<int32> SortedVertexIndices;
 	SortedVertexIndices.Reserve(VertexNum);
 
-	for (FTopologicalVertex* Vertex : VerticesToMerge)
+	for (TSharedPtr<FTopologicalVertex>& Vertex : VerticesToMerge)
 	{
 		VerticesWeight.Add(Vertex->GetCoordinates().DiagonalAxisCoordinate());
 	}
@@ -234,7 +58,7 @@ void FTopomaker::MergeCoincidentVertices(TArray<FTopologicalVertex*>& VerticesTo
 
 	for (int32 IndexI = 0; IndexI < VertexNum; ++IndexI)
 	{
-		FTopologicalVertex* Vertex = VerticesToMerge[SortedVertexIndices[IndexI]];
+		TSharedPtr<FTopologicalVertex>& Vertex = VerticesToMerge[SortedVertexIndices[IndexI]];
 		if (Vertex->HasMarker1())
 		{
 			continue;
@@ -249,7 +73,7 @@ void FTopomaker::MergeCoincidentVertices(TArray<FTopologicalVertex*>& VerticesTo
 
 		for (int32 IndexJ = IndexI + 1; IndexJ < VertexNum; ++IndexJ)
 		{
-			FTopologicalVertex* OtherVertex = VerticesToMerge[SortedVertexIndices[IndexJ]];
+			TSharedPtr<FTopologicalVertex>& OtherVertex = VerticesToMerge[SortedVertexIndices[IndexJ]];
 			if (OtherVertex->HasMarker1())
 			{
 				continue;
@@ -262,7 +86,7 @@ void FTopomaker::MergeCoincidentVertices(TArray<FTopologicalVertex*>& VerticesTo
 			}
 
 			double DistanceSqr = OtherVertex->GetLinkActiveEntity()->SquareDistance(Barycenter);
-			if (DistanceSqr < SewToleranceSquare)
+			if (DistanceSqr < SquareTolerance)
 			{
 				OtherVertex->SetMarker1();
 				Vertex->Link(*OtherVertex);
@@ -271,26 +95,26 @@ void FTopomaker::MergeCoincidentVertices(TArray<FTopologicalVertex*>& VerticesTo
 		}
 	}
 
-	for (FTopologicalVertex* Vertex : VerticesToMerge)
+	for (TSharedPtr<FTopologicalVertex>& Vertex : VerticesToMerge)
 	{
 		Vertex->ResetMarker1();
 	}
 
-	TArray<FTopologicalVertex*> ActiveVertices;
+	TArray<TSharedPtr<FTopologicalVertex>> ActiveVertices;
 	ActiveVertices.Reserve(VertexNum);
 
-	for (FTopologicalVertex* Vertex : VerticesToMerge)
+	for (TSharedPtr<FTopologicalVertex>& Vertex : VerticesToMerge)
 	{
-		FTopologicalVertex& ActiveVertex = *Vertex->GetLinkActiveEntity();
-		if (ActiveVertex.HasMarker1())
+		TSharedPtr<FTopologicalVertex> ActiveVertex = Vertex->GetLinkActiveEntity();
+		if (ActiveVertex->HasMarker1())
 		{
 			continue;
 		}
-		ActiveVertex.SetMarker1();
-		ActiveVertices.Add(&ActiveVertex);
+		ActiveVertex->SetMarker1();
+		ActiveVertices.Add(ActiveVertex);
 	}
 
-	for (FTopologicalVertex* Vertex : ActiveVertices)
+	for (TSharedPtr<FTopologicalVertex>& Vertex : ActiveVertices)
 	{
 		Vertex->ResetMarker1();
 	}
@@ -302,98 +126,181 @@ void FTopomaker::MergeCoincidentVertices(TArray<FTopologicalVertex*>& VerticesTo
 
 }
 
-void FTopomaker::MergeBorderVerticesWithCoincidentOtherVertices(TArray<FTopologicalVertex*>& Vertices)
+TSharedPtr<FTopologicalVertex> SplitAndLink(FTopologicalVertex& StartVertex, FTopologicalEdge& EdgeToLink, FTopologicalEdge& EdgeToSplit, double SquareSewTolerance, double SquareMinEdgeLength)
 {
-	int32 VertexNum = (int32)Vertices.Num();
+	FTopologicalVertex* VertexToLink = EdgeToLink.GetOtherVertex(StartVertex);
+	FTopologicalVertex* EndVertex = EdgeToSplit.GetOtherVertex(StartVertex);
 
-	TArray<double> VerticesWeight;
-	VerticesWeight.Reserve(VertexNum);
-	TArray<int32> SortedVertexIndices;
-	SortedVertexIndices.Reserve(VertexNum);
+	double SquareDistanceToVertexToLink = EndVertex->SquareDistanceBetweenBarycenters(*VertexToLink);
+	if (SquareDistanceToVertexToLink < SquareMinEdgeLength)
+	{
+		VertexToLink->Link(*EndVertex);
+		EdgeToLink.Link(EdgeToSplit);
+		return VertexToLink->GetLinkActiveEntity();
+	}
 
-	for (FTopologicalVertex* Vertex : Vertices)
+	double SquareDistanceToStartVertex = StartVertex.SquareDistanceBetweenBarycenters(*VertexToLink);
+	if (SquareDistanceToStartVertex < SquareMinEdgeLength)
 	{
-		VerticesWeight.Add(Vertex->GetCoordinates().X + Vertex->GetCoordinates().Y + Vertex->GetCoordinates().Z);
+		VertexToLink->Link(StartVertex);
+		EdgeToLink.SetAsDegenerated();
+		return VertexToLink->GetLinkActiveEntity();
 	}
 
-	for (int32 Index = 0; Index < Vertices.Num(); ++Index)
+	FPoint ProjectedPoint;
+	double UProjectedPoint = EdgeToSplit.ProjectPoint(VertexToLink->GetBarycenter(), ProjectedPoint);
+
+	double SquareDistanceToProjectedPoint = ProjectedPoint.SquareDistance(VertexToLink->GetBarycenter());
+	if (SquareDistanceToProjectedPoint > SquareSewTolerance)
 	{
-		SortedVertexIndices.Add(Index);
+		return TSharedPtr<FTopologicalVertex>();
 	}
-	SortedVertexIndices.Sort([&VerticesWeight](const int32& Index1, const int32& Index2) { return VerticesWeight[Index1] < VerticesWeight[Index2]; });
 
-	int32 StartIndexJ = 1;
-	for (int32 IndexI = 0; IndexI < VertexNum; ++IndexI)
+	double SquareDistanceToOtherPoint = ProjectedPoint.SquareDistance(EndVertex->GetBarycenter());
+	if (SquareDistanceToProjectedPoint < SquareMinEdgeLength)
+	{
+		// the new point is closed to the extremity, a degenerated edge will be created, so the edges are joined
+		VertexToLink->Link(*EndVertex);
+		EdgeToLink.Link(EdgeToSplit);
+		return VertexToLink->GetLinkActiveEntity();
+	}
+
+	SquareDistanceToProjectedPoint = ProjectedPoint.SquareDistance(StartVertex.GetBarycenter());
+	if (SquareDistanceToProjectedPoint < SquareMinEdgeLength)
+	{
+		VertexToLink->Link(StartVertex);
+		EdgeToLink.SetAsDegenerated();
+		return VertexToLink->GetLinkActiveEntity();
+	}
+
+	// JoinParallelEdges process all edges connected to startVertex (ConnectedEdges).
+	// Connected edges must remain compliant i.e. all edges of ConnectedEdges must be connected to StartVertex
+	// EdgeToSplit->SplitAt() must keep EdgeToSplit connected to StartVertex
+	bool bKeepStartVertexConnectivity = (StartVertex.GetLink() == EdgeToSplit.GetStartVertex()->GetLink());
+
+	TSharedPtr<FTopologicalEdge> NewEdge;
+	TSharedPtr<FTopologicalVertex> NewVertex = EdgeToSplit.SplitAt(UProjectedPoint, ProjectedPoint, bKeepStartVertexConnectivity, NewEdge);
+	if (!NewVertex.IsValid())
+	{
+		return TSharedPtr<FTopologicalVertex>();
+	}
+
+	VertexToLink->Link(*NewVertex);
+	EdgeToLink.Link(EdgeToSplit);
+	return NewVertex;
+}
+
+void StitchParallelEdges(TArray<TSharedPtr<FTopologicalVertex>>& VerticesToProcess, double SewTolerance, double MinEdgeLength)
+{
+	double SquareTolerance = FMath::Square(SewTolerance);
+	double SquareMinEdgeLength = FMath::Square(MinEdgeLength);
+
+	FTimePoint StartTime = FChrono::Now();
+
+	for (int32 VertexI = 0; VertexI < VerticesToProcess.Num(); ++VertexI)
 	{
-		FTopologicalVertex* Vertex = Vertices[SortedVertexIndices[IndexI]];
-		if (!Vertex->IsBorderVertex())
+		TSharedPtr<FTopologicalVertex>& Vertex = VerticesToProcess[VertexI];
+
+		if (!Vertex.IsValid() || Vertex->IsDeleted() || !Vertex->IsBorderVertex())
 		{
 			continue;
 		}
 
-		double VertexWeigth = VerticesWeight[SortedVertexIndices[IndexI]];
-		FPoint Barycenter = Vertex->GetBarycenter();
-		for (int32 IndexJ = StartIndexJ; IndexJ < VertexNum; ++IndexJ)
+		TArray<FTopologicalEdge*> ConnectedEdges;
+		Vertex->GetConnectedEdges(ConnectedEdges);
+		const int32 ConnectedEdgeCount = ConnectedEdges.Num();
+		if (ConnectedEdgeCount == 1)
 		{
-			FTopologicalVertex& OtherVertex = *Vertices[SortedVertexIndices[IndexJ]];
-			if (Vertex->GetLink() == OtherVertex.GetLink())
-			{
-				continue;
-			}
+			continue;
+		}
 
-			double OtherVertexWeigth = VerticesWeight[SortedVertexIndices[IndexJ]];
+		for (int32 EdgeI = 0; EdgeI < ConnectedEdgeCount - 1; ++EdgeI)
+		{
+			FTopologicalEdge* Edge = ConnectedEdges[EdgeI];
+			ensureCADKernel(Edge->GetLoop() != nullptr);
 
-			if (OtherVertexWeigth + SewTolerance - VertexWeigth < 0)
+			if (Edge->IsDegenerated())
 			{
-				StartIndexJ = IndexJ;
 				continue;
 			}
 
-			if (OtherVertexWeigth - SewTolerance - VertexWeigth > 0)
+			if (!Edge->IsActiveEntity())
 			{
-				break;
+				continue;
 			}
+			bool bFirstEdgeBorder = Edge->IsBorder();
 
-			double DistanceSqr = OtherVertex.GetLinkActiveEntity()->SquareDistance(Barycenter);
-			if (DistanceSqr < SewToleranceSquare)
+			for (int32 EdgeJ = EdgeI + 1; EdgeJ < ConnectedEdgeCount; ++EdgeJ)
 			{
-				TArray<FTopologicalEdge*> CommonEdges;
-				Vertex->GetConnectedEdges(OtherVertex, CommonEdges);
-				if (CommonEdges.Num() > 0)
+				FTopologicalEdge* SecondEdge = ConnectedEdges[EdgeJ];
+				if (SecondEdge->IsDegenerated())
+				{
+					continue;
+				}
+				if (!SecondEdge->IsActiveEntity())
 				{
 					continue;
 				}
+				bool bSecondEdgeBorder = SecondEdge->IsBorder();
 
-				OtherVertex.SetMarker1();
-				Vertex->Link(OtherVertex);
-				Barycenter = Vertex->GetBarycenter();
-			}
-		}
-	}
+				// Process only if at least one edge is Border
+				if (!bFirstEdgeBorder && !bSecondEdgeBorder)
+				{
+					continue;
+				}
 
-	TArray<FTopologicalVertex*> NewVertices;
-	NewVertices.Empty(VertexNum);
+				FPoint StartTangentEdge = Edge->GetTangentAt(*Vertex);
+				FPoint StartTangentOtherEdge = SecondEdge->GetTangentAt(*Vertex);
 
-	for (FTopologicalVertex* Vertex : Vertices)
-	{
-		FTopologicalVertex& ActiveVertex = *Vertex->GetLinkActiveEntity();
-		if (ActiveVertex.HasMarker1())
-		{
-			continue;
-		}
-		ActiveVertex.SetMarker1();
-		NewVertices.Add(&ActiveVertex);
-	}
+				double CosAngle = StartTangentEdge.ComputeCosinus(StartTangentOtherEdge);
+				if (CosAngle < UE_DOUBLE_HALF_SQRT_3) // cos(30 deg)
+				{
+					continue;
+				}
 
-	for (FTopologicalVertex* Vertex : NewVertices)
-	{
-		Vertex->ResetMarker1();
+				FTopologicalVertex& EndVertex = *Edge->GetOtherVertex(*Vertex)->GetLinkActiveEntity();
+				FTopologicalVertex& SecondEdgeEndVertex = *SecondEdge->GetOtherVertex(*Vertex)->GetLinkActiveEntity();
+
+				if (&EndVertex == &SecondEdgeEndVertex && Edge->IsLinkableTo(*SecondEdge, MinEdgeLength))
+				{
+					// should not happen but, just in case
+					Edge->Link(*SecondEdge);
+				}
+				else
+				{
+					double EdgeLength = Edge->Length();
+					double SecondEdgeLength = SecondEdge->Length();
+
+					FTopologicalVertex* OtherVertex = nullptr;
+					TSharedPtr<FTopologicalVertex> NewVertex;
+					if (EdgeLength < SecondEdgeLength)
+					{
+						OtherVertex = &SecondEdgeEndVertex;
+						NewVertex = SplitAndLink(*Vertex, *Edge, *SecondEdge, SquareTolerance, SquareMinEdgeLength);
+					}
+					else
+					{
+						OtherVertex = &EndVertex;
+						NewVertex = SplitAndLink(*Vertex, *SecondEdge, *Edge, SquareTolerance, SquareMinEdgeLength);
+					}
+
+					if (NewVertex.IsValid() && !Vertex->IsLinkedTo(*NewVertex) && !OtherVertex->IsLinkedTo(*NewVertex))
+					{
+						VerticesToProcess.Add(NewVertex);
+					}
+				}
+			}
+		}
 	}
 
-	Swap(NewVertices, Vertices);
+	FDuration StepDuration = FChrono::Elapse(StartTime);
+	FChrono::PrintClockElapse(EVerboseLevel::Log, TEXT("    "), TEXT("Stitch Parallel Edges"), StepDuration);
 }
 
-void FTopomaker::MergeCoincidentEdges(FTopologicalEdge* Edge)
+/**
+ * First step, trivial edge merge i.e. couple of edges with same extremity active vertices
+ */
+void MergeCoincidentEdges(FTopologicalEdge* Edge, double MinEdgeLength)
 {
 	const FTopologicalVertex& StartVertex = *Edge->GetStartVertex()->GetLinkActiveEntity();
 	const FTopologicalVertex& EndVertex = *Edge->GetEndVertex()->GetLinkActiveEntity();
@@ -424,213 +331,286 @@ void FTopomaker::MergeCoincidentEdges(FTopologicalEdge* Edge)
 			continue;
 		}
 
-		const FPoint StartTangentEdge = Edge->GetTangentAt(StartVertex);
-		const FPoint StartTangentOtherEdge = SecondEdge->GetTangentAt(StartVertex);
+		if (Edge->GetFace() != SecondEdge->GetFace() && Edge->IsLinkableTo(*SecondEdge, MinEdgeLength))
+		{
+			Edge->Link(*SecondEdge);
+		}
+	}
+}
+
+void MergeCoincidentEdges(TArray<FTopologicalEdge*>& EdgesToProcess, double MinEdgeLength)
+{
+	FTimePoint StartTime = FChrono::Now();
+
+	for (FTopologicalEdge* Edge : EdgesToProcess)
+	{
+		if (Edge->IsDegenerated() || !Edge->IsBorder())
+		{
+			continue;
+		}
+		MergeCoincidentEdges(Edge, MinEdgeLength);
+	}
+
+	FDuration StepDuration = FChrono::Elapse(StartTime);
+	FChrono::PrintClockElapse(EVerboseLevel::Log, TEXT("    "), TEXT("Merge coincident edges"), StepDuration);
+}
+
+void MergeCoincidentEdges(TArray<TSharedPtr<FTopologicalVertex>>& VerticesToProcess, double MinEdgeLength)
+{
+	FTimePoint StartTime = FChrono::Now();
+
+	for (TSharedPtr<FTopologicalVertex>& Vertex : VerticesToProcess)
+	{
+		if (Vertex->IsDeleted() || !Vertex->IsActiveEntity() || Vertex->HasMarker1())
+		{
+			continue;
+		}
+		Vertex->SetMarker1();
+
+		TArray<FTopologicalEdge*> ConnectedEdges;
+		Vertex->GetConnectedEdges(ConnectedEdges);
+		const int32 ConnectedEdgeCount = ConnectedEdges.Num();
+		if (ConnectedEdgeCount == 1)
+		{
+			continue;
+		}
+
+		for (int32 EdgeI = 0; EdgeI < ConnectedEdgeCount - 1; ++EdgeI)
+		{
+			FTopologicalEdge* Edge = ConnectedEdges[EdgeI];
+			if (!Edge->IsActiveEntity() || Edge->IsDegenerated())
+			{
+				continue;
+			}
+
+			MergeCoincidentEdges(Edge, MinEdgeLength);
+		}
+	}
+
+	for (TSharedPtr<FTopologicalVertex>& Vertex : VerticesToProcess)
+	{
+		Vertex->ResetMarker1();
+	}
+
+	FDuration StepDuration = FChrono::Elapse(StartTime);
+	FChrono::PrintClockElapse(EVerboseLevel::Log, TEXT("    "), TEXT("Merge coincident edges"), StepDuration);
+}
+
+} // namespace TopomakerTools
+
+FTopomaker::FTopomaker(FSession& InSession, double InTolerance, double InForceFactor)
+	: Session(InSession)
+{
+	SetTolerance(InTolerance, InForceFactor);
+
+	int32 ShellCount = 0;
+	for (const TSharedPtr<FBody>& Body : Session.GetModel().GetBodies())
+	{
+		ShellCount += Body->GetShells().Num();
+	}
+	Shells.Reserve(ShellCount);
+
+	for (const TSharedPtr<FBody>& Body : Session.GetModel().GetBodies())
+	{
+		for (const TSharedPtr<FShell>& Shell : Body->GetShells())
+		{
+			Shells.Add(Shell.Get());
+		}
+	}
+
+	InitFaces();
+}
 
-		double CosAngle = StartTangentEdge.ComputeCosinus(StartTangentOtherEdge);
-		if (CosAngle < 0.9)
-		{
-			continue;
-		}
+FTopomaker::FTopomaker(FSession& InSession, const TArray<TSharedPtr<FTopologicalFace>>& InFaces, double InTolerance, double InForceFactor)
+	: Session(InSession)
+	, Faces(InFaces)
+{
+	SetTolerance(InTolerance, InForceFactor);
+}
 
-		if (Edge->GetFace() != SecondEdge->GetFace())
-		{
-			Edge->Link(*SecondEdge, SewTolerance);
-		}
+FTopomaker::FTopomaker(FSession& InSession, const TArray<TSharedPtr<FShell>>& InShells, double InTolerance, double InForceFactor)
+	: Session(InSession)
+{
+	SetTolerance(InTolerance, InForceFactor);
+
+	Shells.Reserve(InShells.Num());
+	for (const TSharedPtr<FShell>& Shell : InShells)
+	{
+		Shells.Add(Shell.Get());
 	}
+	InitFaces();
 }
 
-void FTopomaker::MergeCoincidentEdges(TArray<FTopologicalEdge*>& EdgesToProcess)
+void FTopomaker::InitFaces()
 {
-	FTimePoint StartTime = FChrono::Now();
+	int32 FaceCount = 0;
+	for (FShell* Shell : Shells)
+	{
+		Shell->CompleteMetadata();
+		FaceCount += Shell->FaceCount();
+	}
+	Faces.Reserve(FaceCount);
 
-	for (FTopologicalEdge* Edge : EdgesToProcess)
+	for (FShell* Shell : Shells)
 	{
-		if (Edge->IsDegenerated() || !Edge->IsBorder())
+		Shell->SpreadBodyOrientation();
+		for (const FOrientedFace& Face : Shell->GetFaces())
 		{
-			continue;
+			Faces.Add(Face.Entity);
+		}
+		for (const FOrientedFace& Face : Shell->GetFaces())
+		{
+			Face.Entity->CompleteMetadata();
 		}
-		MergeCoincidentEdges(Edge);
 	}
 
-	FDuration StepDuration = FChrono::Elapse(StartTime);
-	FChrono::PrintClockElapse(EVerboseLevel::Log, TEXT("    "), TEXT("Merge coincident edges"), StepDuration);
+	for (const TSharedPtr<FTopologicalFace>& Face : Faces)
+	{
+		Face->ResetMarker2();
+	}
 }
 
-void FTopomaker::MergeCoincidentEdges(TArray<FTopologicalVertex*>& VerticesToProcess)
+void FTopomaker::RemoveFacesFromShell()
 {
-	FTimePoint StartTime = FChrono::Now();
-
-	for (FTopologicalVertex* Vertex : VerticesToProcess)
+	// remove faces from their shells
+	TSet<FShell*> ShellSet;
+	for (const TSharedPtr<FTopologicalFace>& Face : Faces)
 	{
-		TArray<FTopologicalEdge*> ConnectedEdges;
-		Vertex->GetConnectedEdges(ConnectedEdges);
-		const int32 ConnectedEdgeCount = ConnectedEdges.Num();
-		if (ConnectedEdgeCount == 1)
+		if (Face->GetHost() != nullptr)
 		{
-			continue;
+			ShellSet.Add((FShell*)Face->GetHost());
+			Face->ResetHost();
 		}
+	}
 
-		for (int32 EdgeI = 0; EdgeI < ConnectedEdgeCount - 1; ++EdgeI)
+	for (FShell* Shell : ShellSet)
+	{
+		bool bIsOutter = Shell->IsOutter();
+
+		TArray<FOrientedFace> ShellFace;
+		ShellFace.Reserve(Shell->FaceCount());
+		for (const FOrientedFace& Face : Shell->GetFaces())
 		{
-			FTopologicalEdge* Edge = ConnectedEdges[EdgeI];
-			if (!Edge->IsActiveEntity() || Edge->IsDegenerated())
+			if (!Face.Entity->GetHost())
 			{
-				continue;
+				if (bIsOutter != (Face.Direction == EOrientation::Front))
+				{
+					Face.Entity->SetBackOriented();
+				}
+			}
+			else
+			{
+				ShellFace.Emplace(Face);
 			}
-
-			MergeCoincidentEdges(Edge);
 		}
+		Shell->ReplaceFaces(ShellFace);
 	}
+}
 
-	FDuration StepDuration = FChrono::Elapse(StartTime);
-	FChrono::PrintClockElapse(EVerboseLevel::Log, TEXT("    "), TEXT("Merge coincident edges"), StepDuration);
+void FTopomaker::EmptyShells()
+{
+	for (FShell* Shell : Shells)
+	{
+		Shell->Empty();
+	}
 }
 
-TSharedPtr<FTopologicalVertex> FTopomaker::SplitAndLink(FTopologicalVertex& StartVertex, FTopologicalEdge& EdgeToLink, FTopologicalEdge& EdgeToSplit)
+void FTopomaker::Sew(bool bForceJoining, bool bRemoveThinFaces)
 {
-	FTopologicalVertex* VertexToLink = EdgeToLink.GetOtherVertex(StartVertex);
+	FTimePoint StartJoinTime = FChrono::Now();
 
-	FPoint ProjectedPoint;
-	double UProjectedPoint = EdgeToSplit.ProjectPoint(VertexToLink->GetBarycenter(), ProjectedPoint);
+	TArray<TSharedPtr<FTopologicalVertex>> BorderVertices;
+	GetBorderVertices(BorderVertices);
+	TopomakerTools::MergeCoincidentVertices(BorderVertices, SewTolerance);
 
-	double SquareDistanceToProjectedPoint = ProjectedPoint.SquareDistance(VertexToLink->GetBarycenter());
-	if (SquareDistanceToProjectedPoint > SewToleranceSquare)
-	{
-		return TSharedPtr<FTopologicalVertex>();
-	}
+	CheckSelfConnectedEdge(EdgeLengthTolerance, BorderVertices);
 
-	// Check if the ProjectedPoint is not nearly equal edge boundary
-	FTopologicalVertex* EndVertex = EdgeToSplit.GetOtherVertex(StartVertex);
-	if (EndVertex->SquareDistance(ProjectedPoint) < SewToleranceSquare)
+	// basic case: merge pair of edges connected together at both extremities i.e. pair of edges with same extremity active vertices.
+	TopomakerTools::MergeCoincidentEdges(BorderVertices, EdgeLengthTolerance);
+
+	if (bRemoveThinFaces)
 	{
-		VertexToLink->Link(*EndVertex);
-		EdgeToLink.Link(EdgeToSplit, SewTolerance);
-		// TSharedPtr<FTopologicalVertex>() is returned as EndVertex is not new
-		return TSharedPtr<FTopologicalVertex>();
+		TArray<FTopologicalEdge*> NewBorderEdges;
+		RemoveThinFaces(NewBorderEdges);
 	}
 
-	// JoinParallelEdges process all edges connected to startVertex (ConnectedEdges).
-	// Connected edges must remain compliant i.e. all edges of ConnectedEdges must be connected to StartVertex
-	// EdgeToSplit->SplitAt() must keep EdgeToSplit connected to StartVertex
-	bool bKeepStartVertexConnectivity = (StartVertex.GetLink() == EdgeToSplit.GetStartVertex()->GetLink());
+	MergeUnconnectedSuccessiveEdges();
 
-	TSharedPtr<FTopologicalEdge> NewEdge;
-	TSharedPtr<FTopologicalVertex> NewVertex = EdgeToSplit.SplitAt(UProjectedPoint, ProjectedPoint, bKeepStartVertexConnectivity, NewEdge);
-	if (!NewVertex.IsValid())
+	if (bForceJoining)
 	{
-		return TSharedPtr<FTopologicalVertex>();
+		BorderVertices.Empty(BorderVertices.Num());
+		GetBorderVertices(BorderVertices);
+		TopomakerTools::MergeCoincidentVertices(BorderVertices, SewToleranceToForceJoin);
+
+		CheckSelfConnectedEdge(LargeEdgeLengthTolerance, BorderVertices);
 	}
 
-	VertexToLink->Link(*NewVertex);
-	EdgeToLink.Link(EdgeToSplit, SewTolerance);
+	// re process with new edges from MergeUnconnectedSuccessiveEdges and new merged vertices (if bForceJoining)
+	TopomakerTools::MergeCoincidentEdges(BorderVertices, bForceJoining ? LargeEdgeLengthTolerance : EdgeLengthTolerance);
 
-	return NewVertex;
+	// advance case: two partially coincident edges connected at one extremity i.e. coincident along the shortest edge
+	TopomakerTools::StitchParallelEdges(BorderVertices, bForceJoining ? SewToleranceToForceJoin : SewTolerance, bForceJoining ? LargeEdgeLengthTolerance : EdgeLengthTolerance);
+
+#ifdef CADKERNEL_DEV
+	Report.SewDuration = FChrono::Elapse(StartJoinTime);
+	FChrono::PrintClockElapse(EVerboseLevel::Log, TEXT(""), TEXT("Sew"), Report.SewDuration);
+#endif
 }
 
-void FTopomaker::StitchParallelEdges(TArray<FTopologicalVertex*>& VerticesToProcess)
+void FTopomaker::GetVertices(TArray<TSharedPtr<FTopologicalVertex>>& Vertices)
 {
-	FTimePoint StartTime = FChrono::Now();
+	Vertices.Empty(10 * Faces.Num());
 
-	for (int32 VertexI = 0; VertexI < VerticesToProcess.Num(); ++VertexI)
+	for (TSharedPtr<FTopologicalFace> Face : Faces)
 	{
-		FTopologicalVertex* Vertex = VerticesToProcess[VertexI];
-
-		if (!Vertex->IsBorderVertex())
-		{
-			continue;
-		}
-
-		TArray<FTopologicalEdge*> ConnectedEdges;
-		Vertex->GetConnectedEdges(ConnectedEdges);
-		const int32 ConnectedEdgeCount = ConnectedEdges.Num();
-		if (ConnectedEdgeCount == 1)
+		if (Face->IsDeleted())
 		{
 			continue;
 		}
 
-		for (int32 EdgeI = 0; EdgeI < ConnectedEdgeCount - 1; ++EdgeI)
+		for (TSharedPtr<FTopologicalLoop> Loop : Face->GetLoops())
 		{
-			FTopologicalEdge* Edge = ConnectedEdges[EdgeI];
-			ensureCADKernel(Edge->GetLoop() != nullptr);
-
-			if (Edge->IsDegenerated())
-			{
-				continue;
-			}
-
-			if (!Edge->IsActiveEntity())
-			{
-				continue;
-			}
-			bool bFirstEdgeBorder = Edge->IsBorder();
-
-			for (int32 EdgeJ = EdgeI + 1; EdgeJ < ConnectedEdgeCount; ++EdgeJ)
+			for (const FOrientedEdge& OrientedEdge : Loop->GetEdges())
 			{
-				FTopologicalEdge* SecondEdge = ConnectedEdges[EdgeJ];
-				if (Edge->IsDegenerated())
-				{
-					continue;
-				}
-				if (!SecondEdge->IsActiveEntity())
-				{
-					continue;
-				}
-				bool bSecondEdgeBorder = Edge->IsBorder();
-
-				// Process only if at least one edge is Border
-				if (!bFirstEdgeBorder && !bSecondEdgeBorder)
+				const TSharedPtr<FTopologicalEdge>& Edge = OrientedEdge.Entity;
+				if (!Edge->GetStartVertex()->GetLinkActiveEntity()->HasMarker1())
 				{
-					continue;
+					Edge->GetStartVertex()->GetLinkActiveEntity()->SetMarker1();
+					Vertices.Add(Edge->GetStartVertex()->GetLinkActiveEntity());
 				}
-
-				FPoint StartTangentEdge = Edge->GetTangentAt(*Vertex);
-				FPoint StartTangentOtherEdge = SecondEdge->GetTangentAt(*Vertex);
-
-				double CosAngle = StartTangentEdge.ComputeCosinus(StartTangentOtherEdge);
-				if (CosAngle < 0.9)
+				if (!Edge->GetEndVertex()->GetLinkActiveEntity()->HasMarker1())
 				{
-					continue;
+					Edge->GetEndVertex()->GetLinkActiveEntity()->SetMarker1();
+					Vertices.Add(Edge->GetEndVertex()->GetLinkActiveEntity());
 				}
+			}
+		}
+	}
+	for (TSharedPtr<FTopologicalVertex>& Vertex : Vertices)
+	{
+		Vertex->ResetMarker1();
+	}
+}
 
-				FTopologicalVertex& EndVertex = *Edge->GetOtherVertex(*Vertex)->GetLinkActiveEntity();
-				FTopologicalVertex& OtherEdgeEndVertex = *SecondEdge->GetOtherVertex(*Vertex)->GetLinkActiveEntity();
-				if (&EndVertex == &OtherEdgeEndVertex)
-				{
-					Edge->Link(*SecondEdge, SewTolerance);
-				}
-				else
-				{
-					double EdgeLength = Edge->Length();
-					double OtherEdgeLength = SecondEdge->Length();
+void FTopomaker::GetBorderVertices(TArray<TSharedPtr<FTopologicalVertex>>& BorderVertices)
+{
+	TArray<TSharedPtr<FTopologicalVertex>> Vertices;
+	GetVertices(Vertices);
 
-					TSharedPtr<FTopologicalVertex> NewVertex;
-					if (EdgeLength < OtherEdgeLength)
-					{
-						NewVertex = SplitAndLink(*Vertex, *Edge, *SecondEdge);
-					}
-					else
-					{
-						NewVertex = SplitAndLink(*Vertex, *SecondEdge, *Edge);
-					}
+	BorderVertices.Empty(Vertices.Num());
 
-					if (NewVertex.IsValid())
-					{
-						VerticesToProcess.Add(NewVertex.Get());
-					}
-				}
-			}
+	for (TSharedPtr<FTopologicalVertex>& Vertex : Vertices)
+	{
+		if (Vertex->IsBorderVertex())
+		{
+			BorderVertices.Add(Vertex);
 		}
 	}
-
-	FDuration StepDuration = FChrono::Elapse(StartTime);
-	FChrono::PrintClockElapse(EVerboseLevel::Log, TEXT("    "), TEXT("Stitch Parallel Edges"), StepDuration);
 }
 
-void FTopomaker::MergeUnconnectedAdjacentEdges()
+void FTopomaker::MergeUnconnectedSuccessiveEdges()
 {
 	FTimePoint StartTime = FChrono::Now();
-
 	for (TSharedPtr<FTopologicalFace> FacePtr : Faces)
 	{
 		if (!FacePtr.IsValid())
@@ -692,7 +672,7 @@ void FTopomaker::MergeUnconnectedAdjacentEdges()
 						FPoint StartTangentOtherEdge = ConnectedEdges[1]->GetTangentAt(*EndVertex);
 
 						double CosAngle = StartTangentEdge.ComputeCosinus(StartTangentOtherEdge);
-						if (CosAngle < -0.9)
+						if (CosAngle < -UE_DOUBLE_HALF_SQRT_3) // Cos(30 deg)
 						{
 							bEdgeIsNotTheLast = true;
 						}
@@ -718,61 +698,9 @@ void FTopomaker::MergeUnconnectedAdjacentEdges()
 		{
 			if (Candidates.Num())
 			{
-				TSharedPtr<FTopologicalVertex> StartVertexPtr = Candidates[0].Direction == EOrientation::Front ? Candidates[0].Entity->GetStartVertex() : Candidates[0].Entity->GetEndVertex();
-				TSharedRef<FTopologicalVertex> StartVertex = StartVertexPtr.ToSharedRef();
-				TSharedPtr<FTopologicalVertex> EndVertex = Candidates.Last().Direction == EOrientation::Front ? Candidates.Last().Entity->GetEndVertex() : Candidates.Last().Entity->GetStartVertex();
-
-				FPoint StartTangentEdge = Candidates[0].Entity->GetTangentAt(*StartVertex);
-
-				TArray<FTopologicalEdge*> Edges;
-				StartVertex->GetConnectedEdges(*EndVertex, Edges);
-
-				// Remove edge that is not parallel edge
-				FTopologicalEdge* ParallelEdge = nullptr;
-				for (FTopologicalEdge*& Edge : Edges)
-				{
-					if (!Edge)
-					{
-						continue;
-					}
-
-					if (Edge->GetFace() == Face)
-					{
-						Edge = nullptr;
-						continue;
-					}
-
-					FPoint StartTangentOtherEdge = Edge->GetTangentAt(*StartVertex);
-					double CosAngle = StartTangentEdge.ComputeCosinus(StartTangentOtherEdge);
-					if (CosAngle > 0.9)
-					{
-						ParallelEdge = Edge;
-						break;
-					}
-				}
-
-				if (ParallelEdge)
-				{
-					TSharedPtr<FTopologicalEdge> NewEdge = FTopologicalEdge::CreateEdgeByMergingEdges(Candidates, StartVertex, EndVertex.ToSharedRef());
-					if (!NewEdge.IsValid())
-					{
-						// the edges cannot be merge, they will be connected to the parallel edge with "SplitAndLink"
-						// i.e. instead of merging edges to be one and link to the parallel one
-						// the parallel one is split at the extremities of each edges, and each new edge is linked to its parallel edge 
-						break;
-					}
-
-					// Link to the parallel edge
-					// New edge is link to the first parallel one as the other parallel should already be linked together
-					for (FTopologicalEdge* Edge : Edges)
-					{
-						if (Edge)
-						{
-							Edge->Link(*NewEdge, SewTolerance);
-							break;
-						}
-					}
-				}
+				TSharedRef<FTopologicalVertex> StartVertex = Candidates[0].Direction == EOrientation::Front ? Candidates[0].Entity->GetStartVertex() : Candidates[0].Entity->GetEndVertex();
+				TSharedRef<FTopologicalVertex> EndVertex = Candidates.Last().Direction == EOrientation::Front ? Candidates.Last().Entity->GetEndVertex() : Candidates.Last().Entity->GetStartVertex();
+				FTopologicalEdge::CreateEdgeByMergingEdges(Candidates, StartVertex, EndVertex);
 			}
 		}
 	}
@@ -788,10 +716,10 @@ void FTopomaker::RemoveIsolatedEdges()
 
 	TArray<FTopologicalEdge*> IsolatedEdges;
 
-	TArray<FTopologicalVertex*> Vertices;
+	TArray<TSharedPtr<FTopologicalVertex>> Vertices;
 	GetVertices(Vertices);
 
-	for (const FTopologicalVertex* Vertex : Vertices)
+	for (const TSharedPtr<FTopologicalVertex>& Vertex : Vertices)
 	{
 		for (const FTopologicalVertex* TwinVertex : Vertex->GetTwinEntities())
 		{
@@ -814,11 +742,11 @@ void FTopomaker::RemoveIsolatedEdges()
 
 void FTopomaker::UnlinkNonManifoldVertex()
 {
-	TArray<FTopologicalVertex*> Vertices;
+	TArray<TSharedPtr<FTopologicalVertex>> Vertices;
 	GetVertices(Vertices);
 
 	TMap<FShell*, int32> ShellToVertexCount;
-	for (FTopologicalVertex* Vertex : Vertices)
+	for (const TSharedPtr<FTopologicalVertex>& Vertex : Vertices)
 	{
 		ShellToVertexCount.Empty(Vertex->GetTwinEntityCount());
 		for (const FTopologicalVertex* TwinVertex : Vertex->GetTwinEntities())
@@ -900,23 +828,20 @@ void FindBorderVertex(const TArray<FTopologicalEdge*>& Edges, TArray<FTopologica
 
 
 //#define DEBUG_THIN_FACE
-void FTopomaker::RemoveThinFaces()
+void FTopomaker::RemoveThinFaces(TArray<FTopologicalEdge*>& NewBorderEdges)
 {
 	FTimePoint StartTime = FChrono::Now();
 
+#ifdef DEBUG_THIN_FACE
 	F3DDebugSession _(TEXT("RemoveThinFaces"));
+#endif
 
-	// A thin surface is a surface no larger than ThinToleranceFactor * SewTolerance
-	constexpr double ThinToleranceFactor = 5.;
-	const double ThinTolerance = SewTolerance * ThinToleranceFactor;
-
-	TArray<FTopologicalEdge*> NewBorderEdges;
 	TArray<FTopologicalFace*> DeletedFaces;
 
 	// Find thin faces
 	for (TSharedPtr<FTopologicalFace> Face : Faces)
 	{
-		FFaceAnalyzer Analyer(*Face, ThinTolerance);
+		FFaceAnalyzer Analyer(*Face, ThinFaceWidth);
 		double GapSize = 0;
 		if (Analyer.IsThinFace(GapSize))
 		{
@@ -952,48 +877,15 @@ void FTopomaker::RemoveThinFaces()
 		NewBorderEdges.SetNum(NewIndex);
 	}
 
-	const double SmallEdgeTolerance = ThinTolerance * DOUBLE_HALF_PI; // in the case of the thin surface is a Half cylinder, the diameter is ThinTolerance, so the half circle circumference is ThinTolerance * Pi / 2
-	// Collapse edges with a length smaller than SmallEdgeTolerance
-	for (FTopologicalEdge* Edge : NewBorderEdges)
-	{
-		if (Edge->Length() < SmallEdgeTolerance)
-		{
-			Edge->GetStartVertex()->Link(*Edge->GetEndVertex());
-			Edge->SetAsDegenerated();
-		}
-	}
-
-	TArray<FTopologicalVertex*> BorderVertices;
-	BorderVertices.Reserve(NewBorderEdges.Num());
-	FindBorderVertex(NewBorderEdges, BorderVertices);
-
-	// MergeCoincidentVertices tolerance is increase of a factor ThinToleranceFactor * 2.1 because:
-	// to merge the vertices of thin border of the surface we need to set the tolerance to the size of the width of the surface.
-	// If two thin surface are connected, we need to merge the vertices of the both thin border i.e. 2 * max thin width
-	//
-	//            -------------------*                               -----------------
-	//                               |                                                \
-	//            -------------------*----------------------    =>   -------------------*----------------------
-	//                               |                                                    \
-	//                               *----------------------                               ---------------------
-	double TolValue = Tolerance;
-	SetTolerance(ThinTolerance * 2.1);
-	MergeCoincidentVertices(BorderVertices);
-	SetTolerance(TolValue);
-
-	MergeCoincidentEdges(NewBorderEdges);
-
-	StitchParallelEdges(BorderVertices);
-
 #ifdef CADKERNEL_DEV
 	Report.RemoveThinFacesDuration = FChrono::Elapse(StartTime);
 	FChrono::PrintClockElapse(EVerboseLevel::Log, TEXT(""), TEXT(".RemoveThinFaces"), Report.RemoveThinFacesDuration = FChrono::Elapse(StartTime));
 #endif
 }
 
-void FTopomaker::SetSelfConnectedEdgeDegenerated(TArray<FTopologicalVertex*>& Vertices)
+void FTopomaker::SetSelfConnectedEdgeDegenerated(TArray<TSharedPtr<FTopologicalVertex>>& Vertices)
 {
-	for (const FTopologicalVertex* Vertex : Vertices)
+	for (const TSharedPtr<FTopologicalVertex>& Vertex : Vertices)
 	{
 		for (const FTopologicalVertex* TwinVertex : Vertex->GetTwinEntities())
 		{
@@ -1011,10 +903,19 @@ void FTopomaker::SetSelfConnectedEdgeDegenerated(TArray<FTopologicalVertex*>& Ve
 	}
 }
 
-void FTopomaker::CheckSelfConnectedEdge()
+void FTopomaker::CheckSelfConnectedEdge(double MaxLengthOfDegeneratedEdge, TArray<TSharedPtr<FTopologicalVertex>>& OutBorderVertices)
 {
 	FTimePoint StartTime = FChrono::Now();
 
+	TFunction<void(FTopologicalVertex&)> AddVertexIfBorder = [&](FTopologicalVertex& Vertex)
+	{
+		TSharedPtr<FTopologicalVertex> ActiveVertex = Vertex.GetLinkActiveEntity();
+		if (ActiveVertex->IsBorderVertex())
+		{
+			OutBorderVertices.Add(ActiveVertex);
+		}
+	};
+
 	FMessage::Printf(Log, TEXT("    Self connected edges\n"));
 	for (TSharedPtr<FTopologicalFace> Face : Faces)
 	{
@@ -1025,10 +926,18 @@ void FTopomaker::CheckSelfConnectedEdge()
 				TSharedPtr<FTopologicalEdge> Edge = OrientedEdge.Entity;
 				if (Edge->GetStartVertex()->IsLinkedTo(Edge->GetEndVertex()))
 				{
-					if (!Edge->IsDegenerated() && Edge->Length() < 2 * SewTolerance)
+					if (!Edge->IsDegenerated() && Edge->Length() < MaxLengthOfDegeneratedEdge)
 					{
-						FMessage::Printf(Debug, TEXT("Face %d Edge %d was self connected\n"), Face->GetId(), Edge->GetId());
-						Edge->GetStartVertex()->UnlinkTo(*Edge->GetEndVertex());
+						if (Edge->GetTwinEntityCount() > 1)
+						{
+							FMessage::Printf(Debug, TEXT("Face %d Edge %d was self connected\n"), Face->GetId(), Edge->GetId());
+							Edge->GetStartVertex()->UnlinkTo(*Edge->GetEndVertex());
+						}
+						else
+						{
+							Edge->SetAsDegenerated();
+							FMessage::Printf(Debug, TEXT("Face %d Edge %d is set as degenerated\n"), Face->GetId(), Edge->GetId());
+						}
 					}
 				}
 			}
diff --git a/Engine/Source/Runtime/Datasmith/CADKernel/Private/CADKernel/UI/Display.cpp b/Engine/Source/Runtime/Datasmith/CADKernel/Private/CADKernel/UI/Display.cpp
--- a/Engine/Source/Runtime/Datasmith/CADKernel/Private/CADKernel/UI/Display.cpp
+++ b/Engine/Source/Runtime/Datasmith/CADKernel/Private/CADKernel/UI/Display.cpp
@@ -255,8 +255,12 @@ void DrawQuadripode(double Height, double Base, FPoint& Center, FPoint& InDirect
 void DisplayEntity(const FEntity& Entity)
 {
 #ifdef CADKERNEL_DEV
-	FTimePoint StartTime = FChrono::Now();
+	if (Entity.IsDeleted())
+	{
+		return;
+	}
 
+	FTimePoint StartTime = FChrono::Now();
 	FProgress Progress;
 
 	F3DDebugSession GraphicSession(FString::Printf(TEXT("%s %d"), Entity.GetTypeName(), Entity.GetId()), { Entity.GetId() });
@@ -309,6 +313,11 @@ void DisplayEntity(const FEntity& Entity)
 void DisplayEntity2D(const FEntity& Entity)
 {
 #ifdef CADKERNEL_DEV
+	if (Entity.IsDeleted())
+	{
+		return;
+	}
+
 	FTimePoint StartTime = FChrono::Now();
 
 	FProgress Progress;
@@ -736,7 +745,7 @@ void Draw(const FTopologicalFace& Face)
 				switch (Edge.Entity->GetTwinEntityCount())
 				{
 				case 1:
-					Property =  Edge.Entity->IsDegenerated() ? EVisuProperty::PurpleCurve : EVisuProperty::BorderEdge;
+					Property =  Edge.Entity->IsDegenerated() ? EVisuProperty::OrangeCurve : EVisuProperty::BorderEdge;
 					break;
 				case 2:
 					Property = EVisuProperty::BlueCurve;
diff --git a/Engine/Source/Runtime/Datasmith/CADKernel/Public/CADKernel/Topo/Linkable.h b/Engine/Source/Runtime/Datasmith/CADKernel/Public/CADKernel/Topo/Linkable.h
--- a/Engine/Source/Runtime/Datasmith/CADKernel/Public/CADKernel/Topo/Linkable.h
+++ b/Engine/Source/Runtime/Datasmith/CADKernel/Public/CADKernel/Topo/Linkable.h
@@ -81,7 +81,7 @@ public:
 		TopologicalLink = FEntity::MakeShared<LinkType>((EntityType&)(*this));
 	}
 
-	bool IsLinkedTo(TSharedRef<EntityType> Entity) const
+	bool IsLinkedTo(const TSharedRef<EntityType>& Entity) const
 	{
 		if (this == &*Entity)
 		{
@@ -90,6 +90,15 @@ public:
 		return (Entity->TopologicalLink == TopologicalLink);
 	}
 
+	bool IsLinkedTo(const EntityType& Entity) const
+	{
+		if (this == &Entity)
+		{
+			return true;
+		}
+		return (Entity.TopologicalLink == TopologicalLink);
+	}
+
 	int32 GetTwinEntityCount() const
 	{
 		ensureCADKernel(TopologicalLink.IsValid());
diff --git a/Engine/Source/Runtime/Datasmith/CADKernel/Public/CADKernel/Topo/TopologicalEdge.h b/Engine/Source/Runtime/Datasmith/CADKernel/Public/CADKernel/Topo/TopologicalEdge.h
--- a/Engine/Source/Runtime/Datasmith/CADKernel/Public/CADKernel/Topo/TopologicalEdge.h
+++ b/Engine/Source/Runtime/Datasmith/CADKernel/Public/CADKernel/Topo/TopologicalEdge.h
@@ -168,7 +168,15 @@ public:
 	virtual FInfoEntity& GetInfo(FInfoEntity&) const override;
 #endif
 
-	double GetTolerance3D();
+	double GetTolerance3D() const
+	{
+		return GetCurve()->GetCarrierSurface()->Get3DTolerance();
+	}
+
+	double GetTolerance2DAt(double Coordinate) const
+	{
+		return GetCurve()->GetToleranceAt(Coordinate);
+	}
 
 	virtual EEntity GetEntityType() const override
 	{
@@ -183,13 +191,32 @@ public:
 	 * Checks if the carrier curve is degenerated i.e. the 2d length of the curve is nearly zero
 	 * If the 3d length is nearly zero, the edge is flag as degenerated
 	 */
-	bool CheckIfDegenerated();
+	bool CheckIfDegenerated() const;
+
+	/**
+	 * It can be linked to the edge if :
+	 *  - they are connected at their extremities,
+	 *  - they have the same length (5% or +/- EdgeLengthTolerance)
+	 *  - they are ~tangent at their extremities i.e @see IsTangentAtExtremitiesWith
+	 */
+	bool IsLinkableTo(const FTopologicalEdge& Edge, double EdgeLengthTolerance) const;
+
+	/**
+	 * Link two edges.
+	 * Two edges can be linked if :
+	 *  - they are connected at their extremities (SquareJoiningTolerance),
+	 *  - they are linkable (@see IsLinkableTo)
+	 * 
+	 * This step must be done when the loop is finalize because in some case edges can be delete, split, extend... to avoid problems
+	 */
+	void LinkIfCoincident(FTopologicalEdge& OtherEdge, double EdgeLengthTolerance, double SquareJoiningTolerance);
 
 	/**
-	 * Build face links with its neigbour (link edges) should be done after the loop is finalize.
-	 * In some case, edges can be delete, split, extend... Link edges when the loop is well done and clean avoid problems
+	 * Link with the other edge.
+	 * No check is performed except check if degenerated ot deleted.
+	 * If checks are needed, use LinkIfCoincident
 	 */
-	void Link(FTopologicalEdge& OtherEdge, double SquareJoiningTolerance);
+	void Link(FTopologicalEdge& OtherEdge);
 
 	TSharedRef<const FTopologicalEdge> GetLinkActiveEdge() const
 	{
@@ -222,6 +249,11 @@ public:
 	 */
 	bool IsSameDirection(const FTopologicalEdge& Edge) const;
 
+	/**
+	 * @return true if it is ~tangent with the edge at their extremities i.e Cos(tangents) > cos(30 deg)
+	 */
+	bool IsTangentAtExtremitiesWith(const FTopologicalEdge & Edge) const;
+
 	/**
 	 * @return true if the edge is self connected at its extremities
 	 */
@@ -647,8 +679,10 @@ public:
 	 */
 	bool ExtendTo(bool bStartExtremity, const FPoint2D& NewExtremityCoordinate, TSharedRef<FTopologicalVertex>& NewVertex);
 
-	bool IsSharpEdge();
+	bool IsSharpEdge() const;
 
+	/** @return true if they have the same length +/- 5 % or +/- EdgeLengthTolerance */
+	bool HasSameLengthAs(const FTopologicalEdge& Edge, double EdgeLengthTolerance) const;
 
 	// ======   State Functions   ======
 
diff --git a/Engine/Source/Runtime/Datasmith/CADKernel/Public/CADKernel/Topo/TopologicalVertex.h b/Engine/Source/Runtime/Datasmith/CADKernel/Public/CADKernel/Topo/TopologicalVertex.h
--- a/Engine/Source/Runtime/Datasmith/CADKernel/Public/CADKernel/Topo/TopologicalVertex.h
+++ b/Engine/Source/Runtime/Datasmith/CADKernel/Public/CADKernel/Topo/TopologicalVertex.h
@@ -120,6 +120,11 @@ public:
 		return Coordinates.SquareDistance(OtherVertex.Coordinates);
 	}
 
+	double SquareDistanceBetweenBarycenters(const FTopologicalVertex& OtherVertex) const
+	{
+		return GetLink()->GetBarycenter().SquareDistance(OtherVertex.GetLink()->GetBarycenter());
+	}
+
 	double SquareDistance(const FPoint& Point) const
 	{
 		return Coordinates.SquareDistance(Point);
diff --git a/Engine/Source/Runtime/Datasmith/CADKernel/Public/CADKernel/Topo/Topomaker.h b/Engine/Source/Runtime/Datasmith/CADKernel/Public/CADKernel/Topo/Topomaker.h
--- a/Engine/Source/Runtime/Datasmith/CADKernel/Public/CADKernel/Topo/Topomaker.h
+++ b/Engine/Source/Runtime/Datasmith/CADKernel/Public/CADKernel/Topo/Topomaker.h
@@ -27,26 +27,38 @@ protected:
 	TArray<FShell*> Shells;
 	TArray<TSharedPtr<FTopologicalFace>> Faces;
 
-	double Tolerance;
-
 #ifdef CADKERNEL_DEV
 	FTopomakerReport Report;
 #endif
 
 public:
 
-	FTopomaker(FSession& InSession, double Tolerance);
-	FTopomaker(FSession& InSession, const TArray<TSharedPtr<FShell>>& Shells, double Tolerance);
-	FTopomaker(FSession& InSession, const TArray<TSharedPtr<FTopologicalFace>>& Surfaces, double Tolerance);
+	FTopomaker(FSession& InSession, double InTolerance, double InForceFactor);
+	FTopomaker(FSession& InSession, const TArray<TSharedPtr<FShell>>& Shells, double InTolerance, double InForceFactor);
+	FTopomaker(FSession& InSession, const TArray<TSharedPtr<FTopologicalFace>>& Surfaces, double InTolerance, double InForceFactor);
 
-	void SetTolerance(double NewValue)
+	void SetTolerance(double InTolerance, double InForceFactor)
 	{
-		Tolerance = NewValue;
-		SewTolerance = NewValue * UE_DOUBLE_SQRT_2;
-		SewToleranceSquare = FMath::Square(SewTolerance);
+		Tolerance = InTolerance;
+		ForceJoinFactor = InForceFactor;
+
+		SewTolerance = Tolerance * UE_DOUBLE_SQRT_2;
+
+		// StartVertex (SV) is linked to TwinStartVertex (TSV), the max distance between vertices is SewTolerance
+		// EndVertex (EV) is linked to TwinEndVertex (TEV), the max distance between vertices is SewTolerance
+		//
+		//           * TSV * TEV
+		//     SV *-----------* EV
+		// So the max length is 2*SewTolerance
+		EdgeLengthTolerance = SewTolerance * 2.;
+
+		SewToleranceToForceJoin = Tolerance * ForceJoinFactor;
+		LargeEdgeLengthTolerance = SewToleranceToForceJoin * 2.;
+
+		ThinFaceWidth = Tolerance * ForceJoinFactor;
 	}
 
-	void Sew();
+	void Sew(bool bForceJoining, bool bRemoveThinFaces);
 
 	/**
 	 * Check topology of each body
@@ -65,7 +77,7 @@ public:
 	 */
 	void UnlinkNonManifoldVertex();
 
-	void RemoveThinFaces();
+	void RemoveThinFaces(TArray<FTopologicalEdge*>& NewBorderEdges);
 
 #ifdef CADKERNEL_DEV
 	void PrintSewReport()
@@ -96,67 +108,47 @@ private:
 	/**
 	 * Return an array of active vertices.
 	 */
-	void GetVertices(TArray<FTopologicalVertex*>& Vertices);
+	void GetVertices(TArray<TSharedPtr<FTopologicalVertex>>& Vertices);
 
 	/**
 	 * Return an array of active border vertices.
 	 */
-	void GetBorderVertices(TArray<FTopologicalVertex*>& BorderVertices);
-
-	/**
-	 * Merge Border Vertices with other vertices.
-	 * @param Vertices: the initial array of active vertices, this array is updated at the end of the process
-	 */
-	void MergeCoincidentVertices(TArray<FTopologicalVertex*>& VerticesToMerge);
-
-	/**
-	 * Merge Border Vertices with other vertices.
-	 * @param Vertices: the initial array of active vertices, this array is updated at the end of the process
-	 */
-	void MergeBorderVerticesWithCoincidentOtherVertices(TArray<FTopologicalVertex*>& Vertices);
+	void GetBorderVertices(TArray<TSharedPtr<FTopologicalVertex>>& BorderVertices);
 
-	//void FixCollapsedEdges();
-	void CheckSelfConnectedEdge();
+	void CheckSelfConnectedEdge(double MaxLengthOfDegeneratedEdge, TArray<TSharedPtr<FTopologicalVertex>>& OutBorderVertices);
 	void RemoveIsolatedEdges(); // Useful ???
 
-	void SetSelfConnectedEdgeDegenerated(TArray<FTopologicalVertex*>& Vertices);
+	void SetSelfConnectedEdgeDegenerated(TArray<TSharedPtr<FTopologicalVertex>>& Vertices);
 
 	/**
-	 * First step, trivial edge merge i.e. couple of edges with same extremity vertex
-	 */
-	void MergeCoincidentEdges(TArray<FTopologicalVertex*>& Vertices);
-	void MergeCoincidentEdges(TArray<FTopologicalEdge*>& EdgesToProcess);
-	void MergeCoincidentEdges(FTopologicalEdge* Edge);
-
-	/**
-	 * Second step, parallel edges but with different length. the longest must be split
-	 */
-	void StitchParallelEdges(TArray<FTopologicalVertex*>& Vertices);
-
-	/**
-	 * Split the edge to split except if one side is too small (JoiningTolerance). In this case (too small), the EdgeToSplit is not split but linked to EdgeToLink.
-	 * @return the created vertex or TSharedPtr<FTopologicalVertex>()
-	 */
-	TSharedPtr<FTopologicalVertex> SplitAndLink(FTopologicalVertex& StartVertex, FTopologicalEdge& EdgeToLink, FTopologicalEdge& EdgeToSplit);
-
-	/**
-	 * For each loop of each surface, check if successive edges are unconnected and if their common vertices are connected only to them.
+	 * For each loop of each surface, check if successive edges are unconnected and if their common vertices are connected only to them (UV vs CV: vertex connected to many faces ).
+	 * These edges must be tangent together.
 	 * These edges are merged into one edge.
 	 * E.g. :
 	 * Face A has 3 successive unconnected edges. If these 3 edges are merged to give only one edge, the new edge could be linked to its parallel edge of Face B
+	 * 
+	 *  CV: Connected vertex i.e. linked to different faces
+	 * 
 	 *
 	 *              \                          Face A                                   |    Face C
 	 *               \                                                                  |
-	 *    Face E     CV ------------------- UV ------------------ UV ----------------- CV --------------------
+	 *    Face E     CV ------------------ UV -------------------- UV ---------------- CV --------------------
 	 *               CV -------------------------------------------------------------- CV --------------------
 	 *               /                         Face B                                   |    Face D
 	 *              /                                                                   |
 	 */
-	void MergeUnconnectedAdjacentEdges();
+	void MergeUnconnectedSuccessiveEdges();
+
+	double Tolerance;
+	double ForceJoinFactor;
 
 	double SewTolerance;
-	double SewToleranceSquare;
+	double SewToleranceToForceJoin; // SewTolerance x ForceJoinFactor;
+
+	double EdgeLengthTolerance; // 2x SewTolerance
+	double LargeEdgeLengthTolerance; // 2x SewToleranceToForceJoin
 
+	double ThinFaceWidth;
 };
 
 } // namespace UE::CADKernel
