Deleting entities with sub components in revit leaves "ghost" entities fix,
Revit - Deleting array elements will not be synchronized via direct link fix,
Revit - Deleting multiple objects isn't auto synced on Unreal side

diff --git a/Engine/Source/Runtime/Datasmith/DatasmithCore/Private/DatasmithSceneElementsImpl.cpp b/Engine/Source/Runtime/Datasmith/DatasmithCore/Private/DatasmithSceneElementsImpl.cpp
--- a/Engine/Source/Runtime/Datasmith/DatasmithCore/Private/DatasmithSceneElementsImpl.cpp
+++ b/Engine/Source/Runtime/Datasmith/DatasmithCore/Private/DatasmithSceneElementsImpl.cpp
@@ -1317,8 +1317,11 @@ namespace DatasmithSceneImplInternal
 					for (int32 ChildIndex = InActor->GetChildrenCount() - 1; ChildIndex >= 0; --ChildIndex)
 					{
 						const TSharedPtr< IDatasmithActorElement > Child = InActor->GetChild(ChildIndex);
-						InActor->RemoveChild(Child);
-						SceneImpl->AddActor(Child);
+						if (!Child->IsAComponent())
+						{
+							InActor->RemoveChild(Child);
+							SceneImpl->AddActor(Child);
+						}
 					}
 				}
 				else
@@ -1336,8 +1339,11 @@ namespace DatasmithSceneImplInternal
 					for (int32 ChildIndex = InActor->GetChildrenCount() - 1; ChildIndex >= 0; --ChildIndex)
 					{
 						const TSharedPtr< IDatasmithActorElement > Child = InActor->GetChild(ChildIndex);
-						InActor->RemoveChild(Child);
-						FoundHierarchy.Last()->AddChild(Child);
+						if (!Child->IsAComponent())
+						{
+							InActor->RemoveChild(Child);
+							FoundHierarchy.Last()->AddChild(Child);
+						}
 					}
 				}
 				else
