CAD Kernel: Fix crash due to allocation with negative size of TArray

diff --git a/Engine/Source/Runtime/Datasmith/CADKernel/Private/CADKernel/Mesh/Structure/LoopCleaner.cpp b/Engine/Source/Runtime/Datasmith/CADKernel/Private/CADKernel/Mesh/Structure/LoopCleaner.cpp
--- a/Engine/Source/Runtime/Datasmith/CADKernel/Private/CADKernel/Mesh/Structure/LoopCleaner.cpp
+++ b/Engine/Source/Runtime/Datasmith/CADKernel/Private/CADKernel/Mesh/Structure/LoopCleaner.cpp
@@ -34,11 +34,6 @@ FLoopCleaner::FLoopCleaner(FIsoTriangulator& Triangulator)
 bool FLoopCleaner::CleanLoops()
 {
 
-	//{
-	//	FPoint2D Point = LoopNodes[36].Get2DPoint(EGridSpace::UniformScaled, Grid);
-	//	LoopNodes[56].Set2DPoint(EGridSpace::UniformScaled, Grid, Point);
-	//}
-
 	FindBestLoopExtremity();
 
 #ifdef DEBUG_CLEAN_LOOPS		
@@ -1748,10 +1743,24 @@ void FLoopCleaner::SwapSegments(FIsoSegment& IntersectingSegment, FIsoSegment& S
 
 void FLoopCleaner::SwapSubLoopOrientation(int32 FirstSegmentIndex, int32 LastSegmentIndex)
 {
+	LastSegmentIndex = FitSegmentIndex(LastSegmentIndex);
+
 	TArray<FIsoSegment*> Segments;
-	Segments.Reserve(LastSegmentIndex - FirstSegmentIndex);
+	{
+		int32 SegCount = LastSegmentIndex - FirstSegmentIndex;
+		if (SegCount < 0)
+		{
+			SegCount = 0;
+			int32 SegIndex = FitSegmentIndex(FirstSegmentIndex);
+			do
+			{
+				SegCount++;
+				SegIndex = NextSegmentIndex(SegIndex);
+			} while (SegIndex != LastSegmentIndex);
+		}
+		Segments.Reserve(SegCount);
+	}
 
-	LastSegmentIndex = FitSegmentIndex(LastSegmentIndex);
 	int32 Index = FitSegmentIndex(FirstSegmentIndex);
 	do
 	{
