Fix Topomaker

diff --git a/Engine/Source/Runtime/Datasmith/CADKernel/Private/CADKernel/Topo/Topomaker.cpp b/Engine/Source/Runtime/Datasmith/CADKernel/Private/CADKernel/Topo/Topomaker.cpp
--- a/Engine/Source/Runtime/Datasmith/CADKernel/Private/CADKernel/Topo/Topomaker.cpp
+++ b/Engine/Source/Runtime/Datasmith/CADKernel/Private/CADKernel/Topo/Topomaker.cpp
@@ -157,7 +157,7 @@ TSharedPtr<FTopologicalVertex> SplitAndLink(FTopologicalVertex& StartVertex, FTo
 	}
 
 	double SquareDistanceToOtherPoint = ProjectedPoint.SquareDistance(EndVertex->GetBarycenter());
-	if (SquareDistanceToProjectedPoint < SquareMinEdgeLength)
+	if (SquareDistanceToOtherPoint < SquareMinEdgeLength)
 	{
 		// the new point is closed to the extremity, a degenerated edge will be created, so the edges are joined
 		VertexToLink->Link(*EndVertex);
