Fix typo in variable names

diff --git a/Engine/Source/Runtime/Datasmith/CADKernel/Private/CADKernel/Geo/Sampling/SurfacicPolyline.cpp b/Engine/Source/Runtime/Datasmith/CADKernel/Private/CADKernel/Geo/Sampling/SurfacicPolyline.cpp
--- a/Engine/Source/Runtime/Datasmith/CADKernel/Private/CADKernel/Geo/Sampling/SurfacicPolyline.cpp
+++ b/Engine/Source/Runtime/Datasmith/CADKernel/Private/CADKernel/Geo/Sampling/SurfacicPolyline.cpp
@@ -123,17 +123,17 @@ void FSurfacicPolyline::ComputeIntersectionsWithIsos(const FLinearBoundary& InBo
 		}
 
 		// if the segment is too parallel to Iso, go to the next one
-		double EdgeLocalSlop = ComputeUnorientedSlope(Points2D[Index], Points2D[Index + 1], 0);
+		double EdgeLocalSlope = ComputeUnorientedSlope(Points2D[Index], Points2D[Index + 1], 0);
 		if (TypeIso == EIso::IsoV)
 		{
-			if (EdgeLocalSlop < 0.1 || EdgeLocalSlop > 3.9)
+			if (EdgeLocalSlope < 0.1 || EdgeLocalSlope > 3.9)
 			{
 				continue;
 			}
 		}
 		else
 		{
-			if (EdgeLocalSlop < 2.1 && EdgeLocalSlop > 1.9)
+			if (EdgeLocalSlope < 2.1 && EdgeLocalSlope > 1.9)
 			{
 				continue;
 			}
diff --git a/Engine/Source/Runtime/Datasmith/CADKernel/Private/CADKernel/Math/SlopeUtils.h b/Engine/Source/Runtime/Datasmith/CADKernel/Private/CADKernel/Math/SlopeUtils.h
--- a/Engine/Source/Runtime/Datasmith/CADKernel/Private/CADKernel/Math/SlopeUtils.h
+++ b/Engine/Source/Runtime/Datasmith/CADKernel/Private/CADKernel/Math/SlopeUtils.h
@@ -36,7 +36,7 @@ constexpr double MinusRightSlope = -2.;
 constexpr double PiSlope = 4.;
 }
 
-typedef TFunction<double(const FPoint2D&, const FPoint2D&, double)> SlopMethod;
+typedef TFunction<double(const FPoint2D&, const FPoint2D&, double)> SlopeMethod;
 
 
 /**
@@ -180,12 +180,12 @@ inline double ComputePositiveSlope(const FPoint2D& StartPoint, const FPoint2D& E
 	return TransformIntoPositiveSlope(Slope);
 }
 
-inline double ClockwiseSlop(const FPoint2D& StartPoint, const FPoint2D& EndPoint, double ReferenceSlope)
+inline double ClockwiseSlope(const FPoint2D& StartPoint, const FPoint2D& EndPoint, double ReferenceSlope)
 {
 	return 8. - ComputePositiveSlope(StartPoint, EndPoint, ReferenceSlope);
 }
 
-inline double CounterClockwiseSlop(const FPoint2D& StartPoint, const FPoint2D& EndPoint, double ReferenceSlope)
+inline double CounterClockwiseSlope(const FPoint2D& StartPoint, const FPoint2D& EndPoint, double ReferenceSlope)
 {
 	return ComputePositiveSlope(StartPoint, EndPoint, ReferenceSlope);
 }
diff --git a/Engine/Source/Runtime/Datasmith/CADKernel/Private/CADKernel/Mesh/Meshers/BowyerWatsonTriangulator.h b/Engine/Source/Runtime/Datasmith/CADKernel/Private/CADKernel/Mesh/Meshers/BowyerWatsonTriangulator.h
--- a/Engine/Source/Runtime/Datasmith/CADKernel/Private/CADKernel/Mesh/Meshers/BowyerWatsonTriangulator.h
+++ b/Engine/Source/Runtime/Datasmith/CADKernel/Private/CADKernel/Mesh/Meshers/BowyerWatsonTriangulator.h
@@ -233,16 +233,16 @@ public:
 					}
 				}
 
-				TArray<TPair<int32, double>> BoundaryVertexToSlop;
+				TArray<TPair<int32, double>> BoundaryVertexToSlope;
 				for (int32 Index : BoundaryVertex)
 				{
 					FPoint2D Vertex = Vertices[Index].Value;
-					double Slop = ComputeSlope(NewVertex, Vertex);
-					BoundaryVertexToSlop.Emplace(Index, Slop);
+					double Slope = ComputeSlope(NewVertex, Vertex);
+					BoundaryVertexToSlope.Emplace(Index, Slope);
 				}
 
 				// Sort the vertex to make the boundary polygon
-				BoundaryVertexToSlop.Sort([](const TPair<int32, double>& Vertex1, const TPair<int32, double>& Vertex2) {return Vertex1.Value < Vertex2.Value; });
+				BoundaryVertexToSlope.Sort([](const TPair<int32, double>& Vertex1, const TPair<int32, double>& Vertex2) {return Vertex1.Value < Vertex2.Value; });
 
 				bool bTriangleHasBeenAdded = true;
 				while (bTriangleHasBeenAdded)
@@ -282,7 +282,7 @@ public:
 #ifdef DEBUG_BOWYERWATSON_STEP_2
 						if (VIndex == StopIndex)
 						{
-							DisplaySelectedTrianglesBoundary(BoundaryVertexToSlop);
+							DisplaySelectedTrianglesBoundary(BoundaryVertexToSlope);
 							{
 								F3DDebugSession A(TEXT("New Sel triangle"));
 								DisplayTriangle(TriangleIndex, EVisuProperty::YellowCurve);
@@ -295,26 +295,26 @@ public:
 						}
 #endif
 
-						double NewVertexSlop = ComputeSlope(NewVertex, CandidateVertex);
+						double NewVertexSlope = ComputeSlope(NewVertex, CandidateVertex);
 
 						int32 StartIndex = -1;
 						int32 EndIndex = 0;
-						for (int32 ApexIndex = 0; ApexIndex < BoundaryVertexToSlop.Num(); ++ApexIndex)
+						for (int32 ApexIndex = 0; ApexIndex < BoundaryVertexToSlope.Num(); ++ApexIndex)
 						{
-							if (NewVertexSlop < BoundaryVertexToSlop[ApexIndex].Value)
+							if (NewVertexSlope < BoundaryVertexToSlope[ApexIndex].Value)
 							{
 								EndIndex = ApexIndex;
 								break;
 							}
 						}
 
-						StartIndex = EndIndex == 0 ? BoundaryVertexToSlop.Num() - 1 : EndIndex - 1;
+						StartIndex = EndIndex == 0 ? BoundaryVertexToSlope.Num() - 1 : EndIndex - 1;
 
-						FPoint2D StartPoint = Vertices[BoundaryVertexToSlop[StartIndex].Key].Value;
+						FPoint2D StartPoint = Vertices[BoundaryVertexToSlope[StartIndex].Key].Value;
 						double SlopeNewStartPlusPi = ComputeSlope(NewVertex, StartPoint) + Slope::RightSlope;
 						double SlopeStartCandidate = ComputePositiveSlope(StartPoint, CandidateVertex, SlopeNewStartPlusPi);
 
-						FPoint2D EndPoint = Vertices[BoundaryVertexToSlop[EndIndex].Key].Value;
+						FPoint2D EndPoint = Vertices[BoundaryVertexToSlope[EndIndex].Key].Value;
 						double SlopeNewEndMinusPi = ComputeSlope(NewVertex, EndPoint) - Slope::RightSlope;
 						double SlopeEndCandidate = ComputePositiveSlope(EndPoint, CandidateVertex, SlopeNewEndMinusPi);
 
@@ -329,13 +329,13 @@ public:
 						if (SlopeStartCandidate > 0 && SlopeStartCandidate < Slope::RightSlope && SlopeEndCandidate > Slope::ThreeRightSlope && SlopeEndCandidate < Slope::FullSlope)
 						{
 							bTriangleHasBeenAdded = true;
-							if (EndIndex == 0 && NewVertexSlop > BoundaryVertexToSlop.Last().Value)
+							if (EndIndex == 0 && NewVertexSlope > BoundaryVertexToSlope.Last().Value)
 							{
-								BoundaryVertexToSlop.Emplace(CandidateVertexIndex, NewVertexSlop);
+								BoundaryVertexToSlope.Emplace(CandidateVertexIndex, NewVertexSlope);
 							}
 							else
 							{
-								BoundaryVertexToSlop.EmplaceAt(EndIndex, CandidateVertexIndex, NewVertexSlop);
+								BoundaryVertexToSlope.EmplaceAt(EndIndex, CandidateVertexIndex, NewVertexSlope);
 							}
 							BoundaryVertex.Add(CandidateVertexIndex);
 							SelectedTriangleIndices.Add(TriangleIndex);
diff --git a/Engine/Source/Runtime/Datasmith/CADKernel/Private/CADKernel/Mesh/Meshers/IsoTriangulator.cpp b/Engine/Source/Runtime/Datasmith/CADKernel/Private/CADKernel/Mesh/Meshers/IsoTriangulator.cpp
--- a/Engine/Source/Runtime/Datasmith/CADKernel/Private/CADKernel/Mesh/Meshers/IsoTriangulator.cpp
+++ b/Engine/Source/Runtime/Datasmith/CADKernel/Private/CADKernel/Mesh/Meshers/IsoTriangulator.cpp
@@ -95,6 +95,7 @@ bool FIsoTriangulator::Triangulate()
 #ifdef DEBUG_ISOTRIANGULATOR
 	DisplayLoops(TEXT("FIsoTrianguler::LoopSegments after fix intersection"), true, false);
 	Grid.DisplayIsoSegments(TEXT("FIsoTrianguler::LoopSegments orientation"), DisplaySpace, LoopSegments, false, true, EVisuProperty::YellowCurve);
+	Grid.DisplayIsoSegments(TEXT("FIsoTrianguler::LoopSegments"), DisplaySpace, LoopSegments, true, false, EVisuProperty::YellowCurve);
 	Wait(false);
 #endif
 
@@ -118,7 +119,7 @@ bool FIsoTriangulator::Triangulate()
 		F3DDebugSession A(bDisplay, TEXT("FIsoTrianguler::FinalInnerSegments"));
 		Grid.DisplayIsoSegments(DisplaySpace, FinalInnerSegments, true, false, EVisuProperty::BlueCurve);
 	}
-	InnerToOuterSegmentsIntersectionTool.Display(TEXT("FIsoTrianguler::IntersectionTool InnerToOutter"));
+	InnerToOuterSegmentsIntersectionTool.Display(TEXT("FIsoTrianguler::IntersectionTool InnerToOuter"), EVisuProperty::RedCurve);
 	Chronos.TriangulateDuration1 = FChrono::Elapse(StartTime);
 #endif
 
@@ -324,24 +325,24 @@ void FIsoTriangulator::BuildLoopSegments()
 
 	for (FIsoSegment* Segment : LoopSegments)
 	{
-		double SegmentSlop = ComputeSlope(Segment->GetFirstNode().Get2DPoint(EGridSpace::UniformScaled, Grid), Segment->GetSecondNode().Get2DPoint(EGridSpace::UniformScaled, Grid));
-		if (SegmentSlop < IsoTriangulatorImpl::MaxSlopeToBeIso)
+		double SegmentSlope = ComputeSlope(Segment->GetFirstNode().Get2DPoint(EGridSpace::UniformScaled, Grid), Segment->GetSecondNode().Get2DPoint(EGridSpace::UniformScaled, Grid));
+		if (SegmentSlope < IsoTriangulatorImpl::MaxSlopeToBeIso)
 		{
 			Segment->SetAsIsoU();
 		}
-		if (SegmentSlop < IsoTriangulatorImpl::LimitValueMax(2.) && SegmentSlop > IsoTriangulatorImpl::LimitValueMin(2.))
+		if (SegmentSlope < IsoTriangulatorImpl::LimitValueMax(2.) && SegmentSlope > IsoTriangulatorImpl::LimitValueMin(2.))
 		{
 			Segment->SetAsIsoV();
 		}
-		if (SegmentSlop < IsoTriangulatorImpl::LimitValueMax(4.) && SegmentSlop > IsoTriangulatorImpl::LimitValueMin(4.))
+		if (SegmentSlope < IsoTriangulatorImpl::LimitValueMax(4.) && SegmentSlope > IsoTriangulatorImpl::LimitValueMin(4.))
 		{
 			Segment->SetAsIsoU();
 		}
-		if (SegmentSlop < IsoTriangulatorImpl::LimitValueMax(6.) && SegmentSlop > IsoTriangulatorImpl::LimitValueMin(6.))
+		if (SegmentSlope < IsoTriangulatorImpl::LimitValueMax(6.) && SegmentSlope > IsoTriangulatorImpl::LimitValueMin(6.))
 		{
 			Segment->SetAsIsoV();
 		}
-		if (SegmentSlop > IsoTriangulatorImpl::LimitValueMin(8.))
+		if (SegmentSlope > IsoTriangulatorImpl::LimitValueMin(8.))
 		{
 			Segment->SetAsIsoU();
 		}
@@ -986,8 +987,8 @@ void FIsoTriangulator::FindSegmentToLinkOuterToInnerLoopNodes(FCell& Cell)
 #endif 
 	};
 
-	const SlopMethod GetSlopAtStartNode = CounterClockwiseSlop;
-	const SlopMethod GetSlopAtEndNode = ClockwiseSlop;
+	const SlopeMethod GetSlopeAtStartNode = CounterClockwiseSlope;
+	const SlopeMethod GetSlopeAtEndNode = ClockwiseSlope;
 	TFunction<void(FIsoSegment*)> FindBestTriangle = [&](FIsoSegment* Segment)
 	{
 		FLoopNode& StartNode = (FLoopNode&)Segment->GetFirstNode();
@@ -1027,16 +1028,16 @@ void FIsoTriangulator::FindSegmentToLinkOuterToInnerLoopNodes(FCell& Cell)
 		// StartMaxSlope and EndMaxSlope are at most equal to 4, because if the slop with candidate node is biggest to 4, the nez triangle will be inverted
 		double StartReferenceSlope = ComputePositiveSlope(StartPoint2D, EndPoint2D, 0);
 
-		double StartMaxSlope = GetSlopAtStartNode(StartPoint2D, PreviousNode.Get2DPoint(EGridSpace::UniformScaled, Grid), StartReferenceSlope);
+		double StartMaxSlope = GetSlopeAtStartNode(StartPoint2D, PreviousNode.Get2DPoint(EGridSpace::UniformScaled, Grid), StartReferenceSlope);
 		StartMaxSlope /= 2.;
 		StartMaxSlope = FMath::Min(StartMaxSlope, 4.);
 
 		double EndReferenceSlope = StartReferenceSlope < 4 ? StartReferenceSlope + 4 : StartReferenceSlope - 4;
-		double EndMaxSlope = GetSlopAtEndNode(EndPoint2D, NextNode.Get2DPoint(EGridSpace::UniformScaled, Grid), EndReferenceSlope);
+		double EndMaxSlope = GetSlopeAtEndNode(EndPoint2D, NextNode.Get2DPoint(EGridSpace::UniformScaled, Grid), EndReferenceSlope);
 		EndMaxSlope /= 2.;
 		EndMaxSlope = FMath::Min(EndMaxSlope, 4.);
 
-		const double MinSlopToNotBeAligned = 0.01;
+		const double MinSlopeToNotBeAligned = 0.01;
 
 		FLoopNode* CandidatNode = nullptr;
 		double MinCriteria = HUGE_VALUE;
@@ -1049,8 +1050,8 @@ void FIsoTriangulator::FindSegmentToLinkOuterToInnerLoopNodes(FCell& Cell)
 			{
 				// Check if the node is inside the sector (X) or outside (Z)
 				const FPoint2D& NodePoint2D = Node->Get2DPoint(EGridSpace::UniformScaled, Grid);
-				double SlopeAtStartNode = GetSlopAtStartNode(StartPoint2D, NodePoint2D, StartReferenceSlope);
-				double SlopeAtEndNode = GetSlopAtEndNode(EndPoint2D, NodePoint2D, EndReferenceSlope);
+				double SlopeAtStartNode = GetSlopeAtStartNode(StartPoint2D, NodePoint2D, StartReferenceSlope);
+				double SlopeAtEndNode = GetSlopeAtEndNode(EndPoint2D, NodePoint2D, EndReferenceSlope);
 
 				if (SlopeAtStartNode <= 0 || SlopeAtStartNode >= StartMaxSlope)
 				{
@@ -1066,7 +1067,7 @@ void FIsoTriangulator::FindSegmentToLinkOuterToInnerLoopNodes(FCell& Cell)
 
 				if (
 					// the candidate triangle is inside the current candidate triangle
-					((SlopeAtStartNode < (CandidateSlopeAtStartNode + MinSlopToNotBeAligned)) && (SlopeAtEndNode < (CandidateSlopeAtEndNode + MinSlopToNotBeAligned))) ||
+					((SlopeAtStartNode < (CandidateSlopeAtStartNode + MinSlopeToNotBeAligned)) && (SlopeAtEndNode < (CandidateSlopeAtEndNode + MinSlopeToNotBeAligned))) ||
 					// or the candidate triangle is better the current candidate triangle and doesn't contain the current candidate triangle
 					((PointCriteria < MinCriteria) && ((SlopeAtStartNode > CandidateSlopeAtStartNode) ^ (SlopeAtEndNode > CandidateSlopeAtEndNode))))
 				{
@@ -1919,8 +1920,8 @@ void FIsoTriangulator::MeshCycle(const EGridSpace Space, const TArray<FIsoSegmen
 		return true;
 	};
 
-	const SlopMethod GetSlopAtStartNode = ClockwiseSlop;
-	const SlopMethod GetSlopAtEndNode = CounterClockwiseSlop;
+	const SlopeMethod GetSlopeAtStartNode = ClockwiseSlope;
+	const SlopeMethod GetSlopeAtEndNode = CounterClockwiseSlope;
 
 	TFunction<void(FIsoSegment*, bool)> FindBestTriangle = [&](FIsoSegment* Segment, bool bOrientation)
 	{
@@ -1959,9 +1960,9 @@ void FIsoTriangulator::MeshCycle(const EGridSpace Space, const TArray<FIsoSegmen
 		// TODO ? Need to add test to check if the candidate node will not generate flat triangle 
 
 		// PreviousSegment = [A, X0]
-		FIsoSegment* PreviousSegment = FindNextSegment(EGridSpace::UniformScaled, Segment, &StartNode, GetSlopAtStartNode);
+		FIsoSegment* PreviousSegment = FindNextSegment(EGridSpace::UniformScaled, Segment, &StartNode, GetSlopeAtStartNode);
 		// NextSegment = [B, Xn]
-		FIsoSegment* NextSegment = FindNextSegment(EGridSpace::UniformScaled, Segment, &EndNode, GetSlopAtEndNode);
+		FIsoSegment* NextSegment = FindNextSegment(EGridSpace::UniformScaled, Segment, &EndNode, GetSlopeAtEndNode);
 
 		// PreviousNode = X0
 		FIsoNode& PreviousNode = PreviousSegment->GetOtherNode(&StartNode);
@@ -2005,7 +2006,7 @@ void FIsoTriangulator::MeshCycle(const EGridSpace Space, const TArray<FIsoSegmen
 
 			// StartMaxSlope and EndMaxSlope are at most equal to 4, because if the slop with candidate node is biggest to 4, the nez triangle will be inverted
 			double StartReferenceSlope = ComputePositiveSlope(StartPoint2D, EndPoint2D, 0);
-			double StartMaxSlope = GetSlopAtStartNode(StartPoint2D, PreviousNode.Get2DPoint(EGridSpace::UniformScaled, Grid), StartReferenceSlope);
+			double StartMaxSlope = GetSlopeAtStartNode(StartPoint2D, PreviousNode.Get2DPoint(EGridSpace::UniformScaled, Grid), StartReferenceSlope);
 			if (&EndNode != &PreviousNode)
 			{
 				// Case of probable auto-intersection cycle at PreviousNode, cancel is preferred 
@@ -2017,7 +2018,7 @@ void FIsoTriangulator::MeshCycle(const EGridSpace Space, const TArray<FIsoSegmen
 			StartMaxSlope = FMath::Min(StartMaxSlope, 4.);
 
 			double EndReferenceSlope = StartReferenceSlope < 4 ? StartReferenceSlope + 4 : StartReferenceSlope - 4;
-			double EndMaxSlope = GetSlopAtEndNode(EndPoint2D, NextNode.Get2DPoint(EGridSpace::UniformScaled, Grid), EndReferenceSlope);
+			double EndMaxSlope = GetSlopeAtEndNode(EndPoint2D, NextNode.Get2DPoint(EGridSpace::UniformScaled, Grid), EndReferenceSlope);
 			if (&StartNode != &NextNode)
 			{
 				// Case of probable auto-intersection cycle at PreviousNode, cancel is preferred 
@@ -2029,8 +2030,8 @@ void FIsoTriangulator::MeshCycle(const EGridSpace Space, const TArray<FIsoSegmen
 			EndMaxSlope = FMath::Min(EndMaxSlope, 4.);
 
 			double MinCriteria = HUGE_VALUE;
-			const double MinSlopToNotBeAligned = 0.01;
-			const double EpsilonSlop = 0.0001;
+			const double MinSlopeToNotBeAligned = 0.01;
+			const double EpsilonSlope = 0.0001;
 			double CandidateSlopeAtStartNode = 8.;
 			double CandidateSlopeAtEndNode = 8.;
 
@@ -2047,12 +2048,12 @@ void FIsoTriangulator::MeshCycle(const EGridSpace Space, const TArray<FIsoSegmen
 
 				// Check if the node is inside the sector (X) or outside (Z)
 				const FPoint2D& NodePoint2D = Node->Get2DPoint(EGridSpace::UniformScaled, Grid);
-				double SlopeAtStartNode = GetSlopAtStartNode(StartPoint2D, NodePoint2D, StartReferenceSlope);
-				double SlopeAtEndNode = GetSlopAtEndNode(EndPoint2D, NodePoint2D, EndReferenceSlope);
+				double SlopeAtStartNode = GetSlopeAtStartNode(StartPoint2D, NodePoint2D, StartReferenceSlope);
+				double SlopeAtEndNode = GetSlopeAtEndNode(EndPoint2D, NodePoint2D, EndReferenceSlope);
 
 				if (Node != &PreviousNode)
 				{
-					if (SlopeAtStartNode <= EpsilonSlop || SlopeAtStartNode >= StartMaxSlope - EpsilonSlop)
+					if (SlopeAtStartNode <= EpsilonSlope || SlopeAtStartNode >= StartMaxSlope - EpsilonSlope)
 					{
 						continue;
 					}
@@ -2060,7 +2061,7 @@ void FIsoTriangulator::MeshCycle(const EGridSpace Space, const TArray<FIsoSegmen
 
 				if (Node != &NextNode)
 				{
-					if (SlopeAtEndNode <= EpsilonSlop || SlopeAtEndNode >= EndMaxSlope - EpsilonSlop)
+					if (SlopeAtEndNode <= EpsilonSlope || SlopeAtEndNode >= EndMaxSlope - EpsilonSlope)
 					{
 						continue;
 					}
@@ -2076,12 +2077,12 @@ void FIsoTriangulator::MeshCycle(const EGridSpace Space, const TArray<FIsoSegmen
 				}
 #endif
 
-				if (FMath::IsNearlyEqual(SlopeAtStartNode, CandidateSlopeAtStartNode, MinSlopToNotBeAligned) && SlopeAtEndNode > CandidateSlopeAtEndNode)
+				if (FMath::IsNearlyEqual(SlopeAtStartNode, CandidateSlopeAtStartNode, MinSlopeToNotBeAligned) && SlopeAtEndNode > CandidateSlopeAtEndNode)
 				{
 					continue;
 				}
 
-				if (FMath::IsNearlyEqual(SlopeAtEndNode, CandidateSlopeAtEndNode, EpsilonSlop) && SlopeAtStartNode > CandidateSlopeAtStartNode)
+				if (FMath::IsNearlyEqual(SlopeAtEndNode, CandidateSlopeAtEndNode, EpsilonSlope) && SlopeAtStartNode > CandidateSlopeAtStartNode)
 				{
 					continue;
 				}
@@ -2107,7 +2108,7 @@ void FIsoTriangulator::MeshCycle(const EGridSpace Space, const TArray<FIsoSegmen
 #endif
 				if (
 					// the candidate triangle is inside the current candidate triangle
-					((SlopeAtStartNode < (CandidateSlopeAtStartNode + MinSlopToNotBeAligned)) && (SlopeAtEndNode < (CandidateSlopeAtEndNode + MinSlopToNotBeAligned)))
+					((SlopeAtStartNode < (CandidateSlopeAtStartNode + MinSlopeToNotBeAligned)) && (SlopeAtEndNode < (CandidateSlopeAtEndNode + MinSlopeToNotBeAligned)))
 					||
 					// the candidate triangle is better the current candidate triangle and doesn't contain the current candidate triangle
 					((PointCriteria < MinCriteria) && ((SlopeAtStartNode > CandidateSlopeAtStartNode) ^ (SlopeAtEndNode > CandidateSlopeAtEndNode))))
@@ -2511,7 +2512,7 @@ bool FIsoTriangulator::FindCycle(FIsoSegment* StartSegment, bool LeftSide, TArra
 
 	for (;;)
 	{
-		Segment = FindNextSegment(EGridSpace::UniformScaled, Segment, Node, ClockwiseSlop);
+		Segment = FindNextSegment(EGridSpace::UniformScaled, Segment, Node, ClockwiseSlope);
 		if (Segment == nullptr)
 		{
 			Cycle.Empty();
@@ -2598,7 +2599,7 @@ bool FIsoTriangulator::FindCycle(FIsoSegment* StartSegment, bool LeftSide, TArra
 static bool bDisplayStar = false;
 #endif
 
-FIsoSegment* FIsoTriangulator::FindNextSegment(EGridSpace Space, const FIsoSegment* StartSegment, const FIsoNode* StartNode, SlopMethod GetSlop) const
+FIsoSegment* FIsoTriangulator::FindNextSegment(EGridSpace Space, const FIsoSegment* StartSegment, const FIsoNode* StartNode, SlopeMethod GetSlope) const
 {
 	const FPoint2D& StartPoint = StartNode->Get2DPoint(Space, Grid);
 	const FPoint2D& EndPoint = (StartNode == &StartSegment->GetFirstNode()) ? StartSegment->GetSecondNode().Get2DPoint(Space, Grid) : StartSegment->GetFirstNode().Get2DPoint(Space, Grid);
@@ -2624,7 +2625,7 @@ FIsoSegment* FIsoTriangulator::FindNextSegment(EGridSpace Space, const FIsoSegme
 	{
 		const FPoint2D& OtherPoint = (StartNode == &Segment->GetFirstNode()) ? Segment->GetSecondNode().Get2DPoint(Space, Grid) : Segment->GetFirstNode().Get2DPoint(Space, Grid);
 
-		double Slope = GetSlop(StartPoint, OtherPoint, ReferenceSlope);
+		double Slope = GetSlope(StartPoint, OtherPoint, ReferenceSlope);
 		if (Slope < SMALL_NUMBER_SQUARE) Slope = 8;
 
 #ifdef DEBUG_FIND_NEXTSEGMENT
@@ -2929,8 +2930,8 @@ void FIsoTriangulator::TryToConnectTwoLoopsWithIsocelesTriangle(FCell& Cell, con
 
 	TFunction<FIsoNode* (FIsoSegment*)> FindBestTriangle = [&](FIsoSegment* Segment) -> FIsoNode*
 	{
-		SlopMethod GetSlopAtStartNode = ClockwiseSlop;
-		SlopMethod GetSlopAtEndNode = CounterClockwiseSlop;
+		SlopeMethod GetSlopeAtStartNode = ClockwiseSlope;
+		SlopeMethod GetSlopeAtEndNode = CounterClockwiseSlope;
 
 		// StartNode = A
 		FIsoNode& StartNode = Segment->GetSecondNode();
@@ -2975,7 +2976,7 @@ void FIsoTriangulator::TryToConnectTwoLoopsWithIsocelesTriangle(FCell& Cell, con
 		double EndReferenceSlope = StartReferenceSlope < 4 ? StartReferenceSlope + 4 : StartReferenceSlope - 4;
 
 		double MinCriteria = HUGE_VALUE;
-		const double MinSlopToNotBeAligned = 0.0001;
+		const double MinSlopeToNotBeAligned = 0.0001;
 		double CandidateSlopeAtStartNode = 8.;
 		double CandidateSlopeAtEndNode = 8.;
 
@@ -2992,18 +2993,18 @@ void FIsoTriangulator::TryToConnectTwoLoopsWithIsocelesTriangle(FCell& Cell, con
 				continue;
 			}
 
-			double SlopeAtStartNode = GetSlopAtStartNode(StartPoint2D, NodePoint2D, StartReferenceSlope);
-			double SlopeAtEndNode = GetSlopAtEndNode(EndPoint2D, NodePoint2D, EndReferenceSlope);
+			double SlopeAtStartNode = GetSlopeAtStartNode(StartPoint2D, NodePoint2D, StartReferenceSlope);
+			double SlopeAtEndNode = GetSlopeAtEndNode(EndPoint2D, NodePoint2D, EndReferenceSlope);
 
 			// check the side of the candidate point accordint to the segment
-			if (SlopeAtStartNode <= MinSlopToNotBeAligned)
+			if (SlopeAtStartNode <= MinSlopeToNotBeAligned)
 			{
 				continue;
 			}
 
 			if (
 				// the candidate triangle is inside the current candidate triangle
-				((SlopeAtStartNode < (CandidateSlopeAtStartNode + MinSlopToNotBeAligned)) && (SlopeAtEndNode < (CandidateSlopeAtEndNode + MinSlopToNotBeAligned)))
+				((SlopeAtStartNode < (CandidateSlopeAtStartNode + MinSlopeToNotBeAligned)) && (SlopeAtEndNode < (CandidateSlopeAtEndNode + MinSlopeToNotBeAligned)))
 				||
 				// the candidate triangle is better than the current candidate triangle and doesn't contain the current candidate triangle
 				((PointCriteria < MinCriteria) && ((SlopeAtStartNode > CandidateSlopeAtStartNode) ^ (SlopeAtEndNode > CandidateSlopeAtEndNode))))
diff --git a/Engine/Source/Runtime/Datasmith/CADKernel/Private/CADKernel/Mesh/Meshers/IsoTriangulator.h b/Engine/Source/Runtime/Datasmith/CADKernel/Private/CADKernel/Mesh/Meshers/IsoTriangulator.h
--- a/Engine/Source/Runtime/Datasmith/CADKernel/Private/CADKernel/Mesh/Meshers/IsoTriangulator.h
+++ b/Engine/Source/Runtime/Datasmith/CADKernel/Private/CADKernel/Mesh/Meshers/IsoTriangulator.h
@@ -333,7 +333,7 @@ public:
 
 private:
 
-	FIsoSegment* FindNextSegment(EGridSpace Space, const FIsoSegment* StartSegment, const FIsoNode* StartNode, SlopMethod GetSlop) const;
+	FIsoSegment* FindNextSegment(EGridSpace Space, const FIsoSegment* StartSegment, const FIsoNode* StartNode, SlopeMethod GetSlop) const;
 
 	// ==========================================================================================
 	// 	   Create segments
diff --git a/Engine/Source/Runtime/Datasmith/CADKernel/Private/CADKernel/Mesh/Structure/Grid.cpp b/Engine/Source/Runtime/Datasmith/CADKernel/Private/CADKernel/Mesh/Structure/Grid.cpp
--- a/Engine/Source/Runtime/Datasmith/CADKernel/Private/CADKernel/Mesh/Structure/Grid.cpp
+++ b/Engine/Source/Runtime/Datasmith/CADKernel/Private/CADKernel/Mesh/Structure/Grid.cpp
@@ -76,7 +76,7 @@ void FGrid::ProcessPointCloud()
 	Wait(bDisplay);
 #endif
 	FindPointsCloseToLoop();
-	RemovePointsClosedToLoop();
+	RemovePointsCloseToLoop();
 
 	// Removed of Thin zone boundary (the last boundaries). In case of thin zone, the number of 2d boundary will be biggest than 3d boundary one.
 	// Only EGridSpace::Default2D is needed.
@@ -324,8 +324,7 @@ bool FGrid::GeneratePointCloud()
 		return false;
 	}
 
-	IsInsideFace.Init(true, CuttingSize);
-	IsCloseToLoop.Init(0, CuttingSize);
+	PointMarkers.Init(EPointMarker::None, CuttingSize);
 
 	CountOfInnerNodes = CuttingSize;
 	for (int32 Index = 0; Index < EGridSpace::EndGridSpace; ++Index)
@@ -432,10 +431,10 @@ void FGrid::FindPointsCloseToLoop()
 
 	TFunction<void()> SetCellCloseToLoop = [&]()
 	{
-		IsCloseToLoop[GlobalIndex] = 1;
-		IsCloseToLoop[GlobalIndex - 1] = 1;
-		IsCloseToLoop[GlobalIndex - 1 - CuttingCount[EIso::IsoU]] = 1;
-		IsCloseToLoop[GlobalIndex - CuttingCount[EIso::IsoU]] = 1;
+		SetCloseToLoop(GlobalIndex);
+		SetCloseToLoop(GlobalIndex - 1);
+		SetCloseToLoop(GlobalIndex - 1 - CuttingCount[EIso::IsoU]);
+		SetCloseToLoop(GlobalIndex - CuttingCount[EIso::IsoU]);
 #ifdef DEBUG_FINDPOINTSCLOSETOLOOP
 		DisplayCell();
 #endif
@@ -489,7 +488,7 @@ void FGrid::FindPointsCloseToLoop()
 		return FirstValue + DOUBLE_SMALL_NUMBER < SecondValue;
 	};
 
-	double Slop;
+	double Slope;
 	double Origin;
 
 	TFunction<void(EIso, const int32, const int32, TFunction<void()>, TFunction<void()>)> FindIntersection = [&](EIso MainIso, const int32 DeltaIso, const int32 DeltaOther, TFunction<void()> OffsetIndexIfBigger, TFunction<void()> OffsetIndexIfSmaller)
@@ -500,7 +499,7 @@ void FGrid::FindPointsCloseToLoop()
 		EIso OtherIso = Other(MainIso);
 		while (TestAlongIso(UniformCuttingCoordinates[MainIso][Index[MainIso] - DeltaIso], (*PointB)[MainIso]) || TestAlongOther(UniformCuttingCoordinates[OtherIso][Index[OtherIso] - DeltaOther], (*PointB)[OtherIso]))
 		{
-			double CoordinateOther = Slop * UniformCuttingCoordinates[MainIso][Index[MainIso] - DeltaIso] + Origin;
+			double CoordinateOther = Slope * UniformCuttingCoordinates[MainIso][Index[MainIso] - DeltaIso] + Origin;
 			if (IsReallyBigger(CoordinateOther, UniformCuttingCoordinates[OtherIso][Index[OtherIso] - DeltaOther]))
 			{
 				OffsetIndexIfBigger();
@@ -591,8 +590,8 @@ void FGrid::FindPointsCloseToLoop()
 			ABu = PointB->U - PointA->U;
 			if (FMath::Abs(ABu) > FMath::Abs(ABv))
 			{
-				Slop = ABv / ABu;
-				Origin = PointA->V - Slop * PointA->U;
+				Slope = ABv / ABu;
+				Origin = PointA->V - Slope * PointA->U;
 				if (ABu > 0)
 				{
 					if (ABv > 0)
@@ -618,8 +617,8 @@ void FGrid::FindPointsCloseToLoop()
 			}
 			else
 			{
-				Slop = ABu / ABv;
-				Origin = PointA->U - Slop * PointA->V;
+				Slope = ABu / ABv;
+				Origin = PointA->U - Slope * PointA->V;
 				if (ABu > 0)
 				{
 					if (ABv > 0)
@@ -658,7 +657,7 @@ void FGrid::FindPointsCloseToLoop()
 	Chronos.FindPointsCloseToLoopDuration += FChrono::Elapse(StartTime);
 }
 
-void FGrid::RemovePointsClosedToLoop()
+void FGrid::RemovePointsCloseToLoop()
 {
 	FTimePoint StartTime = FChrono::Now();
 
@@ -740,7 +739,7 @@ void FGrid::RemovePointsClosedToLoop()
 		IndexOfPointsNearAndInsideLoop.Reserve(CuttingSize);
 		for (int32 Index = 0; Index < CuttingSize; ++Index)
 		{
-			if (IsCloseToLoop[Index] && IsInsideFace[Index])
+			if (IsNodeCloseToAndInsideLoop(Index))
 			{
 				IndexOfPointsNearAndInsideLoop.Add(Index);
 			}
@@ -865,8 +864,7 @@ void FGrid::RemovePointsClosedToLoop()
 				continue;
 			}
 
-			IsCloseToLoop[Index] = 0;
-			IsInsideFace[Index] = 0;
+			SetTooCloseToLoop(Index);
 			CountOfInnerNodes--;
 			break;
 		}
@@ -876,7 +874,7 @@ void FGrid::RemovePointsClosedToLoop()
 }
 
 /**
- * For the surfacic normal at a StartPoint of the 3D degenerated curve (Not degenerated in 2d)
+ * For the surface normal at a StartPoint of the 3D degenerated curve (Not degenerated in 2d)
  * The normal is swap if StartPoint is too close to the Boundary
  * The norm of the normal is defined as 1/20 of the parallel boundary Length
  */
@@ -1479,7 +1477,7 @@ void FGrid::FindInnerFacePoints()
 	TArray<char> IntersectLoop;
 
 	IntersectLoop.Init(0, CuttingSize);
-	IsInsideFace.Init(1, CuttingSize);
+	PointMarkers.Init(EPointMarker::IsInside, CuttingSize);
 
 	NbIntersectUForward.Init(0, CuttingSize);
 	NbIntersectUBackward.Init(0, CuttingSize);
@@ -1745,7 +1743,7 @@ void FGrid::FindInnerFacePoints()
 	{
 		if (IntersectLoop[Index])
 		{
-			IsInsideFace[Index] = 0;
+			ResetInsideLoop(Index);
 			CountOfInnerNodes--;
 			continue;
 		}
@@ -1769,7 +1767,7 @@ void FGrid::FindInnerFacePoints()
 		}
 		if (IsInside < 3)
 		{
-			IsInsideFace[Index] = false;
+			ResetInsideLoop(Index);
 			CountOfInnerNodes--;
 		}
 	}
diff --git a/Engine/Source/Runtime/Datasmith/CADKernel/Private/CADKernel/Mesh/Structure/Grid.h b/Engine/Source/Runtime/Datasmith/CADKernel/Private/CADKernel/Mesh/Structure/Grid.h
--- a/Engine/Source/Runtime/Datasmith/CADKernel/Private/CADKernel/Mesh/Structure/Grid.h
+++ b/Engine/Source/Runtime/Datasmith/CADKernel/Private/CADKernel/Mesh/Structure/Grid.h
@@ -71,6 +71,23 @@ struct FGridChronos
 	}
 };
 
+enum class EPointMarker : uint8
+{
+	None = 0x00u,  // No flags.
+
+	IsInside = 0x01u, 
+	IsCloseToLoop = 0x02u,
+
+	IsInsideAndClose = 0x03u,
+
+	//IsTooCloseToLoop = 0x04u,  // node inside the loop but too close to the loop to be include in the mesh
+	//IsInsideAndTooCloseToLoop = 0x05u,
+
+	All = 0xFFu
+};
+ENUM_CLASS_FLAGS(EPointMarker);
+
+
 class FGrid : public FHaveStates
 {
 protected:
@@ -131,16 +148,10 @@ protected:
 	double MinOfMaxElementSize = 0;
 
 	/**
-	 * Array to flag each points as inside the face (if IsInsideDomain[Index] = 1) or outer
-	 * @see IsNodeInsideFace
-	 */
-	TArray<char> IsInsideFace;
-
-	/**
-	 * Array to flag each points as close to a loop of the face if IsCloseToLoop[Index] = 1
-	 * @see IsNodeCloseToLoop
+	 * Array to flag each points as inside, close, ... the face 
+	 * @see IsNode..., 
 	 */
-	TArray<char> IsCloseToLoop;
+	TArray<EPointMarker> PointMarkers;
 
 	/**
 	 * 2D Coordinate of grid nodes in each space
@@ -249,7 +260,7 @@ protected:
 	 *
 	 * This algorithm is inspired of Bresenham's line algorithm
 	 * The loop is traversed segment by segment
-	 * For each segment, according to the slop of the segment, each cell intersecting the segment is selected (FindIntersectionsCaseSlop_) i.e. each corner node of the cell is flagged IsCloseToLoop (SetCellCloseToLoop)
+	 * For each segment, according to the slop of the segment, each cell intersecting the segment is selected (FindIntersectionsCaseSlope_) i.e. each corner node of the cell is flagged IsCloseToLoop (SetCellCloseToLoop)
 	 * @see SlopeUtils
 	 * @see IsNodeCloseToLoop
 	 * @see ProcessPointCloud (called in ProcessPointClound)
@@ -265,8 +276,7 @@ protected:
 	 *
 	 * @see ProcessPointClound (called in ProcessPointClound)
 	 */
-	void RemovePointsClosedToLoop();
-
+	void RemovePointsCloseToLoop();
 
 	/**
 	 * Gets the cutting coordinates of the existing mesh of bordering edges (loop's edges)
@@ -322,7 +332,7 @@ public:
 	 */
 	const bool IsNodeInsideFace(int32 IndexU, int32 IndexV) const
 	{
-		return IsInsideFace[GobalIndex(IndexU, IndexV)] == 1;
+		return IsNodeInsideFace(GobalIndex(IndexU, IndexV));
 	}
 
 	/**
@@ -330,7 +340,30 @@ public:
 	 */
 	const bool IsNodeInsideFace(int32 Index) const
 	{
-		return IsInsideFace[Index] == 1;
+		return (PointMarkers[Index] & EPointMarker::IsInside) == EPointMarker::IsInside;
+	}
+
+	/**
+	 * @return true if 
+	 */
+	//const bool IsNodeFarFromFace(int32 Index) const
+	//{
+	//	return (PointMarkers[Index] & EPointMarker::IsInsideAndClose) == EPointMarker::None;
+	//}
+
+	const bool IsNodeOusideFace(int32 Index) const
+	{
+		return (PointMarkers[Index] & EPointMarker::IsInside) == EPointMarker::None;
+	}
+
+	const bool IsNodeCloseToButOusideFace(int32 Index) const
+	{
+		return (PointMarkers[Index] & EPointMarker::IsInsideAndClose) == EPointMarker::IsCloseToLoop;
+	}
+
+	void SetNodeInside(int32 Index)
+	{
+		PointMarkers[Index] |= EPointMarker::IsInside;
 	}
 
 	/**
@@ -344,19 +377,63 @@ public:
 	/**
 	 * @return true if the node is close to a loop i.e. the loop cross the space [[IndexU - 1, IndexU + 1], [IndexV - 1, IndexV + 1]]
 	 */
-	const bool IsNodeCloseToLoop(int32 IndexU, int32 IndexV) const
+	const bool IsNodeCloseToLoop(int32 Index) const
+	{
+		return (PointMarkers[Index] & EPointMarker::IsCloseToLoop) == EPointMarker::IsCloseToLoop;
+	}
+
+	/*const bool IsNodeTooCloseToLoop2(int32 IndexU, int32 IndexV) const
+	{
+		return IsNodeTooCloseToLoop2(GobalIndex(IndexU, IndexV));
+	}*/
+
+	//const bool IsNodeTooCloseToLoop2(int32 Index) const
+	//{
+	//	return (PointMarkers[Index] & EPointMarker::IsTooCloseToLoop) == EPointMarker::IsTooCloseToLoop;
+	//}
+
+	void SetCloseToLoop(int32 Index)
+	{
+		PointMarkers[Index] |= EPointMarker::IsCloseToLoop;
+	}
+
+	/*void ResetCloseToAndInsideLoop(int32 Index)
 	{
-		return IsCloseToLoop[GobalIndex(IndexU, IndexV)] == 1;
+		PointMarkers[Index] &= ~EPointMarker::IsInsideAndClose;
+	}*/
+
+	void SetTooCloseToLoop(int32 Index)
+	{
+		//PointMarkers[Index] |= EPointMarker::IsTooCloseToLoop;
+		PointMarkers[Index] &= ~EPointMarker::IsInsideAndClose;
+	}
+
+	void ResetInsideLoop(int32 Index)
+	{
+		PointMarkers[Index] &= ~EPointMarker::IsInside;
 	}
 
 	/**
-	 * @return true if the node is close to a loop i.e. the loop cross the space [[IndexU - 1, IndexU + 1], [IndexV - 1, IndexV + 1]]
+	 * @return true if the node is close to or inside a loop i.e. the loop cross the space [[IndexU - 1, IndexU + 1], [IndexV - 1, IndexV + 1]]
 	 */
-	const bool IsNodeCloseToLoop(int32 Index) const
+	/*const bool IsNodeCloseToOrInsideLoop2(int32 Index) const
+	{
+		return (bool) (PointMarkers[Index] & EPointMarker::IsInsideAndClose);
+	}*/
+
+	const bool IsNodeCloseToAndInsideLoop(int32 Index) const
 	{
-		return IsCloseToLoop[Index] == 1;
+		return (PointMarkers[Index] & EPointMarker::IsInsideAndClose) == EPointMarker::IsInsideAndClose;
 	}
 
+	/**
+	 * @return true if the node is inside the face but not too close to the loop to be meshed
+	 */
+	//const bool IsNodeMeshable(int32 Index) const
+	//{
+	//	return (PointMarkers[Index] & EPointMarker::IsInsideAndTooCloseToLoop) == EPointMarker::IsInside;
+	//}
+
 	/**
 	 * @return the FPoint2D (parametric coordinates) of the point at the Index of the grid in the defined grid space
 	 * @see EGridSpace
@@ -641,13 +718,13 @@ public:
 		F3DDebugSession _(Message);
 		for (int32 Index = 0; Index < Points.Num(); ++Index)
 		{
-			if (IsInsideFace[Index])
+			if (IsNodeInsideFace(Index))
 			{
-				DisplayPoint(Points[Index] * DisplayScale, IsCloseToLoop[Index] ? EVisuProperty::BluePoint : EVisuProperty::GreenPoint, Index);
+				DisplayPoint(Points[Index] * DisplayScale, /*IsNodeTooCloseToLoop2(Index) ? EVisuProperty::RedPoint :*/ EVisuProperty::BluePoint, Index);
 			}
 			else
 			{
-				DisplayPoint(Points[Index] * DisplayScale, EVisuProperty::OrangePoint, Index);
+				DisplayPoint(Points[Index] * DisplayScale, IsNodeCloseToLoop(Index) ? EVisuProperty::YellowPoint : EVisuProperty::GreenPoint , Index);
 			}
 		}
 	}
diff --git a/Engine/Source/Runtime/Datasmith/CADKernel/Private/CADKernel/Mesh/Structure/GridDisplay.cpp b/Engine/Source/Runtime/Datasmith/CADKernel/Private/CADKernel/Mesh/Structure/GridDisplay.cpp
--- a/Engine/Source/Runtime/Datasmith/CADKernel/Private/CADKernel/Mesh/Structure/GridDisplay.cpp
+++ b/Engine/Source/Runtime/Datasmith/CADKernel/Private/CADKernel/Mesh/Structure/GridDisplay.cpp
@@ -309,7 +309,7 @@ void FGrid::DisplayGridPoints(EGridSpace DisplaySpace) const
 		F3DDebugSession _(TEXT("FGrid::FindInnerDomainPoints Inside Point"));
 		for (int32 Index = 0; Index < CuttingSize; ++Index)
 		{
-			if (IsInsideFace[Index])
+			if (IsNodeInsideFace(Index))
 			{
 				DisplayPoint(Points2D[DisplaySpace][Index] * DisplayScale, EVisuProperty::BluePoint, Index);
 				NbNum++;
@@ -322,7 +322,7 @@ void FGrid::DisplayGridPoints(EGridSpace DisplaySpace) const
 		F3DDebugSession _(TEXT("FGrid::FindInnerDomainPoints Outside Point"));
 		for (int32 Index = 0; Index < CuttingSize; ++Index)
 		{
-			if (!IsInsideFace[Index])
+			if (!IsNodeInsideFace(Index))
 			{
 				DisplayPoint(Points2D[DisplaySpace][Index] * DisplayScale, EVisuProperty::OrangePoint, Index);
 			}
diff --git a/Engine/Source/Runtime/Datasmith/CADKernel/Private/CADKernel/Mesh/Structure/LoopCleaner.cpp b/Engine/Source/Runtime/Datasmith/CADKernel/Private/CADKernel/Mesh/Structure/LoopCleaner.cpp
--- a/Engine/Source/Runtime/Datasmith/CADKernel/Private/CADKernel/Mesh/Structure/LoopCleaner.cpp
+++ b/Engine/Source/Runtime/Datasmith/CADKernel/Private/CADKernel/Mesh/Structure/LoopCleaner.cpp
@@ -34,6 +34,11 @@ FLoopCleaner::FLoopCleaner(FIsoTriangulator& Triangulator)
 bool FLoopCleaner::CleanLoops()
 {
 
+	//{
+	//	FPoint2D Point = LoopNodes[36].Get2DPoint(EGridSpace::UniformScaled, Grid);
+	//	LoopNodes[56].Set2DPoint(EGridSpace::UniformScaled, Grid, Point);
+	//}
+
 	FindBestLoopExtremity();
 
 #ifdef DEBUG_CLEAN_LOOPS		
@@ -260,15 +265,15 @@ bool FLoopCleaner::TryToRemoveIntersectionOfTwoConsecutiveIntersectingSegments(c
 	TSegment<FPoint2D> IntersectingSegment2D(IntersectingSegment.GetFirstNode().Get2DPoint(EGridSpace::UniformScaled, Grid), IntersectingSegment.GetSecondNode().Get2DPoint(EGridSpace::UniformScaled, Grid));
 	TSegment<FPoint2D> Segment2D(Segment.GetFirstNode().Get2DPoint(EGridSpace::UniformScaled, Grid), Segment.GetSecondNode().Get2DPoint(EGridSpace::UniformScaled, Grid));
 
-	double IntersectingSegmentSlop = ComputeOrientedSlope(IntersectingSegment2D.Point0, IntersectingSegment2D.Point1, 0);
-	double SegmentSlop = ComputeUnorientedSlope(Segment2D.Point1, Segment2D.Point0, IntersectingSegmentSlop);
-	if (SegmentSlop > 2)
+	double IntersectingSegmentSlope = ComputeOrientedSlope(IntersectingSegment2D.Point0, IntersectingSegment2D.Point1, 0);
+	double SegmentSlope = ComputeUnorientedSlope(Segment2D.Point1, Segment2D.Point0, IntersectingSegmentSlope);
+	if (SegmentSlope > 2)
 	{
-		SegmentSlop = 4 - SegmentSlop;
+		SegmentSlope = 4 - SegmentSlope;
 	}
 
 	// if the segment and IntersectingSegment are parallel, segment are moved inside
-	if (SegmentSlop < 0.01)
+	if (SegmentSlope < 0.01)
 	{
 		double StartPointSquareDistance = SquareDistanceOfPointToSegment(Segment2D.Point0, IntersectingSegment2D.Point0, IntersectingSegment2D.Point1);
 		double EndPointSquareDistance = SquareDistanceOfPointToSegment(Segment2D.Point1, IntersectingSegment2D.Point0, IntersectingSegment2D.Point1);
@@ -296,8 +301,8 @@ bool FLoopCleaner::TryToRemoveIntersectionOfTwoConsecutiveIntersectingSegments(c
 		}
 	}
 
-	double OrientedSlop = ComputeOrientedSlope(IntersectingSegment2D.Point0, Segment.GetFirstNode().Get2DPoint(EGridSpace::UniformScaled, Grid), IntersectingSegmentSlop);
-	if (OrientedSlop >= 0)
+	double OrientedSlope = ComputeOrientedSlope(IntersectingSegment2D.Point0, Segment.GetFirstNode().Get2DPoint(EGridSpace::UniformScaled, Grid), IntersectingSegmentSlope);
+	if (OrientedSlope >= 0)
 	{
 		Node = (FLoopNode*)&Segment.GetSecondNode();
 		PreviousNode = (FLoopNode*)&Segment.GetFirstNode();
@@ -446,13 +451,13 @@ void FLoopCleaner::RemoveIntersectionByMovingOutsideSegmentNodeInside(const FIso
 
 	// Is Second node, the outside node ?
 	{
-		double OrientedSlop = ComputeOrientedSlope(IntersectingPoints[0], IntersectingPoints[1], PointToMove);
+		double OrientedSlope = ComputeOrientedSlope(IntersectingPoints[0], IntersectingPoints[1], PointToMove);
 		if (bIsSameInnerLoop)
 		{
-			OrientedSlop *= -1;
+			OrientedSlope *= -1;
 		}
 
-		if (OrientedSlop > 0)
+		if (OrientedSlope > 0)
 		{
 			PointToMove = Segment.GetFirstNode().Get2DPoint(EGridSpace::UniformScaled, Grid);
 			bFirstNodeIsOutside = true;
@@ -501,248 +506,6 @@ void FLoopCleaner::RemoveIntersectionByMovingOutsideSegmentNodeInside(const FIso
 #endif
 }
 
-
-#ifdef UNUSED
-bool FLoopCleaner::TryToRemoveSelfIntersectionByMovingTheClosedOusidePoint(const FIsoSegment& Segment0, const FIsoSegment& Segment1)
-{
-	const FIsoNode* Nodes[2][2] = { {nullptr, nullptr}, {nullptr, nullptr} };
-	Nodes[0][0] = &Segment0.GetFirstNode();
-	Nodes[0][1] = &Segment0.GetSecondNode();
-	Nodes[1][0] = &Segment1.GetFirstNode();
-	Nodes[1][1] = &Segment1.GetSecondNode();
-
-	FPoint2D Points[2][2];
-	Points[0][0] = Nodes[0][0]->Get2DPoint(EGridSpace::UniformScaled, Grid);
-	Points[0][1] = Nodes[0][1]->Get2DPoint(EGridSpace::UniformScaled, Grid);
-	Points[1][0] = Nodes[1][0]->Get2DPoint(EGridSpace::UniformScaled, Grid);
-	Points[1][1] = Nodes[1][1]->Get2DPoint(EGridSpace::UniformScaled, Grid);
-
-	int32 ProjectedNodeIndex[2][2] = { { 0, 0 }, { 0, 0 } };
-	FPoint2D ProjectedPoints[2][2];
-	double Distance[2][2] = { { HUGE_VALUE, HUGE_VALUE }, { HUGE_VALUE, HUGE_VALUE } };
-	FPoint2D ProjectedPointsOut[2][2];
-
-	//int32 OtherSegmentIndex = 0;
-	int32 SegmentIndex = 1;
-	int32 OtherSegmentIndex = 0;
-
-	TFunction<void(int32)> ProjectPoint = [&](int32 OtherSegmentNodeIndex)
-	{
-		double Coordinate;
-		ProjectedPoints[OtherSegmentIndex][OtherSegmentNodeIndex] = ProjectPointOnSegment(Points[OtherSegmentIndex][OtherSegmentNodeIndex], Points[SegmentIndex][0], Points[SegmentIndex][1], Coordinate, false);
-		bool bProjectedPoint1IsInside = (Coordinate >= -DOUBLE_SMALL_NUMBER && Coordinate <= 1 + DOUBLE_SMALL_NUMBER);
-		if (bProjectedPoint1IsInside)
-		{
-			Distance[OtherSegmentIndex][OtherSegmentNodeIndex] = ProjectedPoints[OtherSegmentIndex][OtherSegmentNodeIndex].SquareDistance(Points[OtherSegmentIndex][OtherSegmentNodeIndex]);
-		}
-
-#ifdef DEBUG_SELF_CLOSED_OUSIDE_POINT
-		if (bDisplay)
-		{
-			{
-				F3DDebugSession _(TEXT("Segs"));
-				DisplaySegment(Points[SegmentIndex][0], Points[SegmentIndex][1], 0, EVisuProperty::BlueCurve);
-				DisplayPoint(Points[OtherSegmentIndex][OtherSegmentNodeIndex], EVisuProperty::RedPoint);
-				DisplayPoint(ProjectedPoints[OtherSegmentIndex][OtherSegmentNodeIndex], EVisuProperty::RedPoint);
-				Wait(false);
-			}
-		}
-#endif
-	};
-
-	SegmentIndex = 0;
-	OtherSegmentIndex = 1;
-	ProjectPoint(0);
-	ProjectPoint(1);
-	SegmentIndex = 1;
-	OtherSegmentIndex = 0;
-	ProjectPoint(0);
-	ProjectPoint(1);
-
-	TFunction<void(int32, int32)> MoveNode = [&](int32 MoveSegmentIndex, int32 MovePointIndex)
-	{
-		FPoint2D MoveDirection = ProjectedPoints[MoveSegmentIndex][MovePointIndex] - Points[MoveSegmentIndex][MovePointIndex];
-		MoveDirection.Normalize();
-
-		if (MoveDirection.SquareLength() < 0.5)
-		{
-			MoveDirection = Points[MoveSegmentIndex][MovePointIndex == 0 ? 1 : 0] - ProjectedPoints[MoveSegmentIndex][MovePointIndex];
-			MoveDirection.Normalize();
-		}
-
-		MoveDirection *= GeometricTolerance;
-		FPoint2D NewCoordinate = ProjectedPoints[MoveSegmentIndex][MovePointIndex] + MoveDirection;
-
-		FLoopNode& Node = (FLoopNode&) const_cast<FIsoNode&>(*Nodes[MoveSegmentIndex][MovePointIndex]);
-		FLoopNode& PreviousNode = Node.GetPreviousNode();
-		FLoopNode& NextNode = Node.GetNextNode();
-
-		if ((PreviousNode.Get2DPoint(EGridSpace::UniformScaled, Grid).SquareDistance(NewCoordinate) < SquareGeometricTolerance2) ||
-			(NextNode.Get2DPoint(EGridSpace::UniformScaled, Grid).SquareDistance(NewCoordinate) < SquareGeometricTolerance2))
-		{
-			RemoveNodeOfLoop(Node);
-			return;
-		}
-
-		const_cast<FIsoNode&>((*Nodes[MoveSegmentIndex][MovePointIndex])).Set2DPoint(EGridSpace::UniformScaled, Grid, NewCoordinate);
-
-#ifdef DEBUG_SELF_CLOSED_OUSIDE_POINT
-		if (bDisplay)
-		{
-			F3DDebugSession _(TEXT("New Segs"));
-			DisplayPoint(NewCoordinate, EVisuProperty::RedPoint);
-			Display(EGridSpace::UniformScaled, Segment1, 0, EVisuProperty::RedCurve);
-			Display(EGridSpace::UniformScaled, Segment0, 0, EVisuProperty::BlueCurve);
-			Wait(false);
-		}
-#endif
-	};
-
-
-	int32 MoveSegment = 0;
-	int32 MovePoint = 0;
-	double MinDistance = Distance[0][0];
-
-	TFunction<void(int32, int32)> FindNodeToMove = [&](int32 MoveSegmentIndex, int32 MovePointIndex)
-	{
-		if (Distance[MoveSegmentIndex][MovePointIndex] < MinDistance)
-		{
-			MinDistance = Distance[MoveSegmentIndex][MovePointIndex];
-			MoveSegment = MoveSegmentIndex;
-			MovePoint = MovePointIndex;
-		}
-	};
-
-	FindNodeToMove(0, 1);
-	FindNodeToMove(1, 0);
-	FindNodeToMove(1, 1);
-
-	MoveNode(MoveSegment, MovePoint);
-
-	return true;
-}
-
-bool FLoopCleaner::TryToRemoveIntersectionByMovingTheClosedOusidePoint(const FIsoSegment& Segment0, const FIsoSegment& Segment1)
-{
-	// For Inner/outer segment intersecting outer loop (i.e. IntersectingSegment is Outer), the segment has to be moved inside IntersectingSegment
-	// For Inner segment intersecting its loop, the segment has to be moved inside IntersectingSegment but the orientation of the loop is counterclockwise => need to swap orientation
-	// For Inner segment intersecting another inner loop, the segment has to be moved outside the loop but the orientation of the loop is counterclockwise => don't need to swap orientation
-
-	// As the iteration of the segment start from outer loop's segments, Segment can not be from outer loop if IntersectingSegment is not from outer loop
-	int32 SegmentLoopIndex[2];
-	SegmentLoopIndex[0] = Segment0.GetFirstNode().IsALoopNode() ? ((FLoopNode&)Segment0.GetFirstNode()).GetLoopIndex() : 0;
-	SegmentLoopIndex[1] = Segment1.GetFirstNode().IsALoopNode() ? ((FLoopNode&)Segment1.GetFirstNode()).GetLoopIndex() : 0;
-	bool bIsOutLoopSegment = SegmentLoopIndex[0] == 0;
-	bool bIsSameLoop = SegmentLoopIndex[0] == SegmentLoopIndex[1];
-
-	bool bToInside = bIsOutLoopSegment || bIsSameLoop;
-
-	const FIsoNode* Nodes[2][2] = { {nullptr, nullptr}, {nullptr, nullptr} };
-	Nodes[0][0] = &Segment0.GetFirstNode();
-	Nodes[0][1] = &Segment0.GetSecondNode();
-	Nodes[1][0] = &Segment1.GetFirstNode();
-	Nodes[1][1] = &Segment1.GetSecondNode();
-
-	FPoint2D Points[2][2];
-	Points[0][0] = Nodes[0][0]->Get2DPoint(EGridSpace::UniformScaled, Grid);
-	Points[0][1] = Nodes[0][1]->Get2DPoint(EGridSpace::UniformScaled, Grid);
-	Points[1][0] = Nodes[1][0]->Get2DPoint(EGridSpace::UniformScaled, Grid);
-	Points[1][1] = Nodes[1][1]->Get2DPoint(EGridSpace::UniformScaled, Grid);
-
-	int32 ProjectedNodeIndex[2] = { 0, 0 };
-	FPoint2D ProjectedPoints[2];
-	double Distance[2] = { HUGE_VALUE, HUGE_VALUE };
-	int32 OtherSegmentIndex = 0;
-	int32 SegmentIndex = 1;
-
-	TFunction<void(int32)> ProjectInnerPoint = [&](int32 OtherSegmentNodeIndex)
-	{
-		double Coordinate;
-		ProjectedPoints[OtherSegmentIndex] = ProjectPointOnSegment(Points[OtherSegmentIndex][OtherSegmentNodeIndex], Points[SegmentIndex][0], Points[SegmentIndex][1], Coordinate, false);
-		bool bProjectedPoint1IsInside = (Coordinate >= -DOUBLE_SMALL_NUMBER && Coordinate <= 1 + DOUBLE_SMALL_NUMBER);
-		if (bProjectedPoint1IsInside)
-		{
-			Distance[OtherSegmentIndex] = ProjectedPoints[OtherSegmentIndex].Distance(Points[OtherSegmentIndex][OtherSegmentNodeIndex]);
-		}
-
-#ifdef DEBUG_CLOSED_OUSIDE_POINT
-		if (bDisplay)
-		{
-			{
-				F3DDebugSession _(TEXT("Segs"));
-				DisplaySegment(Points[SegmentIndex][0], Points[SegmentIndex][1], 0, EVisuProperty::BlueCurve);
-				DisplayPoint(Points[OtherSegmentIndex][OtherSegmentNodeIndex], EVisuProperty::RedPoint);
-				DisplayPoint(ProjectedPoints[OtherSegmentIndex], EVisuProperty::RedPoint);
-				Wait(false);
-			}
-		}
-#endif
-	};
-
-	TFunction<void()> FindInnerNodeAndProjectIt = [&]()
-	{
-		OtherSegmentIndex = SegmentIndex ? 0 : 1;
-
-		double OrientedSlop = ComputeOrientedSlope(Points[SegmentIndex][0], Points[SegmentIndex][1], Points[OtherSegmentIndex][0]);
-		ProjectedNodeIndex[OtherSegmentIndex] = ((SegmentLoopIndex[0] == 0) == (OrientedSlop < 0)) ? 0 : 1;
-		ProjectInnerPoint(ProjectedNodeIndex[OtherSegmentIndex]);
-	};
-
-	SegmentIndex = 0;
-	FindInnerNodeAndProjectIt();
-	SegmentIndex = 1;
-	FindInnerNodeAndProjectIt();
-
-	if (Distance[0] == HUGE_VALUE && Distance[1] == HUGE_VALUE)
-	{
-		return false;
-	}
-
-	TFunction<void(int32)> MoveNode = [&](int32 MovePointIndex)
-	{
-		FPoint2D MoveDirection = ProjectedPoints[MovePointIndex] - Points[MovePointIndex][ProjectedNodeIndex[MovePointIndex]];
-		MoveDirection.Normalize();
-
-		if (MoveDirection.SquareLength() < 0.5)
-		{
-			MoveDirection = Points[MovePointIndex][ProjectedNodeIndex[MovePointIndex] == 0 ? 1 : 0] - ProjectedPoints[MovePointIndex];
-			MoveDirection.Normalize();
-		}
-
-		MoveDirection *= GeometricTolerance;
-		FPoint2D NewCoordinate = ProjectedPoints[MovePointIndex] + MoveDirection;
-
-		FLoopNode& Node = (FLoopNode&) const_cast<FIsoNode&>(*Nodes[MovePointIndex][ProjectedNodeIndex[MovePointIndex]]);
-		FLoopNode& PreviousNode = Node.GetPreviousNode();
-		FLoopNode& NextNode = Node.GetNextNode();
-
-		if ((PreviousNode.Get2DPoint(EGridSpace::UniformScaled, Grid).SquareDistance(NewCoordinate) < SquareGeometricTolerance2) ||
-			(NextNode.Get2DPoint(EGridSpace::UniformScaled, Grid).SquareDistance(NewCoordinate) < SquareGeometricTolerance2))
-		{
-			RemoveNodeOfLoop(Node);
-			return;
-		}
-
-		const_cast<FIsoNode&>((*Nodes[MovePointIndex][ProjectedNodeIndex[MovePointIndex]])).Set2DPoint(EGridSpace::UniformScaled, Grid, NewCoordinate);
-
-#ifdef DEBUG_CLOSED_OUSIDE_POINT
-		if (bDisplay)
-		{
-			F3DDebugSession _(TEXT("New Segs"));
-			DisplayPoint(NewCoordinate, EVisuProperty::RedPoint);
-			Display(EGridSpace::UniformScaled, Segment1, 0, EVisuProperty::RedCurve);
-			Display(EGridSpace::UniformScaled, Segment0, 0, EVisuProperty::BlueCurve);
-			Wait(false);
-		}
-#endif
-	};
-
-	int32 MovePoint = (Distance[0] < Distance[1]) ? 0 : 1;
-	MoveNode(MovePoint);
-	return true;
-}
-#endif
-
 void FLoopCleaner::FixLoopOrientation()
 {
 	FLoopNode const* const* StartNode = NodesOfLoop.FindByPredicate([](const FLoopNode* Node) { return !Node->IsDelete(); });
@@ -1105,10 +868,10 @@ bool FLoopCleaner::RemoveIntersectionsOfSubLoop(int32 IntersectionIndex, int32 I
 				}
 
 
-				double SlopSide0 = ComputeUnorientedSlope(GetPrevious(NodeSide0)->Get2DPoint(EGridSpace::UniformScaled, Grid), NodeSide0->Get2DPoint(EGridSpace::UniformScaled, Grid), GetNext(NodeSide0)->Get2DPoint(EGridSpace::UniformScaled, Grid));
-				double SlopSide1 = ComputeUnorientedSlope(GetPrevious(NodeSide1)->Get2DPoint(EGridSpace::UniformScaled, Grid), NodeSide1->Get2DPoint(EGridSpace::UniformScaled, Grid), GetNext(NodeSide1)->Get2DPoint(EGridSpace::UniformScaled, Grid));
+				double SlopeSide0 = ComputeUnorientedSlope(GetPrevious(NodeSide0)->Get2DPoint(EGridSpace::UniformScaled, Grid), NodeSide0->Get2DPoint(EGridSpace::UniformScaled, Grid), GetNext(NodeSide0)->Get2DPoint(EGridSpace::UniformScaled, Grid));
+				double SlopeSide1 = ComputeUnorientedSlope(GetPrevious(NodeSide1)->Get2DPoint(EGridSpace::UniformScaled, Grid), NodeSide1->Get2DPoint(EGridSpace::UniformScaled, Grid), GetNext(NodeSide1)->Get2DPoint(EGridSpace::UniformScaled, Grid));
 
-				if (SlopSide0 < SlopSide1)
+				if (SlopeSide0 < SlopeSide1)
 				{
 					MoveNodeBehindOther(NodeSide1, NodeSide0);
 				}
@@ -1497,16 +1260,16 @@ bool FLoopCleaner::IsAPinch(const LoopCleanerImpl::FPinchIntersectionContext& Co
 
 bool FLoopCleaner::DisconnectCoincidentNodes(const LoopCleanerImpl::FPinchIntersectionContext& Context)
 {
-	double Slop00 = ComputeSlope(*Context.Points[0][1], *Context.Points[0][0]);
-	double Slop02 = ComputeSlope(*Context.Points[0][1], *Context.Points[0][2]);
-	double Slop10 = ComputeSlope(*Context.Points[1][1], *Context.Points[1][0]);
-	double Slop12 = ComputeSlope(*Context.Points[1][1], *Context.Points[1][2]);
+	double Slope00 = ComputeSlope(*Context.Points[0][1], *Context.Points[0][0]);
+	double Slope02 = ComputeSlope(*Context.Points[0][1], *Context.Points[0][2]);
+	double Slope10 = ComputeSlope(*Context.Points[1][1], *Context.Points[1][0]);
+	double Slope12 = ComputeSlope(*Context.Points[1][1], *Context.Points[1][2]);
 
-	double MediumSlop0 = TransformIntoPositiveSlope((Slop00 + Slop02) * 0.5);
-	double MediumSlop1 = TransformIntoPositiveSlope((Slop10 + Slop12) * 0.5);
+	double MediumSlope0 = TransformIntoPositiveSlope((Slope00 + Slope02) * 0.5);
+	double MediumSlope1 = TransformIntoPositiveSlope((Slope10 + Slope12) * 0.5);
 
-	FPoint2D P1 = SlopeToVector(MediumSlop1);
-	FPoint2D P0 = SlopeToVector(MediumSlop0);
+	FPoint2D P1 = SlopeToVector(MediumSlope1);
+	FPoint2D P0 = SlopeToVector(MediumSlope0);
 	P1.Normalize()*= GeometricTolerance*30;
 	P0.Normalize()*= GeometricTolerance*30;
 
@@ -1670,10 +1433,10 @@ bool FLoopCleaner::TryToSwapSegmentsOrRemoveLoop(const TPair<double, double>& In
 	}
 #endif
 
-	double Slop = ComputeSlope(Segment0_Point0, Segment0_Point1);
-	Slop = ComputeUnorientedSlope(Segment1_Point0, Segment1_Point1, Slop);
+	double Slope = ComputeSlope(Segment0_Point0, Segment0_Point1);
+	Slope = ComputeUnorientedSlope(Segment1_Point0, Segment1_Point1, Slope);
 
-	if (Slop < 2)
+	if (Slope < 2)
 	{
 		FIsoSegment* Segment0 = Segment0_Node0->GetSegmentConnectedTo(Segment0_Node1);
 		FIsoSegment* Segment1 = Segment1_Node0->GetSegmentConnectedTo(Segment1_Node1);
@@ -1747,10 +1510,9 @@ void FLoopCleaner::SwapSubLoopOrientation(int32 FirstSegmentIndex, int32 LastSeg
 
 	TArray<FIsoSegment*> Segments;
 	{
-		int32 SegCount = LastSegmentIndex - FirstSegmentIndex;
-		if (SegCount < 0)
+		int32 SegCount = 0;
+		if (LastSegmentIndex < FirstSegmentIndex)
 		{
-			SegCount = 0;
 			int32 SegIndex = FitSegmentIndex(FirstSegmentIndex);
 			do
 			{
@@ -1758,6 +1520,10 @@ void FLoopCleaner::SwapSubLoopOrientation(int32 FirstSegmentIndex, int32 LastSeg
 				SegIndex = NextSegmentIndex(SegIndex);
 			} while (SegIndex != LastSegmentIndex);
 		}
+		else
+		{
+			SegCount = LastSegmentIndex - FirstSegmentIndex;
+		}
 		Segments.Reserve(SegCount);
 	}
 
@@ -2003,10 +1769,10 @@ bool FLoopCleaner::RemoveOuterNode(const TPair<double, double>& Intersection)
 	const FPoint2D* Point1 = &Node1->Get2DPoint(EGridSpace::UniformScaled, Grid);
 	const FPoint2D* NextPoint = &NextNode->Get2DPoint(EGridSpace::UniformScaled, Grid);
 
-	double Slop0 = ComputeUnorientedSlope(*Point0, *PreviousPoint, *Point1);
-	double Slop1 = ComputeUnorientedSlope(*Point1, *PreviousPoint, *NextPoint);
+	double Slope0 = ComputeUnorientedSlope(*Point0, *PreviousPoint, *Point1);
+	double Slope1 = ComputeUnorientedSlope(*Point1, *PreviousPoint, *NextPoint);
 
-	if (Slop0 < Slop1)
+	if (Slope0 < Slope1)
 	{
 		return RemoveNodeOfLoop(*Node0);
 	}
@@ -2190,25 +1956,25 @@ void FLoopCleaner::FindBestLoopExtremity()
 	FLoopNode* ExtremityNodes[4] = { nullptr, nullptr, nullptr, nullptr };
 
 	FLoopNode* BestNode = nullptr;
-	double OptimalSlop = 9.;
+	double OptimalSlope = 9.;
 	LoopIndex = 0;
 
 	//F3DDebugSession _(TEXT("FindBestLoopExtremity"));
 
 	BestStartNodeOfLoops.Reserve(Grid.GetLoopCount());
 
-	TFunction<void(FLoopNode*)> CompareWithSlopAt = [&](FLoopNode* Node)
+	TFunction<void(FLoopNode*)> CompareWithSlopeAt = [&](FLoopNode* Node)
 	{
 		FLoopNode& PreviousNode = Node->GetPreviousNode();
 		FLoopNode& NextNode = Node->GetNextNode();
-		double Slop = ComputePositiveSlope(Node->Get2DPoint(EGridSpace::UniformScaled, Grid), PreviousNode.Get2DPoint(EGridSpace::UniformScaled, Grid), NextNode.Get2DPoint(EGridSpace::UniformScaled, Grid));
+		double Slope = ComputePositiveSlope(Node->Get2DPoint(EGridSpace::UniformScaled, Grid), PreviousNode.Get2DPoint(EGridSpace::UniformScaled, Grid), NextNode.Get2DPoint(EGridSpace::UniformScaled, Grid));
 
-		//F3DDebugSession A(FString::Printf(TEXT("Node Slop %f"), Slop));
+		//F3DDebugSession A(FString::Printf(TEXT("Node Slope %f"), Slope));
 		//Display(EGridSpace::UniformScaled, *Node, 0, EVisuProperty::RedPoint);
 
-		if ((Slop > OptimalSlop) == (LoopIndex == 0))
+		if ((Slope > OptimalSlope) == (LoopIndex == 0))
 		{
-			OptimalSlop = Slop;
+			OptimalSlope = Slope;
 			BestNode = Node;
 		}
 	};
@@ -2216,11 +1982,11 @@ void FLoopCleaner::FindBestLoopExtremity()
 	TFunction<void()> FindLoopExtremity = [&]()
 	{
 		BestNode = nullptr;
-		OptimalSlop = (LoopIndex == 0) ? -1 : 9.;
+		OptimalSlope = (LoopIndex == 0) ? -1 : 9.;
 
 		for (FLoopNode* Node : ExtremityNodes)
 		{
-			CompareWithSlopAt(Node);
+			CompareWithSlopeAt(Node);
 		}
 		BestStartNodeOfLoops.Add(BestNode);
 
@@ -2320,7 +2086,7 @@ EOrientation FLoopCleaner::GetLoopOrientation(const FLoopNode* StartNode)
 	const FLoopNode* BestNode = nullptr;
 
 	LoopIndex = StartNode->GetLoopIndex();
-	double OptimalSlop = (LoopIndex == 0) ? -1 : 9.;
+	double OptimalSlope = (LoopIndex == 0) ? -1 : 9.;
 
 	double MaxFrontSlope = 4;
 	double MinBackSlope = 4;
@@ -2328,14 +2094,14 @@ EOrientation FLoopCleaner::GetLoopOrientation(const FLoopNode* StartNode)
 #ifdef DEBUG_LOOP_ORIENTATION
 	F3DDebugSession _(Grid.bDisplay, TEXT("GetLoopOrientation"));
 #endif
-	TFunction<void(const FLoopNode*)> CompareWithSlopAt = [&](const FLoopNode* Node)
+	TFunction<void(const FLoopNode*)> CompareWithSlopeAt = [&](const FLoopNode* Node)
 	{
 		FLoopNode& PreviousNode = Node->GetPreviousNode();
 		FLoopNode& NextNode = Node->GetNextNode();
 		double Slope = ComputePositiveSlope(Node->Get2DPoint(EGridSpace::UniformScaled, Grid), PreviousNode.Get2DPoint(EGridSpace::UniformScaled, Grid), NextNode.Get2DPoint(EGridSpace::UniformScaled, Grid));
 
 #ifdef DEBUG_LOOP_ORIENTATION
-		F3DDebugSession A(Grid.bDisplay, FString::Printf(TEXT("Slop %f"), Slope));
+		F3DDebugSession A(Grid.bDisplay, FString::Printf(TEXT("Slope %f"), Slope));
 		Grid.DisplayIsoNode(EGridSpace::UniformScaled, *Node);
 		Grid.DisplayIsoSegment(EGridSpace::UniformScaled, *Node, PreviousNode);
 		Grid.DisplayIsoSegment(EGridSpace::UniformScaled, *Node, NextNode);
@@ -2357,7 +2123,7 @@ EOrientation FLoopCleaner::GetLoopOrientation(const FLoopNode* StartNode)
 
 		for (const FLoopNode* Node : ExtremityNodes)
 		{
-			CompareWithSlopAt(Node);
+			CompareWithSlopeAt(Node);
 		}
 
 		// init for next loop
diff --git a/Engine/Source/Runtime/Datasmith/CADKernel/Private/CADKernel/Mesh/Structure/LoopCleaner.h b/Engine/Source/Runtime/Datasmith/CADKernel/Private/CADKernel/Mesh/Structure/LoopCleaner.h
--- a/Engine/Source/Runtime/Datasmith/CADKernel/Private/CADKernel/Mesh/Structure/LoopCleaner.h
+++ b/Engine/Source/Runtime/Datasmith/CADKernel/Private/CADKernel/Mesh/Structure/LoopCleaner.h
@@ -24,7 +24,7 @@ struct FMesherReport;
 
 namespace LoopCleanerImpl
 {
-typedef TFunction<double(const FPoint2D&, const FPoint2D&, double)> GetSlopMethod;
+typedef TFunction<double(const FPoint2D&, const FPoint2D&, double)> GetSlopeMethod;
 typedef TFunction<FLoopNode* (FLoopNode*)> GetNextNodeMethod;
 typedef TFunction<const FLoopNode* (const FLoopNode*)> GetNextConstNodeMethod;
 typedef TFunction<const FLoopNode* (const FIsoSegment*)> GetSegmentToNodeMethod;
@@ -141,7 +141,6 @@ private:
 	TArray<TPair<double, double>> Intersections;
 	FIntersectionSegmentTool LoopSegmentsIntersectionTool;
 
-
 	const double GeometricTolerance;
 	const double SquareGeometricTolerance;
 	const double SquareGeometricTolerance2;
@@ -211,9 +210,9 @@ private:
 	 */
 	bool CheckAndRemovePick(const FPoint2D& PreviousPoint, const FPoint2D& NodeToRemovePoint, const FPoint2D& NextPoint, FLoopNode& NodeToRemove)
 	{
-		double Slop = ComputeUnorientedSlope(NodeToRemovePoint, PreviousPoint, NextPoint);
+		double Slope = ComputeUnorientedSlope(NodeToRemovePoint, PreviousPoint, NextPoint);
 		bool bRemoveNode = false;
-		if (Slop < 0.1)
+		if (Slope < 0.1)
 		{
 			double SquareDistance1 = SquareDistanceOfPointToSegment(PreviousPoint, NodeToRemovePoint, NextPoint);
 			double SquareDistance2 = SquareDistanceOfPointToSegment(NextPoint, NodeToRemovePoint, PreviousPoint);
diff --git a/Engine/Source/Runtime/Datasmith/CADKernel/Private/CADKernel/Mesh/Structure/ThinZone2D.cpp b/Engine/Source/Runtime/Datasmith/CADKernel/Private/CADKernel/Mesh/Structure/ThinZone2D.cpp
--- a/Engine/Source/Runtime/Datasmith/CADKernel/Private/CADKernel/Mesh/Structure/ThinZone2D.cpp
+++ b/Engine/Source/Runtime/Datasmith/CADKernel/Private/CADKernel/Mesh/Structure/ThinZone2D.cpp
@@ -118,15 +118,15 @@ void FThinZone2DFinder::FindClosedSegments()
 			if (MinSquareThickness > SquareDistance)
 			{
 				// check the angle between segments. As they are opposite, the cosAngle as to be close to -1
-				double Slop = Segment->ComputeUnorientedSlopeOf(Candidate);
-				if (Slop < 3.33) // Angle <5Pi/6 (150 deg)
+				double Slope = Segment->ComputeUnorientedSlopeOf(Candidate);
+				if (Slope < 3.33) // Angle <5Pi/6 (150 deg)
 				{
 					continue;
 				}
 
 				// check the angle between segment and Middle-Projection. As they are opposite, the cosAngle as to be close to 0
-				Slop = Segment->ComputeUnorientedSlopeOf(SegmentMiddle, Projection);
-				if (Slop < 1.33) // Angle <Pi/3 (60 deg)
+				Slope = Segment->ComputeUnorientedSlopeOf(SegmentMiddle, Projection);
+				if (Slope < 1.33) // Angle <Pi/3 (60 deg)
 				{
 					continue;
 				}
diff --git a/Engine/Source/Runtime/Datasmith/CADKernel/Private/CADKernel/Topo/TopologicalEdge.cpp b/Engine/Source/Runtime/Datasmith/CADKernel/Private/CADKernel/Topo/TopologicalEdge.cpp
--- a/Engine/Source/Runtime/Datasmith/CADKernel/Private/CADKernel/Topo/TopologicalEdge.cpp
+++ b/Engine/Source/Runtime/Datasmith/CADKernel/Private/CADKernel/Topo/TopologicalEdge.cpp
@@ -1008,12 +1008,12 @@ void FTopologicalEdge::ComputeEdge2DProperties(FEdge2DProperties& EdgeCharacteri
 
 	for (int32 Index = StartIndex; Index <= EndIndex; ++Index)
 	{
-		double Slop = ComputeUnorientedSlope(Polyline2D[Index], Polyline2D[Index + 1], 0);
-		if (Slop > 2.)
+		double Slope = ComputeUnorientedSlope(Polyline2D[Index], Polyline2D[Index + 1], 0);
+		if (Slope > 2.)
 		{
-			Slop = 4. - Slop;
+			Slope = 4. - Slope;
 		}
-		EdgeCharacteristics.Add(Slop, Polyline3D[Index].Distance(Polyline3D[Index + 1]));
+		EdgeCharacteristics.Add(Slope, Polyline3D[Index].Distance(Polyline3D[Index + 1]));
 	}
 }
 
diff --git a/Engine/Source/Runtime/Datasmith/CADKernel/Private/CADKernel/Topo/TopologicalLoop.cpp b/Engine/Source/Runtime/Datasmith/CADKernel/Private/CADKernel/Topo/TopologicalLoop.cpp
--- a/Engine/Source/Runtime/Datasmith/CADKernel/Private/CADKernel/Topo/TopologicalLoop.cpp
+++ b/Engine/Source/Runtime/Datasmith/CADKernel/Private/CADKernel/Topo/TopologicalLoop.cpp
@@ -334,9 +334,9 @@ bool FTopologicalLoop::Orient()
 		// Pic case: the slop is compute between previous and next segment of the extrema 
 		//     if the slop is closed to 0 or 8, it could be pick in self intersecting. We try to recompute the slop a the closed point
 		//          if the slop is still closed to 0 or 8, the pic is not used
-		double Slop = ComputePositiveSlope(LoopSampling[Index], LoopSampling[NextIndex], LoopSampling[PreviousIndex]);
+		double Slope = ComputePositiveSlope(LoopSampling[Index], LoopSampling[NextIndex], LoopSampling[PreviousIndex]);
 
-		if(Slop > 7.9 || Slop < 0.1)
+		if(Slope > 7.9 || Slope < 0.1)
 		{
 			double SquareLengthBefore = LoopSampling[Index].SquareDistance(LoopSampling[PreviousIndex]);
 			double SquareLengthAfter = LoopSampling[Index].SquareDistance(LoopSampling[NextIndex]);
@@ -355,12 +355,12 @@ bool FTopologicalLoop::Orient()
 				}
 			}
 
-			Slop = ComputePositiveSlope(LoopSampling[Index], LoopSampling[NextIndex], LoopSampling[PreviousIndex]);
+			Slope = ComputePositiveSlope(LoopSampling[Index], LoopSampling[NextIndex], LoopSampling[PreviousIndex]);
 
 #ifdef DEBUG_ORIENT
 			if (bDisplayDebug)
 			{
-				F3DDebugSession _(*FString::Printf(TEXT("Fix Pic Node %f"), Slop));
+				F3DDebugSession _(*FString::Printf(TEXT("Fix Pic Node %f"), Slope));
 				DisplayPoint(LoopSampling[Index], EVisuProperty::BluePoint, Index);
 				DisplaySegment(LoopSampling[PreviousIndex], LoopSampling[Index], EVisuProperty::GreenCurve);
 				DisplaySegment(LoopSampling[NextIndex], LoopSampling[Index], EVisuProperty::GreenCurve);
@@ -369,46 +369,46 @@ bool FTopologicalLoop::Orient()
 		}
 
 
-		if (Slop > 7.9 || Slop < 0.1)
+		if (Slope > 7.9 || Slope < 0.1)
 		{
 			UndefinedOrientationCount++;
 		}
-		else if (Slop > 4.2)
+		else if (Slope > 4.2)
 		{
 			WrongOrientationCount++;
 		}
-		else if (Slop < 3.8)
+		else if (Slope < 3.8)
 		{
 			GoodOrientationCount++;
 		}
 		else
 		{
 			// Extrema case: the slop is compute between the next segment and the nearest BBox side 
-			double ReferenceSlop = 0;
+			double ReferenceSlope = 0;
 			if (FMath::IsNearlyEqual(LoopSampling[Index].U, UMin))
 			{
-				ReferenceSlop = 6;
+				ReferenceSlope = 6;
 			}
 			else if (FMath::IsNearlyEqual(LoopSampling[Index].U, UMax))
 			{
-				ReferenceSlop = 2;
+				ReferenceSlope = 2;
 			}
 			else if (FMath::IsNearlyEqual(LoopSampling[Index].V, VMin))
 			{
-				ReferenceSlop = 0;
+				ReferenceSlope = 0;
 			}
 			else if (FMath::IsNearlyEqual(LoopSampling[Index].V, VMax))
 			{
-				ReferenceSlop = 4;
+				ReferenceSlope = 4;
 			}
 
-			Slop = ComputeUnorientedSlope(LoopSampling[Index], LoopSampling[NextIndex], ReferenceSlop);
+			Slope = ComputeUnorientedSlope(LoopSampling[Index], LoopSampling[NextIndex], ReferenceSlope);
 			// slop should be closed to [0, 0.2] or [3.8, 4]
-			if (Slop > 3.8)
+			if (Slope > 3.8)
 			{
 				WrongOrientationCount++;
 			}
-			else if (Slop < 0.2)
+			else if (Slope < 0.2)
 			{
 				GoodOrientationCount++;
 			}
@@ -427,7 +427,7 @@ bool FTopologicalLoop::Orient()
 						DisplaySegment(LoopSampling[NextIndex], LoopSampling[Index], EVisuProperty::YellowCurve);
 					}
 					{
-						F3DDebugSession _(*FString::Printf(TEXT("Node %f"), Slop));
+						F3DDebugSession _(*FString::Printf(TEXT("Node %f"), Slope));
 						DisplayPoint(LoopSampling[Index], EVisuProperty::RedPoint, Index);
 					}
 					Wait();
@@ -454,7 +454,7 @@ bool FTopologicalLoop::Orient()
 				DisplaySegment(LoopSampling[PreviousIndex], LoopSampling[Index], EVisuProperty::YellowCurve);
 			}
 			{
-				F3DDebugSession _(*FString::Printf(TEXT("Node %f"), Slop));
+				F3DDebugSession _(*FString::Printf(TEXT("Node %f"), Slope));
 				DisplayPoint(LoopSampling[Index], EVisuProperty::RedPoint, Index);
 			}
 			Wait();
@@ -822,11 +822,11 @@ void FTopologicalLoop::RemoveDegeneratedEdges()
 		FPoint2D& DegeneratedEdgeExtremity = (bPrevious == (DegeneratedOrientedEdge.Direction == EOrientation::Front)) ? DegeneratedEdgeExtremities[0].Point2D : DegeneratedEdgeExtremities[1].Point2D;
 		FPoint2D& OtherDegeneratedEdgeExtremity = (bPrevious == (DegeneratedOrientedEdge.Direction == EOrientation::Front)) ? DegeneratedEdgeExtremities[1].Point2D : DegeneratedEdgeExtremities[0].Point2D;
 
-		double Slop = ComputeUnorientedSlope(FPoint2D::ZeroPoint, DegeneratedEdgeTangent, NearEdgeTangent);
-		constexpr double FlatSlop = 3.9; // ~175 deg. 
+		double Slope = ComputeUnorientedSlope(FPoint2D::ZeroPoint, DegeneratedEdgeTangent, NearEdgeTangent);
+		constexpr double FlatSlope = 3.9; // ~175 deg. 
 		// a degenerated edge has to be really tangent with its neighbor edge to be merge with it
 		// otherwise it could distort the neighbor edge.
-		if (Slop > FlatSlop)
+		if (Slope > FlatSlope)
 		{
 			NearOrientedEdge.Entity->ExtendTo((bPrevious == (NearOrientedEdge.Direction != EOrientation::Front)), OtherDegeneratedEdgeExtremity, OtherDegeneratedEdgeVertex);
 			DegeneratedEdge->Delete();
diff --git a/Engine/Source/Runtime/Datasmith/CADKernel/Private/CADKernel/Topo/TopologicalVertex.cpp b/Engine/Source/Runtime/Datasmith/CADKernel/Private/CADKernel/Topo/TopologicalVertex.cpp
--- a/Engine/Source/Runtime/Datasmith/CADKernel/Private/CADKernel/Topo/TopologicalVertex.cpp
+++ b/Engine/Source/Runtime/Datasmith/CADKernel/Private/CADKernel/Topo/TopologicalVertex.cpp
@@ -154,22 +154,25 @@ void FVertexLink::DefineActiveEntity()
 		return;
 	}
 
-	double DistanceSquare = HUGE_VALUE;
+	double MinDistanceSquare = HUGE_VALUE;
 	FTopologicalVertex* ClosedVertex = TwinEntities.HeapTop();
 	for (FTopologicalVertex* Vertex : TwinEntities)
 	{
 		ensureCADKernel(!Vertex->IsDeleted());
 
 		double Square = Vertex->SquareDistance(Barycenter);
-		if (Square < DistanceSquare)
+		if (Square < MinDistanceSquare)
 		{
-			DistanceSquare = Square;
+			MinDistanceSquare = Square;
 			ClosedVertex = Vertex;
 			if (FMath::IsNearlyZero(Square))
 			{
 				break;
 			}
 		}
+#ifdef CADKERNEL_DEV
+		ensureCADKernel(Square < 0.01);
+#endif
 	}
 	ActiveEntity = ClosedVertex;
 }
diff --git a/Engine/Source/Runtime/Datasmith/CADKernel/Public/CADKernel/Topo/TopologicalEdge.h b/Engine/Source/Runtime/Datasmith/CADKernel/Public/CADKernel/Topo/TopologicalEdge.h
--- a/Engine/Source/Runtime/Datasmith/CADKernel/Public/CADKernel/Topo/TopologicalEdge.h
+++ b/Engine/Source/Runtime/Datasmith/CADKernel/Public/CADKernel/Topo/TopologicalEdge.h
@@ -653,7 +653,7 @@ public:
 	/**
 	 * Compute the edge 2D properties i.e. the mean and standard deviation of the slop of the edge in the parametric space of the carrier surface
 	 */
-	void ComputeEdge2DProperties(FEdge2DProperties& SlopCharacteristics);
+	void ComputeEdge2DProperties(FEdge2DProperties& SlopeCharacteristics);
 
 	void GetExtremities(FSurfacicCurveExtremities& Extremities) const
 	{
@@ -753,17 +753,17 @@ public:
 struct CADKERNEL_API FEdge2DProperties
 {
 	double StandardDeviation = 0;
-	double MediumSlop = 0;
+	double MediumSlope = 0;
 	double Length3D = 0;
 	EIso IsoType = EIso::UndefinedIso;
 	bool bIsMesh = false;
 	double MeshedLength = 0;
 
-	void Add(double InSlop, double InLength)
+	void Add(double InSlope, double InLength)
 	{
-		double Temp = InSlop * InLength;
-		MediumSlop += Temp;
-		Temp *= InSlop;
+		double Temp = InSlope * InLength;
+		MediumSlope += Temp;
+		Temp *= InSlope;
 		StandardDeviation += Temp;
 		Length3D += InLength;
 	}
@@ -771,13 +771,13 @@ struct CADKERNEL_API FEdge2DProperties
 	// Finalize has been done on each Property
 	void Add2(FEdge2DProperties& Property)
 	{
-		StandardDeviation = (FMath::Square(StandardDeviation) + FMath::Square(MediumSlop)) * Length3D + (FMath::Square(Property.StandardDeviation) + FMath::Square(Property.MediumSlop)) * Property.Length3D;
-		MediumSlop = MediumSlop * Length3D + Property.MediumSlop * Property.Length3D;
+		StandardDeviation = (FMath::Square(StandardDeviation) + FMath::Square(MediumSlope)) * Length3D + (FMath::Square(Property.StandardDeviation) + FMath::Square(Property.MediumSlope)) * Property.Length3D;
+		MediumSlope = MediumSlope * Length3D + Property.MediumSlope * Property.Length3D;
 		Length3D += Property.Length3D;
 
-		MediumSlop /= Length3D;
+		MediumSlope /= Length3D;
 		StandardDeviation /= Length3D;
-		StandardDeviation -= FMath::Square(MediumSlop);
+		StandardDeviation -= FMath::Square(MediumSlope);
 		StandardDeviation = sqrt(StandardDeviation);
 	}
 
@@ -785,15 +785,15 @@ struct CADKERNEL_API FEdge2DProperties
 	void Add(FEdge2DProperties& Property)
 	{
 		StandardDeviation += Property.StandardDeviation;
-		MediumSlop += Property.MediumSlop;
+		MediumSlope += Property.MediumSlope;
 		Length3D += Property.Length3D;
 	}
 
 	void Finalize()
 	{
-		MediumSlop /= Length3D;
+		MediumSlope /= Length3D;
 		StandardDeviation /= Length3D;
-		StandardDeviation -= FMath::Square(MediumSlop);
+		StandardDeviation -= FMath::Square(MediumSlope);
 		if (StandardDeviation < 0)
 		{
 			StandardDeviation = 0;
@@ -804,14 +804,14 @@ struct CADKERNEL_API FEdge2DProperties
 		}
 		IsoType = EIso::UndefinedIso;
 
-		if (MediumSlop < 0.2)
+		if (MediumSlope < 0.2)
 		{
 			if (StandardDeviation < 0.1)
 			{
 				IsoType = EIso::IsoU;
 			}
 		}
-		else if (MediumSlop > 1.8)
+		else if (MediumSlope > 1.8)
 		{
 			if (StandardDeviation < 0.1)
 			{
