- Improve IntersectionSegmentTool

diff --git a/Engine/Source/Runtime/Datasmith/CADKernel/Private/CADKernel/Mesh/Meshers/IsoTriangulator/IntersectionSegmentTool.cpp b/Engine/Source/Runtime/Datasmith/CADKernel/Private/CADKernel/Mesh/Meshers/IsoTriangulator/IntersectionSegmentTool.cpp
--- a/Engine/Source/Runtime/Datasmith/CADKernel/Private/CADKernel/Mesh/Meshers/IsoTriangulator/IntersectionSegmentTool.cpp
+++ b/Engine/Source/Runtime/Datasmith/CADKernel/Private/CADKernel/Mesh/Meshers/IsoTriangulator/IntersectionSegmentTool.cpp
@@ -11,159 +11,63 @@
 namespace UE::CADKernel
 { 
 
-const FIsoSegment* FIntersectionSegmentTool::DoesIntersect(const FIsoSegment& Segment) const
+const FIsoSegment* FIntersectionSegmentTool::FindIntersectingSegment(const FIsoSegment& Segment) const
 {
-	return DoesIntersect(Segment.GetFirstNode(), Segment.GetSecondNode());
+	using namespace IntersectionSegmentTool;
+	return TIntersectionSegmentTool<IntersectionSegmentTool::FSegment>::FindIntersectingSegment(&Segment.GetFirstNode(), &Segment.GetSecondNode());
 }
 
-FIsoSegment* FIntersectionSegmentTool::DoesIntersect(const FIsoSegment& Segment)
+FIsoSegment* FIntersectionSegmentTool::FindIntersectingSegment(const FIsoSegment& Segment)
 {
-	return const_cast<FIsoSegment*> (static_cast<const FIntersectionSegmentTool*>(this)->DoesIntersect(Segment.GetFirstNode(), Segment.GetSecondNode()));
+	return const_cast<FIsoSegment*> (static_cast<const FIntersectionSegmentTool*>(this)->TIntersectionSegmentTool<IntersectionSegmentTool::FSegment>::FindIntersectingSegment(&Segment.GetFirstNode(), &Segment.GetSecondNode()));
 }
 
-bool FIntersectionSegmentTool::FindIntersections(const FIsoNode& StartNode, const FIsoNode& EndNode, TArray<const FIsoSegment*>& OutIntersectedSegments) const
+IntersectionSegmentTool::FSegment::FSegment(const FGrid& Grid, const FIsoNode& StartNode, const FIsoNode& EndNode)
+	: IntersectionToolBase::FSegment(StartNode.Get2DPoint(EGridSpace::UniformScaled, Grid), EndNode.Get2DPoint(EGridSpace::UniformScaled, Grid))
+	, IsoSegment(nullptr)
 {
-	const FPoint2D& StartPoint = StartNode.Get2DPoint(EGridSpace::UniformScaled, Grid);
-	const FPoint2D& EndPoint = EndNode.Get2DPoint(EGridSpace::UniformScaled, Grid);
-	const FSegment4IntersectionTools StartEndSegment(StartPoint, EndPoint);
-
-	OutIntersectedSegments.Empty(10);
-
-	for (const FSegment4IntersectionTools& Segment : Segments)
-	{
-		if (!Segment.IsoSegment)
-		{
-			continue;
-		}
-
-		if (&Segment.IsoSegment->GetFirstNode() == &StartNode || &Segment.IsoSegment->GetSecondNode() == &StartNode)
-		{
-			continue;
-		}
-
-		if (&Segment.IsoSegment->GetFirstNode() == &EndNode || &Segment.IsoSegment->GetSecondNode() == &EndNode)
-		{
-			continue;
-		}
-
-		if (IntersectSegments2D(Segment.Segment2D, StartEndSegment.Segment2D))
-		{
-			OutIntersectedSegments.Add(Segment.IsoSegment);
-		}
-	}
-
-	return OutIntersectedSegments.Num() > 0;
 }
 
-const FIsoSegment* FIntersectionSegmentTool::DoesIntersect(const FIsoNode& StartNode, const FIsoNode& EndNode) const
+IntersectionSegmentTool::FSegment::FSegment(const FGrid& Grid, const FIsoNode& StartNode, const FPoint2D& EndPoint)
+	: IntersectionToolBase::FSegment(StartNode.Get2DPoint(EGridSpace::UniformScaled, Grid), EndPoint)
+	, IsoSegment(nullptr)
 {
-	const FPoint2D& StartPoint = StartNode.Get2DPoint(EGridSpace::UniformScaled, Grid);
-	const FPoint2D& EndPoint = EndNode.Get2DPoint(EGridSpace::UniformScaled, Grid);
-	const FSegment4IntersectionTools StartEndSegment(StartPoint, EndPoint);
-
-	double AxisMax = StartEndSegment.Boundary[EIso::IsoU].Max + StartEndSegment.Boundary[EIso::IsoV].Max;
-
-	for (const FSegment4IntersectionTools& Segment : Segments)
-	{
-		if (!Segment.IsoSegment || Segment.IsoSegment->IsDelete())
-		{
-			continue;
-		}
-
-		if (bSegmentsAreSorted && AxisMax < Segment.AxisMin)
-		{
-			break;
-		}
-
-		if (&Segment.IsoSegment->GetFirstNode() == &StartNode || &Segment.IsoSegment->GetSecondNode() == &StartNode)
-		{
-			continue;
-		}
-
-		if (&Segment.IsoSegment->GetFirstNode() == &EndNode || &Segment.IsoSegment->GetSecondNode() == &EndNode)
-		{
-			continue;
-		}
-
-		if (IntersectSegments2D(Segment.Segment2D, StartEndSegment.Segment2D))
-		{
-			//DisplaySegment(Segment.Segment2D.Point0, Segment.Segment2D.Point1, 0, RedCurve);
-			return Segment.IsoSegment;
-		}
-	}
-
-	return nullptr;
 }
 
-bool FIntersectionSegmentTool::DoesIntersect(const FIsoNode& StartNode, const FPoint2D& EndPoint) const
+IntersectionSegmentTool::FSegment::FSegment(const FGrid& Grid, const FIsoSegment& InSegment)
+	: IntersectionToolBase::FSegment(InSegment.GetFirstNode().Get2DPoint(EGridSpace::UniformScaled, Grid), InSegment.GetSecondNode().Get2DPoint(EGridSpace::UniformScaled, Grid))
+	, IsoSegment(&InSegment)
 {
-	const FPoint2D& StartPoint = StartNode.Get2DPoint(EGridSpace::UniformScaled, Grid);
-	const FSegment4IntersectionTools StartEndSegment(StartPoint, EndPoint);
-
-	double DMax = StartEndSegment.Boundary[EIso::IsoU].Max + StartEndSegment.Boundary[EIso::IsoV].Max;
-
-	for (const FSegment4IntersectionTools& Segment : Segments)
-	{
-		if (!Segment.IsoSegment)
-		{
-			continue;
-		}
-
-		if (bSegmentsAreSorted && DMax < Segment.AxisMin)
-		{
-			break;
-		}
-
-		if (&Segment.IsoSegment->GetFirstNode() == &StartNode || &Segment.IsoSegment->GetSecondNode() == &StartNode)
-		{
-			continue;
-		}
-
-		if (!StartEndSegment.CouldItIntersect(Segment))
-		{
-			continue;
-		}
-
-		if (IntersectSegments2D(Segment.Segment2D, StartEndSegment.Segment2D))
-		{
-			return true;
-		}
-	}
+}
 
-	return false;
+IntersectionNodePairTool::FSegment::FSegment(const FIsoNode* InStartNode, const FIsoNode* InEndNode, const FPoint2D& InStartPoint, const FPoint2D& InEndPoint)
+	: IntersectionToolBase::FSegment(InStartPoint, InEndPoint)
+	, StartNode(InStartNode)
+	, EndNode(InEndNode)
+{
 }
 
-bool FIntersectionSegmentTool::DoesIntersect(const FPoint2D& StartPoint, const FPoint2D& EndPoint) const
+IntersectionNodePairTool::FSegment::FSegment(const FGrid& Grid, const FIsoNode& InStartNode, const FIsoNode& InEndNode)
+	: IntersectionToolBase::FSegment(InStartNode.Get2DPoint(EGridSpace::UniformScaled, Grid), InEndNode.Get2DPoint(EGridSpace::UniformScaled, Grid))
+	, StartNode(&InStartNode)
+	, EndNode(&InEndNode)
 {
-	const FSegment4IntersectionTools StartEndSegment(StartPoint, EndPoint);
-	double DMax = StartEndSegment.Boundary[EIso::IsoU].Max + StartEndSegment.Boundary[EIso::IsoV].Max;
+}
 
-	for (const FSegment4IntersectionTools& Segment : Segments)
-	{
-		if (bSegmentsAreSorted && DMax < Segment.AxisMin)
-		{
-			break;
-		}
+bool IntersectionSegmentTool::FSegment::IsValid() const
+{
+	return IsoSegment != nullptr && !IsoSegment->IsDelete();
+}
 
-		if (!StartEndSegment.CouldItIntersect(Segment))
-		{
-			continue;
-		}
 
-		if (IntersectSegments2D(Segment.Segment2D, StartEndSegment.Segment2D))
-		{
-			return true;
-		}
-	}
-	return false;
+const FIsoNode* IntersectionSegmentTool::FSegment::GetFirstNode() const
+{
+	return &IsoSegment->GetFirstNode();
 }
 
-FSegment4IntersectionTools::FSegment4IntersectionTools(const FGrid& Grid, const FIsoSegment& InSegment)
-	: Segment2D(InSegment.GetFirstNode().Get2DPoint(EGridSpace::UniformScaled, Grid), InSegment.GetSecondNode().Get2DPoint(EGridSpace::UniformScaled, Grid))
-	, Boundary(Segment2D.Point0, Segment2D.Point1)
-	, IsoSegment(&InSegment)
+const FIsoNode* IntersectionSegmentTool::FSegment::GetSecondNode() const
 {
-	AxisMin = Boundary[EIso::IsoU].Min + Boundary[EIso::IsoV].Min;
+	return &IsoSegment->GetSecondNode();
 }
 
 } // namespace UE::CADKernel
diff --git a/Engine/Source/Runtime/Datasmith/CADKernel/Private/CADKernel/Mesh/Meshers/IsoTriangulator/IntersectionSegmentTool.h b/Engine/Source/Runtime/Datasmith/CADKernel/Private/CADKernel/Mesh/Meshers/IsoTriangulator/IntersectionSegmentTool.h
--- a/Engine/Source/Runtime/Datasmith/CADKernel/Private/CADKernel/Mesh/Meshers/IsoTriangulator/IntersectionSegmentTool.h
+++ b/Engine/Source/Runtime/Datasmith/CADKernel/Private/CADKernel/Mesh/Meshers/IsoTriangulator/IntersectionSegmentTool.h
@@ -11,10 +11,13 @@ namespace UE::CADKernel
 {
 class FGrid;
 class FIsoNode;
+class FIsoInnerNode;
 class FIsoSegment;
 class FPoint2D;
 
-struct FSegment4IntersectionTools
+namespace IntersectionToolBase
+{
+struct FSegment
 {
 	const TSegment<FPoint2D> Segment2D;
 
@@ -23,58 +26,180 @@ struct FSegment4IntersectionTools
 	 */
 	const FSurfacicBoundary Boundary;
 
-	const FIsoSegment* IsoSegment;
-
 	/**
 	 * Uses as criterion to sort segments to optimize
 	 * AxisMin = Boundary[EIso::IsoU].Min + Boundary[EIso::IsoV].Min
 	 */
 	double AxisMin;
-
-	FSegment4IntersectionTools(const FGrid& Grid, const FIsoSegment& InSegment);
+	double AxisMax;
 
 	/**
 	 * WARNING StartPoint, EndPoint must be defined in EGridSpace::UniformScaled
 	 */
-	FSegment4IntersectionTools(const FPoint2D& StartPoint, const FPoint2D& EndPoint)
+	FSegment(const FPoint2D& StartPoint, const FPoint2D& EndPoint)
 		: Segment2D(StartPoint, EndPoint)
 		, Boundary(StartPoint, EndPoint)
-		, IsoSegment(nullptr)
 	{
 		AxisMin = Boundary[EIso::IsoU].Min + Boundary[EIso::IsoV].Min;
+		AxisMax = Boundary[EIso::IsoU].Max + Boundary[EIso::IsoV].Max;
 	}
 
-	/**
-	 * @return false if their Boundaries are not intersecting
-	 */
-	bool CouldItIntersect(const FSegment4IntersectionTools& Segment) const
+	virtual ~FSegment() = default;
+
+	virtual bool IsValid() const = 0;
+
+	virtual const FIsoNode* GetFirstNode() const = 0;
+	virtual const FIsoNode* GetSecondNode() const = 0;
+
+	virtual const FIsoSegment* GetIsoSegment() const 
+	{
+		return nullptr;
+	}
+
+	bool DoesItStartFrom(const FIsoNode* StartNode, const FIsoNode* EndNode) const
 	{
-		if ((Boundary[EIso::IsoU].Min > Segment.Boundary[EIso::IsoU].Max) || (Boundary[EIso::IsoV].Min > Segment.Boundary[EIso::IsoV].Max))
+		if (GetFirstNode() == StartNode || GetSecondNode() == StartNode)
+		{
+			return true;
+		}
+
+		if (GetFirstNode() == EndNode || GetSecondNode() == EndNode)
+		{
+			return true;
+		}
+
+		return false;
+	}
+
+	bool DoesItStartFrom(const FIsoNode* StartNode, const FPoint2D* EndPoint) const
+	{
+		if (GetFirstNode() == StartNode || GetSecondNode() == StartNode)
+		{
+			return true;
+		}
+
+		return false;
+	}
+
+	bool DoesItStartFrom(const FPoint2D* StartPoint, const FPoint2D* EndPoint) const
+	{
+		return false;
+	}
+
+	bool IsFullyBefore(const FSegment& Segment) const
+	{
+		return AxisMax < Segment.AxisMin;
+	}
+
+	bool IsFullyAfter(const FSegment& Segment) const
+	{
+		return Segment.AxisMax < AxisMin;
+	}
+
+	bool CouldItIntersect(const FSurfacicBoundary& SegmentBoundary) const
+	{
+		if ((Boundary[EIso::IsoU].Min > SegmentBoundary[EIso::IsoU].Max) || (Boundary[EIso::IsoV].Min > SegmentBoundary[EIso::IsoV].Max))
 		{
 			return false;
 		}
 
-		if ((Boundary[EIso::IsoU].Max < Segment.Boundary[EIso::IsoU].Min) || (Boundary[EIso::IsoV].Max < Segment.Boundary[EIso::IsoV].Min))
+		if ((Boundary[EIso::IsoU].Max < SegmentBoundary[EIso::IsoU].Min) || (Boundary[EIso::IsoV].Max < SegmentBoundary[EIso::IsoV].Min))
 		{
 			return false;
 		}
 		return true;
 	}
+
+	bool DoesItIntersect(const FSegment& Segment) const
+	{
+		if (!CouldItIntersect(Segment.Boundary))
+		{
+			return false;
+		}
+
+		return IntersectSegments2D(Segment2D, Segment.Segment2D);
+	}
 };
+}
 
-class FIntersectionSegmentTool
+namespace IntersectionSegmentTool
 {
-	TArray<FSegment4IntersectionTools> Segments;
+
+struct FSegment : public IntersectionToolBase::FSegment
+{
+	const FIsoSegment* IsoSegment;
+
+	FSegment(const FGrid& Grid, const FIsoSegment& InSegment);
+	FSegment(const FGrid& Grid, const FIsoNode& StartNode, const FIsoNode& EndNode);
+	FSegment(const FGrid& Grid, const FIsoNode& StartNode, const FPoint2D& EndPoint);
+	FSegment(const FGrid& Grid, const FPoint2D& StartPoint, const FPoint2D& EndPoint)
+		: IntersectionToolBase::FSegment(StartPoint, EndPoint)
+		, IsoSegment(nullptr)
+	{
+	}
+
+	virtual bool IsValid() const override;
+
+	virtual const FIsoNode* GetFirstNode() const override;
+	virtual const FIsoNode* GetSecondNode() const override;
+
+	virtual const FIsoSegment* GetIsoSegment() const override
+	{
+		return IsoSegment;
+	}
+
+};
+
+}
+
+namespace IntersectionNodePairTool
+{
+
+struct FSegment : public IntersectionToolBase::FSegment
+{
+	const FIsoNode* StartNode;
+	const FIsoNode* EndNode;
+
+	FSegment(const FIsoNode* StartNode, const FIsoNode* EndNode, const FPoint2D& StartPoint, const FPoint2D& EndPoint);
+	FSegment(const FGrid& Grid, const FIsoNode& StartNode, const FIsoNode& EndNode);
+
+	virtual bool IsValid() const override
+	{
+		return true;
+	}
+
+	virtual const FIsoNode* GetFirstNode() const override
+	{
+		return StartNode;
+	}
+
+	virtual const FIsoNode* GetSecondNode() const override
+	{
+		return EndNode;
+	}
+
+};
+
+}
+
+template<typename SegmentType>
+class TIntersectionSegmentTool
+{
+protected:
 	const FGrid& Grid;
+
+	TArray<SegmentType> Segments;
 	bool bSegmentsAreSorted;
 
 public:
-	FIntersectionSegmentTool(const FGrid& InGrid)
+	TIntersectionSegmentTool(const FGrid& InGrid)
 		: Grid(InGrid)
 		, bSegmentsAreSorted(false)
 	{
 	}
 
+	virtual ~TIntersectionSegmentTool() = default;
+
 	void Empty(int32 InMaxNum)
 	{
 		Segments.Empty(InMaxNum);
@@ -96,38 +221,142 @@ public:
 	void SetCount(int32 NewCount)
 	{
 		ensureCADKernel(NewCount < Segments.Num() && !bSegmentsAreSorted);
-		while (NewCount != Segments.Num())
-		{
-			Segments.RemoveAt(Segments.Num() - 1);
-		}
+		Segments.SetNum(NewCount);
+	}
+
+	int32 Count()
+	{
+		return Segments.Num();
 	}
 
 	/**
-	 * @return true if the segment is found and removed
+	 * segments are sorted by DMin increasing
 	 */
-	bool Remove(const FIsoSegment* Segment)
+	void Sort()
+	{
+		Algo::Sort(Segments, [](const SegmentType& Segment1, const SegmentType& Segment2) { return Segment1.AxisMin < Segment2.AxisMin; });
+		bSegmentsAreSorted = true;
+	}
+
+	template<typename ExtremityType1, typename ExtremityType2>
+	const FIsoSegment* FindIntersectingSegment(const ExtremityType1* StartExtremity, const ExtremityType2* EndExtremity) const
 	{
-		int32 SegmentIndex = Segments.IndexOfByPredicate([Segment](const FSegment4IntersectionTools& SegmentIter)
+		using namespace IntersectionSegmentTool;
+		FSegment InSegment(Grid, *StartExtremity, *EndExtremity);
+
+		for (const FSegment& Segment : Segments)
+		{
+			if (!Segment.IsValid())
 			{
-				return SegmentIter.IsoSegment == Segment;
-			});
+				continue;
+			}
 
-		if (SegmentIndex != INDEX_NONE)
+			if (bSegmentsAreSorted)
+			{
+				if (Segment.IsFullyBefore(InSegment))
+				{
+					continue;
+				}
+
+				if (Segment.IsFullyAfter(InSegment))
+				{
+					break;
+				}
+}
+
+			if (Segment.DoesItStartFrom(StartExtremity, EndExtremity))
+			{
+				continue;
+			}
+
+			if (Segment.DoesItIntersect(InSegment))
+			{
+				return Segment.IsoSegment;
+			}
+		}
+
+		return nullptr;
+	}
+
+	template<typename ExtremityType1, typename ExtremityType2>
+	int32 FindIntersectingSegments(const ExtremityType1* StartExtremity, const ExtremityType2* EndExtremity, TArray<const FIsoSegment*>* OutIntersectedSegments) const
+	{
+		SegmentType InSegment(Grid, *StartExtremity, *EndExtremity);
+		if(OutIntersectedSegments)
 		{
-			Segments.RemoveAt(SegmentIndex);
-			return true;
+			OutIntersectedSegments->Empty(10);
 		}
-		return false;
+
+		int32 IntersectionCount = 0;
+		for (const SegmentType& Segment : Segments)
+		{
+			if (!Segment.IsValid())
+			{
+				continue;
+			}
+
+			if (bSegmentsAreSorted)
+			{
+				if (Segment.IsFullyBefore(InSegment))
+				{
+					continue;
+				}
+
+				if (Segment.IsFullyAfter(InSegment))
+				{
+					break;
+				}
+}
+
+			if (Segment.DoesItStartFrom(StartExtremity, EndExtremity))
+			{
+				continue;
+			}
+
+			if (Segment.DoesItIntersect(InSegment))
+			{
+				++IntersectionCount;
+				if (OutIntersectedSegments)
+				{
+					OutIntersectedSegments->Add(Segment.GetIsoSegment());
+				}
+			}
+		}
+
+		return IntersectionCount;
 	}
 
-	int32 Count()
+#ifdef CADKERNEL_DEV
+	virtual void Display(bool bDisplay, const TCHAR* Message, EVisuProperty Property = EVisuProperty::BlueCurve) const
+	{
+		if (!bDisplay)
+		{
+			return;
+		}
+
+		int32 Index = 0;
+		Open3DDebugSession(Message);
+		for (const SegmentType& Segment : Segments)
+		{
+			DisplaySegment(Segment.Segment2D[0] * DisplayScale, Segment.Segment2D[1] * DisplayScale, Index++, Property);
+		}
+		Close3DDebugSession();
+	}
+#endif
+
+};
+
+class FIntersectionSegmentTool : public TIntersectionSegmentTool<IntersectionSegmentTool::FSegment>
+{
+public:
+	FIntersectionSegmentTool(const FGrid& InGrid)
+		: TIntersectionSegmentTool<IntersectionSegmentTool::FSegment>(InGrid)
 	{
-		return Segments.Num();
 	}
 
 	bool Update(const FIsoSegment* Segment)
 	{
-		int32 SegmentIndex = Segments.IndexOfByPredicate([Segment](const FSegment4IntersectionTools& SegmentIter)
+		int32 SegmentIndex = Segments.IndexOfByPredicate([Segment](const IntersectionSegmentTool::FSegment& SegmentIter)
 			{
 				return SegmentIter.IsoSegment == Segment;
 			});
@@ -140,7 +369,6 @@ public:
 			return true;
 		}
 		return false;
-
 	}
 
 	void AddSegments(FIsoSegment** InNewSegments, int32 Count)
@@ -173,43 +401,99 @@ public:
 	void AddSegment(const FPoint2D& StartPoint, const FPoint2D& EndPoint)
 	{
 		bSegmentsAreSorted = false;
-		Segments.Emplace(StartPoint, EndPoint);
+		Segments.Emplace(Grid, StartPoint, EndPoint);
 	}
 
-	const FIsoSegment* DoesIntersect(const FIsoSegment& Segment) const;
-	FIsoSegment* DoesIntersect(const FIsoSegment& Segment);
+	/**
+	 * @return true if the segment is found and removed
+	 */
+	bool Remove(const FIsoSegment* Segment)
+	{
+		int32 SegmentIndex = Segments.IndexOfByPredicate([Segment](const IntersectionSegmentTool::FSegment& SegmentIter)
+			{
+				return SegmentIter.IsoSegment == Segment;
+			});
+
+		if (SegmentIndex != INDEX_NONE)
+		{
+			Segments.RemoveAt(SegmentIndex);
+			return true;
+		}
+		return false;
+	}
+
+	const FIsoSegment* FindIntersectingSegment(const FIsoSegment& Segment) const;
+	FIsoSegment* FindIntersectingSegment(const FIsoSegment& Segment);
+
+	bool DoesIntersect(const FIsoSegment& Segment) const
+	{
+		return FindIntersectingSegment(Segment) != nullptr;
+	}
 
 	/**
 	 * WARNING StartPoint, EndPoint must be defined in EGridSpace::UniformScaled
 	 */
-	bool DoesIntersect(const FPoint2D& StartPoint, const FPoint2D& EndPoint) const;
+	bool DoesIntersect(const FPoint2D& StartPoint, const FPoint2D& EndPoint) const
+	{
+		return TIntersectionSegmentTool<IntersectionSegmentTool::FSegment>::FindIntersectingSegment(&StartPoint, &EndPoint) != nullptr;
+	}
 
 	/**
 	 * Allow StartNode to be connected to one segment
 	 * WARNING EndPoint must be defined in EGridSpace::UniformScaled
 	 */
-	bool DoesIntersect(const FIsoNode& StartNode, const FPoint2D& EndPoint) const;
+	bool DoesIntersect(const FIsoNode& StartNode, const FPoint2D& EndPoint) const
+	{
+		return TIntersectionSegmentTool<IntersectionSegmentTool::FSegment>::FindIntersectingSegment(&StartNode, &EndPoint) != nullptr;
+	}
+
+	bool DoesIntersect(const FIsoNode& StartNode, const FIsoNode& EndNode) const
+	{
+		return TIntersectionSegmentTool<IntersectionSegmentTool::FSegment>::FindIntersectingSegment(&StartNode, &EndNode) != nullptr;
+	}
 
 	/**
 	 * Allow StartNode and EndNode to be connected to one segment
 	 */
-	const FIsoSegment* DoesIntersect(const FIsoNode& StartNode, const FIsoNode& EndNode) const;
+	const FIsoSegment* FindIntersectingSegment(const FIsoNode& StartNode, const FIsoNode& EndNode) const
+	{
+		return TIntersectionSegmentTool<IntersectionSegmentTool::FSegment>::FindIntersectingSegment(&StartNode, &EndNode);
+	}
 
-	bool FindIntersections(const FIsoNode& StartNode, const FIsoNode& EndNode, TArray<const FIsoSegment*>& OutIntersections) const;
+	bool FindIntersectingSegments(const FIsoNode& StartNode, const FIsoNode& EndNode, TArray<const FIsoSegment*>& OutIntersections) const
+	{
+		return TIntersectionSegmentTool<IntersectionSegmentTool::FSegment>::FindIntersectingSegments(&StartNode, &EndNode, &OutIntersections) > 0;
+	}
 
-	/**
-	 * segments are sorted by DMin increasing
-	 */
-	void Sort()
+};
+
+class FIntersectionNodePairTool : public TIntersectionSegmentTool<IntersectionNodePairTool::FSegment>
+{
+
+public:
+	FIntersectionNodePairTool(const FGrid& InGrid)
+		: TIntersectionSegmentTool<IntersectionNodePairTool::FSegment>(InGrid)
 	{
-		Algo::Sort(Segments, [](const FSegment4IntersectionTools& Segment1, const FSegment4IntersectionTools& Segment2) { return Segment1.AxisMin < Segment2.AxisMin; });
-		bSegmentsAreSorted = true;
 	}
 
-#ifdef CADKERNEL_DEV
-	void Display(const TCHAR* Message, EVisuProperty Property = EVisuProperty::BlueCurve) const;
-#endif
+	void AddSegment(const FIsoNode& StartNode, const FIsoNode& EndNode)
+	{
+		bSegmentsAreSorted = false;
+		Segments.Emplace(Grid, StartNode, EndNode);
+	}
+
+	void AddSegment(const FIsoNode* StartNode, const FIsoNode* EndNode, const FPoint2D& StartPoint, const FPoint2D& EndPoint)
+	{
+		bSegmentsAreSorted = false;
+		Segments.Emplace(StartNode, EndNode, StartPoint, EndPoint);
+	}
+
+	int32 CountIntersections(const FIsoNode& StartNode, const FIsoNode& EndNode) const
+	{
+		return TIntersectionSegmentTool<IntersectionNodePairTool::FSegment>::FindIntersectingSegments(&StartNode, &EndNode, nullptr);
+	}
 
 };
+
 } // namespace UE::CADKernel
 
diff --git a/Engine/Source/Runtime/Datasmith/CADKernel/Private/CADKernel/Mesh/Structure/LoopCleaner.cpp b/Engine/Source/Runtime/Datasmith/CADKernel/Private/CADKernel/Mesh/Structure/LoopCleaner.cpp
--- a/Engine/Source/Runtime/Datasmith/CADKernel/Private/CADKernel/Mesh/Structure/LoopCleaner.cpp
+++ b/Engine/Source/Runtime/Datasmith/CADKernel/Private/CADKernel/Mesh/Structure/LoopCleaner.cpp
@@ -166,7 +166,7 @@ bool FLoopCleaner::UncrossLoops()
 		}
 #endif
 
-		if (const FIsoSegment* IntersectingSegment = LoopSegmentsIntersectionTool.DoesIntersect(Segment))
+		if (const FIsoSegment* IntersectingSegment = LoopSegmentsIntersectionTool.FindIntersectingSegment(Segment))
 		{
 #ifdef DEBUG_UNCROSS_LOOPS
 			if (bDisplay)
@@ -960,7 +960,7 @@ bool FLoopCleaner::RemoveIntersectionsOfSubLoop(int32 IntersectionIndex, int32 I
 						return false;
 					}
 
-					if (const FIsoSegment* Intersection = LoopSegmentsIntersectionTool.DoesIntersect(*Node, *GetNext(Node)))
+					if (const FIsoSegment* Intersection = LoopSegmentsIntersectionTool.FindIntersectingSegment(*Node, *GetNext(Node)))
 					{
 						MoveNodeBehindSegment(*Intersection, *GetNext(Node));
 					}
@@ -1634,7 +1634,7 @@ void FLoopCleaner::FindLoopIntersections()
 			return;
 		}
 
-		if (LoopSegmentsIntersectionTool.FindIntersections(*Node, *NextNode, IntersectedSegments))
+		if (LoopSegmentsIntersectionTool.FindIntersectingSegments(*Node, *NextNode, IntersectedSegments))
 		{
 			for (const FIsoSegment* IntersectedSegment : IntersectedSegments)
 			{
