- Improve Grid's Nodes typing
- Improve build of InnerToOuterSegmentsIntersectionTool

diff --git a/Engine/Source/Runtime/Datasmith/CADKernel/Private/CADKernel/Mesh/Meshers/IsoTriangulator.cpp b/Engine/Source/Runtime/Datasmith/CADKernel/Private/CADKernel/Mesh/Meshers/IsoTriangulator.cpp
--- a/Engine/Source/Runtime/Datasmith/CADKernel/Private/CADKernel/Mesh/Meshers/IsoTriangulator.cpp
+++ b/Engine/Source/Runtime/Datasmith/CADKernel/Private/CADKernel/Mesh/Meshers/IsoTriangulator.cpp
@@ -109,7 +109,7 @@ bool FIsoTriangulator::Triangulate()
 
 #ifdef DEBUG_ISOTRIANGULATOR
 	Grid.DisplayIsoSegments(TEXT("FIsoTrianguler::ThinZoneSegments"), DisplaySpace, ThinZoneSegments, false);
-	LoopSegmentsIntersectionTool.Display(TEXT("FIsoTrianguler::IntersectionTool Loop"));
+	LoopSegmentsIntersectionTool.Display(Grid.bDisplay, TEXT("FIsoTrianguler::IntersectionTool Loop"));
 #endif
 
 	BuildInnerSegments();
@@ -119,7 +119,7 @@ bool FIsoTriangulator::Triangulate()
 		F3DDebugSession A(bDisplay, TEXT("FIsoTrianguler::FinalInnerSegments"));
 		Grid.DisplayIsoSegments(DisplaySpace, FinalInnerSegments, true, false, EVisuProperty::BlueCurve);
 	}
-	InnerToOuterSegmentsIntersectionTool.Display(TEXT("FIsoTrianguler::IntersectionTool InnerToOuter"), EVisuProperty::RedCurve);
+	InnerToOuterSegmentsIntersectionTool.Display(bDisplay, TEXT("FIsoTrianguler::IntersectionTool InnerToOuter"), EVisuProperty::RedCurve);
 	Chronos.TriangulateDuration1 = FChrono::Elapse(StartTime);
 #endif
 
@@ -128,7 +128,7 @@ bool FIsoTriangulator::Triangulate()
 
 	BuildInnerSegmentsIntersectionTool();
 #ifdef DEBUG_ISOTRIANGULATOR
-	InnerSegmentsIntersectionTool.Display(TEXT("FIsoTrianguler::IntersectionTool Inner"));
+	InnerSegmentsIntersectionTool.Display(Grid.bDisplay, TEXT("FIsoTrianguler::IntersectionTool Inner"));
 #endif
 
 #ifdef DEBUG_ISOTRIANGULATOR
@@ -220,7 +220,7 @@ void FIsoTriangulator::BuildNodes()
 	InnerNodeCount = 0;
 	for (int32 Index = 0; Index < (int32)Grid.GetTotalCuttingCount(); ++Index)
 	{
-		if (Grid.IsNodeInsideFace(Index))
+		if (Grid.IsNodeInsideAndMeshable(Index))
 		{
 			FIsoInnerNode& Node = InnerNodes.Emplace_GetRef(Index, FaceIndex++, InnerNodeCount++);
 			GlobalIndexToIsoInnerNodes[Index] = &Node;
@@ -241,7 +241,7 @@ void FIsoTriangulator::FillMeshNodes()
 	InnerNodeCoordinates.Reserve(InnerNodeCount);
 	for (int32 Index = 0; Index < (int32)Grid.GetInner3DPoints().Num(); ++Index)
 	{
-		if (Grid.IsNodeInsideFace(Index))
+		if (Grid.IsNodeInsideAndMeshable(Index))
 		{
 			InnerNodeCoordinates.Emplace(Grid.GetInner3DPoints()[Index]);
 		}
@@ -496,61 +496,79 @@ void FIsoTriangulator::BuildInnerSegments()
 		return false;
 	};
 
+	TFunction<void(const int32, const int32, const ESegmentType)> AddToInnerToOuterSegmentsIntersectionTool = [&](const int32 IndexNode1, const int32 IndexNode2, const ESegmentType InType)
+	{
+		const FPoint2D& Point1 = Grid.GetInner2DPoint(EGridSpace::UniformScaled, IndexNode1);
+		const FPoint2D& Point2 = Grid.GetInner2DPoint(EGridSpace::UniformScaled, IndexNode2);
+
+		const FIsoInnerNode* Node1 = GlobalIndexToIsoInnerNodes[IndexNode1];
+		const FIsoInnerNode* Node2 = GlobalIndexToIsoInnerNodes[IndexNode2];
+
+		InnerToOuterSegmentsIntersectionTool.AddSegment(Node1, Node2, Point1, Point2);
+	};
+
 	TFunction<void(const int32, const int32, const ESegmentType)> BuildSegmentIfValid = [&](const int32 IndexNode1, const int32 IndexNode2, const ESegmentType InType)
 	{
-		if (!Grid.IsNodeInsideFace(IndexNode1) || !Grid.IsNodeInsideFace(IndexNode2))
+		if (Grid.IsNodeOutsideFace(IndexNode1) && Grid.IsNodeOutsideFace(IndexNode2))
 		{
 #ifdef DEBUG_BUILDINNERSEGMENTS
 			if (bDisplay)
 			{
-				DisplaySegment(Grid.GetInner2DPoint(EGridSpace::UniformScaled, IndexNode1) * Grid.DisplayScale, Grid.GetInner2DPoint(EGridSpace::UniformScaled, IndexNode2) * Grid.DisplayScale, IndexNode1, EVisuProperty::BlueCurve);
+				DisplaySegment(Grid.GetInner2DPoint(EGridSpace::UniformScaled, IndexNode1) * DisplayScale, Grid.GetInner2DPoint(EGridSpace::UniformScaled, IndexNode2) * DisplayScale, IndexNode1, EVisuProperty::RedCurve);
 			}
 #endif
-			InnerToOuterSegmentsIntersectionTool.AddSegment(Grid.GetInner2DPoint(EGridSpace::UniformScaled, IndexNode1), Grid.GetInner2DPoint(EGridSpace::UniformScaled, IndexNode2));
 			return;
 		}
 
-		if (Grid.IsNodeCloseToLoop(IndexNode1) && Grid.IsNodeCloseToLoop(IndexNode2))
+		if (Grid.IsNodeCloseToButOusideFace(IndexNode1) && Grid.IsNodeCloseToButOusideFace(IndexNode2))
 		{
-
-			if (LoopSegmentsIntersectionTool.DoesIntersect(Grid.GetInner2DPoint(EGridSpace::UniformScaled, IndexNode1), Grid.GetInner2DPoint(EGridSpace::UniformScaled, IndexNode2))
-				|| AlmostHitsLoop(Grid.GetInner2DPoint(EGridSpace::UniformScaled, IndexNode1), Grid.GetInner2DPoint(EGridSpace::UniformScaled, IndexNode2), InType))
-			{
-				InnerToOuterSegmentsIntersectionTool.AddSegment(Grid.GetInner2DPoint(EGridSpace::UniformScaled, IndexNode1), Grid.GetInner2DPoint(EGridSpace::UniformScaled, IndexNode2));
 #ifdef DEBUG_BUILDINNERSEGMENTS
-				if (bDisplay)
-				{
-					DisplaySegment(Grid.GetInner2DPoint(EGridSpace::UniformScaled, IndexNode1) * Grid.DisplayScale, Grid.GetInner2DPoint(EGridSpace::UniformScaled, IndexNode2) * Grid.DisplayScale, IndexNode1, EVisuProperty::YellowCurve);
-				}
-#endif
-				return;
-			}
-			else
+			if (bDisplay)
 			{
-#ifdef DEBUG_BUILDINNERSEGMENTS
-				if (bDisplay)
-				{
-					DisplaySegment(Grid.GetInner2DPoint(EGridSpace::UniformScaled, IndexNode1) * Grid.DisplayScale, Grid.GetInner2DPoint(EGridSpace::UniformScaled, IndexNode2) * Grid.DisplayScale, IndexNode1, EVisuProperty::GreenCurve);
-				}
+				DisplaySegment(Grid.GetInner2DPoint(EGridSpace::UniformScaled, IndexNode1) * DisplayScale, Grid.GetInner2DPoint(EGridSpace::UniformScaled, IndexNode2) * DisplayScale, IndexNode1, EVisuProperty::PinkCurve);
+			}
 #endif
+			AddToInnerToOuterSegmentsIntersectionTool(IndexNode1, IndexNode2, InType);
+			return;
+		}
+
+
+		if (Grid.IsNodeInsideAndMeshable(IndexNode1) && Grid.IsNodeInsideAndMeshable(IndexNode2))
+		{
+#ifdef DEBUG_BUILDINNERSEGMENTS
+			if (bDisplay)
+			{
+				DisplaySegment(Grid.GetInner2DPoint(EGridSpace::UniformScaled, IndexNode1) * DisplayScale, Grid.GetInner2DPoint(EGridSpace::UniformScaled, IndexNode2) * DisplayScale, IndexNode1, EVisuProperty::BlueCurve);
 			}
+#endif
+
+			FIsoInnerNode& Node1 = *GlobalIndexToIsoInnerNodes[IndexNode1];
+			FIsoInnerNode& Node2 = *GlobalIndexToIsoInnerNodes[IndexNode2];
+			FIsoSegment& Segment = IsoSegmentFactory.New();
+			Segment.Init(Node1, Node2, InType);
+			Segment.ConnectToNode();
+			FinalInnerSegments.Add(&Segment);
+			return;
 		}
-		else
+
+		if (Grid.IsNodeInsideButTooCloseToLoop(IndexNode1) && Grid.IsNodeInsideButTooCloseToLoop(IndexNode2))
 		{
 #ifdef DEBUG_BUILDINNERSEGMENTS
 			if (bDisplay)
 			{
-				DisplaySegment(Grid.GetInner2DPoint(EGridSpace::UniformScaled, IndexNode1) * Grid.DisplayScale, Grid.GetInner2DPoint(EGridSpace::UniformScaled, IndexNode2) * Grid.DisplayScale, IndexNode1, EVisuProperty::RedCurve);
+				DisplaySegment(Grid.GetInner2DPoint(EGridSpace::UniformScaled, IndexNode1) * DisplayScale, Grid.GetInner2DPoint(EGridSpace::UniformScaled, IndexNode2) * DisplayScale, IndexNode1, EVisuProperty::GreenCurve);
 			}
 #endif
+			return;
 		}
 
-		FIsoInnerNode& Node1 = *GlobalIndexToIsoInnerNodes[IndexNode1];
-		FIsoInnerNode& Node2 = *GlobalIndexToIsoInnerNodes[IndexNode2];
-		FIsoSegment& Segment = IsoSegmentFactory.New();
-		Segment.Init(Node1, Node2, InType);
-		Segment.ConnectToNode();
-		FinalInnerSegments.Add(&Segment);
+#ifdef DEBUG_BUILDINNERSEGMENTS
+		if (bDisplay)
+		{
+			DisplaySegment(Grid.GetInner2DPoint(EGridSpace::UniformScaled, IndexNode1) * DisplayScale, Grid.GetInner2DPoint(EGridSpace::UniformScaled, IndexNode2) * DisplayScale, IndexNode1, EVisuProperty::YellowCurve);
+		}
+#endif
+		AddToInnerToOuterSegmentsIntersectionTool(IndexNode1, IndexNode2, InType);
 	};
 
 	for (int32 uIndex = 0; uIndex < NumU; uIndex++)
@@ -569,6 +587,13 @@ void FIsoTriangulator::BuildInnerSegments()
 		}
 	}
 
+#ifdef DEBUG_BUILDINNERSEGMENTS
+	Wait(bDisplay);
+	InnerToOuterSegmentsIntersectionTool.Display(bDisplay, TEXT("InnerToOuterSegmentsIntersectionTool End"), EVisuProperty::RedCurve);
+#endif
+
+	InnerToOuterSegmentsIntersectionTool.Sort();
+
 #ifdef CADKERNEL_DEV
 	Chronos.BuildInnerSegmentsDuration += FChrono::Elapse(StartTime);
 #endif
@@ -617,7 +642,7 @@ void FIsoTriangulator::BuildInnerSegmentsIntersectionTool()
 		{
 			for (int32 IndexU = 0; IndexU < NumU; ++IndexU, ++Index)
 			{
-				if (!Grid.IsNodeInsideFace(Index))
+				if (!Grid.IsNodeInsideAndMeshable(Index))
 				{
 					continue;
 				}
@@ -651,7 +676,7 @@ void FIsoTriangulator::BuildInnerSegmentsIntersectionTool()
 	{
 		for (int32 IndexU = 0; IndexU < NumU; ++IndexU, ++Index)
 		{
-			if (!Grid.IsNodeInsideFace(Index))
+			if (!Grid.IsNodeInsideAndMeshable(Index))
 			{
 				continue;
 			}
@@ -937,7 +962,7 @@ void FIsoTriangulator::FindIsoSegmentToLinkLoopToLoop()
 			});
 		FindIsoSegmentAlong(IsoU, IsoUCoordinates, IsoVCoordinates);
 
-		Algo::Sort(SortedLoopNodesAlong, [&](const FLoopNode* LoopNode1, const FLoopNode* LoopNode2)
+		Algo::Sort(SortedLoopNodesAlong, [&](const FLoopNode* LoopNode1, const FLoopNode* LoopNode2) 
 			{
 				const FPoint2D& Node1Coordinates = LoopNode1->Get2DPoint(EGridSpace::UniformScaled, Grid);
 				const FPoint2D& Node2Coordinates = LoopNode2->Get2DPoint(EGridSpace::UniformScaled, Grid);
@@ -1147,7 +1172,7 @@ void FIsoTriangulator::FindIsoSegmentToLinkOuterLoopNodes(FCell& Cell)
 			TryToConnectVertexSubLoopWithTheMostIsoSegment(Cell, SubLoop);
 		}
 	}
-	
+
 	Cell.SelectSegmentInCandidateSegments(IsoSegmentFactory);
 }
 
@@ -1587,82 +1612,6 @@ bool FIsoTriangulator::CanCycleBeMeshed(const TArray<FIsoSegment*>& Cycle, FInte
 	return true;
 }
 
-#ifdef WIP_ADD_STEP_TO_TO_FAVOR_ISO_SEGMENTS
-//#define DEBUG_FIND_ISO_SEGMENT_CYCLE
-void FIsoTriangulator::FindIsoCandidateSegmentInCycle(TArray<FIsoNode*> CycleNodes)
-{
-#ifdef DEBUG_FIND_ISO_SEGMENT_CYCLE
-	F3DDebugSession _(bDisplay, TEXT("FindIsoCandidateSegmentInCycle"));
-#endif
-
-	const double FlatSlope = 0.10; // ~5 deg: The segment must make an angle less than 10 deg with the Iso
-	double MinSlope = FlatSlope + DOUBLE_SMALL_NUMBER; //.25; // ~10 deg: The segment must make an angle less than 10 deg with the Iso
-
-	FIsoNode* CandidateA = nullptr;
-	FIsoNode* CandidateB = nullptr;
-
-	int32 CycleCount = CycleNodes.Num();
-
-	TArray<FIsoSegment*> CandiateSegment;
-	CandiateSegment.Reserve(CycleCount);
-
-	for (int32 IndexA = 0; IndexA < CycleCount - 2; ++IndexA)
-	{
-		FIsoNode* NodeA = CycleNodes[IndexA];
-		const FPoint2D& ACoordinates = NodeA->Get2DPoint(EGridSpace::UniformScaled, Grid);
-		Grid.DisplayIsoNode(UniformScaled, *NodeA, 0, EVisuProperty::GreenPoint);
-
-		for (int32 IndexB = IndexA + 2; IndexB < CycleCount; ++IndexB)
-		{
-			FIsoNode* NodeB = CycleNodes[IndexB];
-
-			if(NodeA->GetSegmentConnectedTo(NodeB))
-			{
-				continue;
-			}
-
-			const FPoint2D& BCoordinates = NodeB->Get2DPoint(EGridSpace::UniformScaled, Grid);
-
-			double Slope = ComputeSlopeRelativeToNearestAxis(ACoordinates, BCoordinates);
-			if (Slope < MinSlope)
-			{
-				MinSlope = Slope;
-				CandidateA = NodeA;
-				CandidateB = NodeB;
-			}
-		}
-
-		if (MinSlope < FlatSlope)
-		{
-			FIsoSegment& Segment = IsoSegmentFactory.New();
-			Segment.Init(*CandidateA, *CandidateB, ESegmentType::LoopToLoop);
-			Segment.SetCandidate();
-
-			CandiateSegment.Add(&Segment);
-
-#ifdef DEBUG_FIND_ISO_SEGMENT_CYCLE
-			Grid.DisplayIsoSegment(UniformScaled, Segment, 0, EVisuProperty::YellowCurve);
-			Wait(bDisplay);
-#endif
-			MinSlope = FlatSlope + DOUBLE_SMALL_NUMBER;
-		}
-	}
-
-	if(bDisplay)
-	{
-		Wait();
-	}
-
-	// find the best iso segments i.e. ones that generate smooth surface
-	// check that the selected segments won't block the meshing i.e. a cycle segment cannot be meshed
-	for(FIsoSegment* Segment : CandiateSegment)
-	{
-	}
-
-}
-#endif
-
-
 void FIsoTriangulator::MeshCycle(const EGridSpace Space, const TArray<FIsoSegment*>& Cycle, const TArray<bool>& CycleOrientation)
 {
 	int32 NodeCycleNum = (int32)Cycle.Num();
@@ -1777,6 +1726,7 @@ void FIsoTriangulator::MeshCycle(const EGridSpace Space, const TArray<FIsoSegmen
 	static int32 CycleId = 0;
 	F3DDebugSession _(bDisplay, FString::Printf(TEXT("Mesh cycle %d"), CycleId++));
 	Grid.DisplayIsoSegments(TEXT("Cycle"), EGridSpace::UniformScaled, Cycle);
+	InnerToOuterSegmentsIntersectionTool.Display(Grid.bDisplay, TEXT("InnerToOuterSegmentsIntersectionTool"), EVisuProperty::RedCurve);
 	Wait(bDisplay);
 #endif
 
@@ -2174,13 +2124,7 @@ void FIsoTriangulator::MeshCycle(const EGridSpace Space, const TArray<FIsoSegmen
 				CycleIntersectionTool.Sort();
 			}
 		}
-	};
-
-#ifdef WIP_ADD_STEP_TO_TO_FAVOR_ISO_SEGMENTS
-	FindIsoCandidateSegmentInCycle(CycleNodes);
-#endif
-
-
+						};
 
 	for (int32 Index = 0; Index < SegmentStack.Num(); ++Index)
 	{
@@ -2418,7 +2362,7 @@ void FIsoTriangulator::TriangulateOverCycle(const EGridSpace Space)
 			Cycle.Empty();
 			CycleOrientation.Empty();
 			bool bLeftSide = false;
-			if(!FindCycle(Segment, bLeftSide, Cycle, CycleOrientation))
+			if (!FindCycle(Segment, bLeftSide, Cycle, CycleOrientation))
 			{
 				continue;
 			}
@@ -2725,7 +2669,7 @@ void FIsoTriangulator::ConnectCellSubLoopsByNeighborhood(FCell& Cell)
 
 	int32 LoopCount = Cell.SubLoops.Num();
 
-	TArray<TPair<int32,FPoint2D>> LoopBarycenters;
+	TArray<TPair<int32, FPoint2D>> LoopBarycenters;
 	LoopBarycenters.Reserve(LoopCount + 4);
 
 	int32 LoopIndex = -1;
@@ -3071,10 +3015,7 @@ void FIsoTriangulator::TryToConnectVertexSubLoopWithTheMostIsoSegment(FCell& Cel
 {
 #ifdef DEBUG_FIND_ISO_SEGMENT_TO_LINK_OUTER_LOOP_NODES
 	F3DDebugSession _(bDisplay, TEXT("TryToConnectVertexSubLoopWithTheMostIsoSegment"));
-	if(bDisplay)
-	{
-		Wait();
-	}
+	Wait(bDisplay);
 #endif
 
 	const double FlatSlope = 0.10; // ~5 deg: The segment must make an angle less than 10 deg with the Iso
@@ -3275,7 +3216,7 @@ void FIsoTriangulator::ConnectCellCornerToInnerLoop(FCell& Cell)
 		}
 	}
 
-#ifdef DEBUG_DELAUNAY
+#ifdef DEBUG_CONNECTCELLCORNERTOINNERLOOP
 	F3DDebugSession _(TEXT("With cell corners"));
 #endif
 
@@ -3321,12 +3262,17 @@ void FIsoTriangulator::ConnectCellCornerToInnerLoop(FCell& Cell)
 		{
 			for (int32 IndexLoopA : Cell.BorderLoopIndices)
 			{
+#ifdef DEBUG_CONNECTCELLCORNERTOINNERLOOP
 				DisplayPoint(CellNodes[ICell]->Get2DPoint(EGridSpace::UniformScaled, Grid), EVisuProperty::GreenPoint);
+#endif
 				FindAndTryCreateCandidateSegmentToLinkLoopToCorner(IndexLoopA, CellNodes[ICell]);
 			}
 
 			if (Cell.bHasOuterLoop)
 			{
+#ifdef DEBUG_CONNECTCELLCORNERTOINNERLOOP
+				DisplayPoint(CellNodes[ICell]->Get2DPoint(EGridSpace::UniformScaled, Grid), EVisuProperty::GreenPoint);
+#endif
 				FindAndTryCreateCandidateSegmentToLinkLoopToCorner(0, CellNodes[ICell]);
 			}
 		}
diff --git a/Engine/Source/Runtime/Datasmith/CADKernel/Private/CADKernel/Mesh/Meshers/IsoTriangulator.h b/Engine/Source/Runtime/Datasmith/CADKernel/Private/CADKernel/Mesh/Meshers/IsoTriangulator.h
--- a/Engine/Source/Runtime/Datasmith/CADKernel/Private/CADKernel/Mesh/Meshers/IsoTriangulator.h
+++ b/Engine/Source/Runtime/Datasmith/CADKernel/Private/CADKernel/Mesh/Meshers/IsoTriangulator.h
@@ -110,8 +110,6 @@ protected:
 	TArray<FIsoInnerNode> InnerNodes;
 	int32 InnerNodeCount = 0;
 
-
-
 	TFactory<FIsoSegment> IsoSegmentFactory;
 
 	TArray<FIsoSegment*> LoopSegments;
@@ -127,8 +125,7 @@ protected:
 	FIntersectionSegmentTool LoopSegmentsIntersectionTool;
 	FIntersectionSegmentTool InnerSegmentsIntersectionTool;
 	FIntersectionSegmentTool InnerToLoopSegmentsIntersectionTool;
-	FIntersectionSegmentTool InnerToOuterSegmentsIntersectionTool;
-
+	FIntersectionNodePairTool InnerToOuterSegmentsIntersectionTool;
 
 	/**
 	 * Define all the lower left index of grid node that the upper cell is surrounding a loop
@@ -264,12 +261,6 @@ public:
 	 */
 	void FindIsoSegmentToLinkLoopToLoop();
 
-	/**
-	 * The purpose of the method is select a minimal set of segments connecting loops together
-	 * The final segments will be selected with SelectSegmentInCandidateSegments
-	 */
-	void ConnectCellSubLoopsByNeighborhood();  // to rename and clean
-
 	void FindCandidateSegmentsToLinkInnerAndLoop();
 
 	/**
@@ -301,10 +292,6 @@ public:
 	 */
 	void MeshCycle(const EGridSpace Space, const TArray<FIsoSegment*>& cycle, const TArray<bool>& cycleOrientation);
 
-#ifdef WIP_ADD_STEP_TO_TO_FAVOR_ISO_SEGMENTS
-	void FindIsoCandidateSegmentInCycle(TArray<FIsoNode*> CycleNodes);
-#endif
-
 	bool CanCycleBeMeshed(const TArray<FIsoSegment*>& Cycle, FIntersectionSegmentTool& CycleIntersectionTool);
 
 	/**
diff --git a/Engine/Source/Runtime/Datasmith/CADKernel/Private/CADKernel/Mesh/Structure/Grid.cpp b/Engine/Source/Runtime/Datasmith/CADKernel/Private/CADKernel/Mesh/Structure/Grid.cpp
--- a/Engine/Source/Runtime/Datasmith/CADKernel/Private/CADKernel/Mesh/Structure/Grid.cpp
+++ b/Engine/Source/Runtime/Datasmith/CADKernel/Private/CADKernel/Mesh/Structure/Grid.cpp
@@ -324,7 +324,7 @@ bool FGrid::GeneratePointCloud()
 		return false;
 	}
 
-	PointMarkers.Init(EPointMarker::None, CuttingSize);
+	NodeMarkers.Init(ENodeMarker::None, CuttingSize);
 
 	CountOfInnerNodes = CuttingSize;
 	for (int32 Index = 0; Index < EGridSpace::EndGridSpace; ++Index)
@@ -739,7 +739,7 @@ void FGrid::RemovePointsCloseToLoop()
 		IndexOfPointsNearAndInsideLoop.Reserve(CuttingSize);
 		for (int32 Index = 0; Index < CuttingSize; ++Index)
 		{
-			if (IsNodeCloseToAndInsideLoop(Index))
+			if (IsNodeInsideAndCloseToLoop(Index))
 			{
 				IndexOfPointsNearAndInsideLoop.Add(Index);
 			}
@@ -865,7 +865,6 @@ void FGrid::RemovePointsCloseToLoop()
 			}
 
 			SetTooCloseToLoop(Index);
-			CountOfInnerNodes--;
 			break;
 		}
 	}
@@ -1477,7 +1476,7 @@ void FGrid::FindInnerFacePoints()
 	TArray<char> IntersectLoop;
 
 	IntersectLoop.Init(0, CuttingSize);
-	PointMarkers.Init(EPointMarker::IsInside, CuttingSize);
+	NodeMarkers.Init(ENodeMarker::IsInside, CuttingSize);
 
 	NbIntersectUForward.Init(0, CuttingSize);
 	NbIntersectUBackward.Init(0, CuttingSize);
diff --git a/Engine/Source/Runtime/Datasmith/CADKernel/Private/CADKernel/Mesh/Structure/Grid.h b/Engine/Source/Runtime/Datasmith/CADKernel/Private/CADKernel/Mesh/Structure/Grid.h
--- a/Engine/Source/Runtime/Datasmith/CADKernel/Private/CADKernel/Mesh/Structure/Grid.h
+++ b/Engine/Source/Runtime/Datasmith/CADKernel/Private/CADKernel/Mesh/Structure/Grid.h
@@ -71,21 +71,17 @@ struct FGridChronos
 	}
 };
 
-enum class EPointMarker : uint8
+enum class ENodeMarker : uint8
 {
 	None = 0x00u,  // No flags.
 
-	IsInside = 0x01u, 
-	IsCloseToLoop = 0x02u,
-
-	IsInsideAndClose = 0x03u,
-
-	//IsTooCloseToLoop = 0x04u,  // node inside the loop but too close to the loop to be include in the mesh
-	//IsInsideAndTooCloseToLoop = 0x05u,
+	IsInside = 0x01u,
+	IsInsideButTooCloseToLoop = 0x02u,  // node inside the loop but too close to the loop to be include in the mesh
+	IsCloseToLoop = 0x04u,
 
 	All = 0xFFu
 };
-ENUM_CLASS_FLAGS(EPointMarker);
+ENUM_CLASS_FLAGS(ENodeMarker);
 
 
 class FGrid : public FHaveStates
@@ -148,10 +144,10 @@ protected:
 	double MinOfMaxElementSize = 0;
 
 	/**
-	 * Array to flag each points as inside, close, ... the face 
-	 * @see IsNode..., 
+	 * Array to flag each points as inside, close, ... the face
+	 * @see IsNode...,
 	 */
-	TArray<EPointMarker> PointMarkers;
+	TArray<ENodeMarker> NodeMarkers;
 
 	/**
 	 * 2D Coordinate of grid nodes in each space
@@ -328,42 +324,54 @@ public:
 	}
 
 	/**
-	 * @return true if the node is inner the external loop and outer the inner loops
+	 * @return true if the node is inner the external loop and outer the inner loops and not too close
 	 */
-	const bool IsNodeInsideFace(int32 IndexU, int32 IndexV) const
+	const bool IsNodeInsideAndMeshable(int32 Index) const
 	{
-		return IsNodeInsideFace(GobalIndex(IndexU, IndexV));
+		return (NodeMarkers[Index] & ENodeMarker::IsInside) == ENodeMarker::IsInside;
+	}
+
+	const bool IsNodeInsideButTooCloseToLoop(int32 Index) const
+	{
+		return (NodeMarkers[Index] & ENodeMarker::IsInsideButTooCloseToLoop) == ENodeMarker::IsInsideButTooCloseToLoop;
 	}
 
 	/**
-	 * @return true if the node is inner the external loop and outer the inner loops
+	 * @return true if the bits EPointMarker::IsInside, EPointMarker::IsInsideButTooCloseToLoop and EPointMarker::IsCloseToLoop are alls equal to 0
 	 */
-	const bool IsNodeInsideFace(int32 Index) const
+	const bool IsNodeFarFromFace(int32 Index) const
 	{
-		return (PointMarkers[Index] & EPointMarker::IsInside) == EPointMarker::IsInside;
+		constexpr ENodeMarker IsInsideAndClose = ENodeMarker::IsInside | ENodeMarker::IsInsideButTooCloseToLoop | ENodeMarker::IsCloseToLoop;
+		return (NodeMarkers[Index] & IsInsideAndClose) == ENodeMarker::None;
 	}
 
 	/**
-	 * @return true if 
+	 * @return true if the bits EPointMarker::IsInside, EPointMarker::IsInsideButTooCloseToLoop and EPointMarker::IsCloseToLoop are alls equal to 0
 	 */
-	//const bool IsNodeFarFromFace(int32 Index) const
-	//{
-	//	return (PointMarkers[Index] & EPointMarker::IsInsideAndClose) == EPointMarker::None;
-	//}
+	const bool IsNodeOutsideFace(int32 Index) const
+	{
+		constexpr ENodeMarker IsInsideAndClose = ENodeMarker::IsInside | ENodeMarker::IsInsideButTooCloseToLoop;
+		return (NodeMarkers[Index] & IsInsideAndClose) == ENodeMarker::None;
+	}
 
-	const bool IsNodeOusideFace(int32 Index) const
+	const bool IsNodeInsideAndCloseToLoop(int32 Index) const
 	{
-		return (PointMarkers[Index] & EPointMarker::IsInside) == EPointMarker::None;
+		constexpr ENodeMarker IsInsideAndClose = ENodeMarker::IsInside | ENodeMarker::IsCloseToLoop;
+		return (NodeMarkers[Index] & IsInsideAndClose) == IsInsideAndClose;
 	}
 
+	/**
+	 * @return true if the bits EPointMarker::IsInside and EPointMarker::IsInsideButTooCloseToLoop are booth equal to 0 and EPointMarker::IsCloseToLoop  is equal to 1
+	 */
 	const bool IsNodeCloseToButOusideFace(int32 Index) const
 	{
-		return (PointMarkers[Index] & EPointMarker::IsInsideAndClose) == EPointMarker::IsCloseToLoop;
+		constexpr ENodeMarker IsInsideAndClose = ENodeMarker::IsInside | ENodeMarker::IsInsideButTooCloseToLoop | ENodeMarker::IsCloseToLoop;
+		return (NodeMarkers[Index] & IsInsideAndClose) == ENodeMarker::IsCloseToLoop;
 	}
 
 	void SetNodeInside(int32 Index)
 	{
-		PointMarkers[Index] |= EPointMarker::IsInside;
+		NodeMarkers[Index] |= ENodeMarker::IsInside;
 	}
 
 	/**
@@ -379,61 +387,30 @@ public:
 	 */
 	const bool IsNodeCloseToLoop(int32 Index) const
 	{
-		return (PointMarkers[Index] & EPointMarker::IsCloseToLoop) == EPointMarker::IsCloseToLoop;
+		return (NodeMarkers[Index] & ENodeMarker::IsCloseToLoop) == ENodeMarker::IsCloseToLoop;
 	}
 
-	/*const bool IsNodeTooCloseToLoop2(int32 IndexU, int32 IndexV) const
+	const bool IsNodeTooCloseToLoop(int32 Index) const
 	{
-		return IsNodeTooCloseToLoop2(GobalIndex(IndexU, IndexV));
-	}*/
-
-	//const bool IsNodeTooCloseToLoop2(int32 Index) const
-	//{
-	//	return (PointMarkers[Index] & EPointMarker::IsTooCloseToLoop) == EPointMarker::IsTooCloseToLoop;
-	//}
+		return (NodeMarkers[Index] & ENodeMarker::IsInsideButTooCloseToLoop) == ENodeMarker::IsInsideButTooCloseToLoop;
+	}
 
 	void SetCloseToLoop(int32 Index)
 	{
-		PointMarkers[Index] |= EPointMarker::IsCloseToLoop;
+		NodeMarkers[Index] |= ENodeMarker::IsCloseToLoop;
 	}
 
-	/*void ResetCloseToAndInsideLoop(int32 Index)
-	{
-		PointMarkers[Index] &= ~EPointMarker::IsInsideAndClose;
-	}*/
-
 	void SetTooCloseToLoop(int32 Index)
 	{
-		//PointMarkers[Index] |= EPointMarker::IsTooCloseToLoop;
-		PointMarkers[Index] &= ~EPointMarker::IsInsideAndClose;
+		NodeMarkers[Index] |= ENodeMarker::IsInsideButTooCloseToLoop;
+		NodeMarkers[Index] &= ~ENodeMarker::IsInside;
 	}
 
 	void ResetInsideLoop(int32 Index)
 	{
-		PointMarkers[Index] &= ~EPointMarker::IsInside;
+		NodeMarkers[Index] &= ~ENodeMarker::IsInside;
 	}
 
-	/**
-	 * @return true if the node is close to or inside a loop i.e. the loop cross the space [[IndexU - 1, IndexU + 1], [IndexV - 1, IndexV + 1]]
-	 */
-	/*const bool IsNodeCloseToOrInsideLoop2(int32 Index) const
-	{
-		return (bool) (PointMarkers[Index] & EPointMarker::IsInsideAndClose);
-	}*/
-
-	const bool IsNodeCloseToAndInsideLoop(int32 Index) const
-	{
-		return (PointMarkers[Index] & EPointMarker::IsInsideAndClose) == EPointMarker::IsInsideAndClose;
-	}
-
-	/**
-	 * @return true if the node is inside the face but not too close to the loop to be meshed
-	 */
-	//const bool IsNodeMeshable(int32 Index) const
-	//{
-	//	return (PointMarkers[Index] & EPointMarker::IsInsideAndTooCloseToLoop) == EPointMarker::IsInside;
-	//}
-
 	/**
 	 * @return the FPoint2D (parametric coordinates) of the point at the Index of the grid in the defined grid space
 	 * @see EGridSpace
@@ -718,13 +695,17 @@ public:
 		F3DDebugSession _(Message);
 		for (int32 Index = 0; Index < Points.Num(); ++Index)
 		{
-			if (IsNodeInsideFace(Index))
+			if (IsNodeInsideAndMeshable(Index))
+			{
+				DisplayPoint(Points[Index] * DisplayScale, EVisuProperty::BluePoint, Index);
+			}
+			else if (IsNodeTooCloseToLoop(Index))
 			{
-				DisplayPoint(Points[Index] * DisplayScale, /*IsNodeTooCloseToLoop2(Index) ? EVisuProperty::RedPoint :*/ EVisuProperty::BluePoint, Index);
+				DisplayPoint(Points[Index] * DisplayScale, EVisuProperty::RedPoint, Index);
 			}
 			else
 			{
-				DisplayPoint(Points[Index] * DisplayScale, IsNodeCloseToLoop(Index) ? EVisuProperty::YellowPoint : EVisuProperty::GreenPoint , Index);
+				DisplayPoint(Points[Index] * DisplayScale, IsNodeCloseToLoop(Index) ? EVisuProperty::YellowPoint : EVisuProperty::GreenPoint, Index);
 			}
 		}
 	}
diff --git a/Engine/Source/Runtime/Datasmith/CADKernel/Private/CADKernel/Mesh/Structure/GridDisplay.cpp b/Engine/Source/Runtime/Datasmith/CADKernel/Private/CADKernel/Mesh/Structure/GridDisplay.cpp
--- a/Engine/Source/Runtime/Datasmith/CADKernel/Private/CADKernel/Mesh/Structure/GridDisplay.cpp
+++ b/Engine/Source/Runtime/Datasmith/CADKernel/Private/CADKernel/Mesh/Structure/GridDisplay.cpp
@@ -309,7 +309,7 @@ void FGrid::DisplayGridPoints(EGridSpace DisplaySpace) const
 		F3DDebugSession _(TEXT("FGrid::FindInnerDomainPoints Inside Point"));
 		for (int32 Index = 0; Index < CuttingSize; ++Index)
 		{
-			if (IsNodeInsideFace(Index))
+			if (IsNodeInsideAndMeshable(Index))
 			{
 				DisplayPoint(Points2D[DisplaySpace][Index] * DisplayScale, EVisuProperty::BluePoint, Index);
 				NbNum++;
@@ -322,7 +322,7 @@ void FGrid::DisplayGridPoints(EGridSpace DisplaySpace) const
 		F3DDebugSession _(TEXT("FGrid::FindInnerDomainPoints Outside Point"));
 		for (int32 Index = 0; Index < CuttingSize; ++Index)
 		{
-			if (!IsNodeInsideFace(Index))
+			if (!IsNodeInsideAndMeshable(Index))
 			{
 				DisplayPoint(Points2D[DisplaySpace][Index] * DisplayScale, EVisuProperty::OrangePoint, Index);
 			}
