[NO NON-PII COMMIT MESSAGE]

diff --git a/Engine/Source/Runtime/Datasmith/CADKernel/Private/CADKernel/Math/Geometry.cpp b/Engine/Source/Runtime/Datasmith/CADKernel/Private/CADKernel/Math/Geometry.cpp
--- a/Engine/Source/Runtime/Datasmith/CADKernel/Private/CADKernel/Math/Geometry.cpp
+++ b/Engine/Source/Runtime/Datasmith/CADKernel/Private/CADKernel/Math/Geometry.cpp
@@ -5,6 +5,7 @@
 #include "CADKernel/Core/KernelParameters.h" 
 #include "CADKernel/Core/System.h"
 #include "CADKernel/Math/Geometry.h"
+#include "CADKernel/Math/Point.h"
 #include "CADKernel/Math/MatrixH.h"
 #include "CADKernel/Utils/Util.h"
 
@@ -12,23 +13,18 @@ namespace UE::CADKernel
 {
 double ComputeCurvature(const FPoint& Gradient, const FPoint& Laplacian)
 {
-	FPoint GradientCopy = Gradient;
-	FPoint LaplacianCopy = Laplacian;
-	GradientCopy.Normalize();
-	LaplacianCopy.Normalize();
-	FPoint Normal = GradientCopy ^ LaplacianCopy;
+	const FPoint GradientCopy = Gradient.Normalize();
+	const FPoint LaplacianCopy = Laplacian.Normalize();
+	const FPoint Normal = GradientCopy ^ LaplacianCopy;
 	return (Normal.Length() * Laplacian.Length()) / (Gradient.Length() * Gradient.Length());
 }
 
 double ComputeCurvature(const FPoint& Normal, const FPoint& Gradient, const FPoint& Laplacian)
 {
-	FPoint GradientCopy = Gradient;
-	FPoint LaplacianCopy = Laplacian;
-	FPoint NormalCopy = Normal;
-	GradientCopy.Normalize();
-	LaplacianCopy.Normalize();
-	NormalCopy.Normalize();
-	FPoint Coef = (LaplacianCopy ^ GradientCopy) ^ NormalCopy;
+	const FPoint GradientCopy = Gradient.Normalize();
+	const FPoint LaplacianCopy = Laplacian.Normalize();
+	const FPoint NormalCopy = Normal.Normalize();
+	const FPoint Coef = (LaplacianCopy ^ GradientCopy) ^ NormalCopy;
 	return (Coef.Length() * Laplacian.Length()) / Gradient.SquareLength();
 }
 
@@ -36,8 +32,8 @@ void FindLoopIntersectionsWithIso(const EIso Iso, const double IsoParameter, con
 {
 	OutIntersections.Empty(8);
 
-	int32 UIndex = Iso == EIso::IsoU ? 0 : 1;
-	int32 VIndex = Iso == EIso::IsoU ? 1 : 0;
+	const int32 UIndex = Iso == EIso::IsoU ? 0 : 1;
+	const int32 VIndex = Iso == EIso::IsoU ? 1 : 0;
 
 	TFunction<void(const FPoint2D&, const FPoint2D&)> ComputeIntersection = [&](const FPoint2D& Point1, const FPoint2D& Point2)
 	{
@@ -72,8 +68,8 @@ void FindLoopIntersectionsWithIso(const EIso Iso, const double IsoParameter, con
 
 bool IntersectSegments2D(const TSegment<FPoint2D>& SegmentAB, const TSegment<FPoint2D>& SegmentCD)
 {
-	constexpr const double Min = -DOUBLE_SMALL_NUMBER;
-	constexpr const double Max = 1. + DOUBLE_SMALL_NUMBER;
+	constexpr const double Min = -DOUBLE_KINDA_SMALL_NUMBER;
+	constexpr const double Max = 1. + DOUBLE_KINDA_SMALL_NUMBER;
 
 	TFunction<bool(double, double, double, double)> Intersect = [](double A, double B, double C, double D)
 	{
@@ -106,9 +102,9 @@ bool IntersectSegments2D(const TSegment<FPoint2D>& SegmentAB, const TSegment<FPo
 		}
 	};
 
-	FPoint2D AB = SegmentAB[1] - SegmentAB[0];
-	FPoint2D CD = SegmentCD[1] - SegmentCD[0];
-	FPoint2D CA = SegmentAB[0] - SegmentCD[0];
+	const FPoint2D AB = SegmentAB[1] - SegmentAB[0];
+	const FPoint2D CD = SegmentCD[1] - SegmentCD[0];
+	const FPoint2D CA = SegmentAB[0] - SegmentCD[0];
 
 	const double ParallelCoef = CD ^ AB;
 	if (FMath::IsNearlyZero(ParallelCoef))
@@ -133,4 +129,5 @@ bool IntersectSegments2D(const TSegment<FPoint2D>& SegmentAB, const TSegment<FPo
 	const double CDIntersectionCoordinate = (CA ^ AB) / ParallelCoef;
 	return (ABIntersectionCoordinate <= Max && ABIntersectionCoordinate >= Min && CDIntersectionCoordinate <= Max && CDIntersectionCoordinate >= Min);
 }
+
 }
diff --git a/Engine/Source/Runtime/Datasmith/CADKernel/Private/CADKernel/Mesh/Structure/Grid.cpp b/Engine/Source/Runtime/Datasmith/CADKernel/Private/CADKernel/Mesh/Structure/Grid.cpp
--- a/Engine/Source/Runtime/Datasmith/CADKernel/Private/CADKernel/Mesh/Structure/Grid.cpp
+++ b/Engine/Source/Runtime/Datasmith/CADKernel/Private/CADKernel/Mesh/Structure/Grid.cpp
@@ -66,18 +66,29 @@ void FGrid::ProcessPointCloud()
 
 	DisplayGridNormal();
 
+	DisplayInnerDomainPoints(TEXT("FGrid::PointCloud before FindInnerFacePoints"), GetInner2DPoints(EGridSpace::UniformScaled));
+
 	Wait(bDisplay);
 #endif
 
 	FindInnerFacePoints();
 
 #ifdef DEBUG_GRID
-	DisplayGridPoints(DisplaySpace);
-	Wait(bDisplay);
+	DisplayInnerDomainPoints(TEXT("FGrid::PointCloud after FindInnerFacePoints"), GetInner2DPoints(EGridSpace::UniformScaled));
 #endif
+
 	FindPointsCloseToLoop();
+
+#ifdef DEBUG_GRID
+	DisplayInnerDomainPoints(TEXT("FGrid::PointCloud after FindPointsCloseToLoop"), GetInner2DPoints(EGridSpace::UniformScaled));
+#endif
+
 	RemovePointsCloseToLoop();
 
+#ifdef DEBUG_GRID
+	DisplayInnerDomainPoints(TEXT("FGrid::PointCloud after RemovePointsCloseToLoop"), GetInner2DPoints(EGridSpace::UniformScaled));
+#endif
+
 	// Removed of Thin zone boundary (the last boundaries). In case of thin zone, the number of 2d boundary will be biggest than 3d boundary one.
 	// Only EGridSpace::Default2D is needed.
 	FaceLoops2D[EGridSpace::Default2D].SetNum(FaceLoops3D.Num());
@@ -407,10 +418,10 @@ void FGrid::FindPointsCloseToLoop()
 		if (bDisplay)
 		{
 			F3DDebugSession _(*FString::Printf(TEXT("Cell %d"), CellIndex++));
-			DisplayPoint(Points2D[EGridSpace::UniformScaled][GlobalIndex]);
-			DisplayPoint(Points2D[EGridSpace::UniformScaled][GlobalIndex - 1]);
-			DisplayPoint(Points2D[EGridSpace::UniformScaled][GlobalIndex - 1 - CuttingCount[EIso::IsoU]]);
-			DisplayPoint(Points2D[EGridSpace::UniformScaled][GlobalIndex - CuttingCount[EIso::IsoU]]);
+			DisplayPoint(Points2D[EGridSpace::UniformScaled][GlobalIndex] * DisplayScale);
+			DisplayPoint(Points2D[EGridSpace::UniformScaled][GlobalIndex - 1] * DisplayScale);
+			DisplayPoint(Points2D[EGridSpace::UniformScaled][GlobalIndex - 1 - CuttingCount[EIso::IsoU]] * DisplayScale);
+			DisplayPoint(Points2D[EGridSpace::UniformScaled][GlobalIndex - CuttingCount[EIso::IsoU]] * DisplayScale);
 			Wait(bWaitCell);
 		}
 	};
@@ -571,7 +582,7 @@ void FGrid::FindPointsCloseToLoop()
 			if (bDisplay)
 			{
 				F3DDebugSession _(TEXT("SEG"));
-				DisplaySegment(*PointB, *PointA);
+				DisplaySegment(*PointB * DisplayScale, *PointA * DisplayScale);
 			}
 #endif
 
diff --git a/Engine/Source/Runtime/Datasmith/CADKernel/Private/CADKernel/Mesh/Structure/Grid.h b/Engine/Source/Runtime/Datasmith/CADKernel/Private/CADKernel/Mesh/Structure/Grid.h
--- a/Engine/Source/Runtime/Datasmith/CADKernel/Private/CADKernel/Mesh/Structure/Grid.h
+++ b/Engine/Source/Runtime/Datasmith/CADKernel/Private/CADKernel/Mesh/Structure/Grid.h
@@ -695,7 +695,11 @@ public:
 		F3DDebugSession _(Message);
 		for (int32 Index = 0; Index < Points.Num(); ++Index)
 		{
-			if (IsNodeInsideAndMeshable(Index))
+			if (IsNodeInsideAndCloseToLoop(Index))
+			{
+				DisplayPoint(Points[Index] * DisplayScale, EVisuProperty::OrangePoint, Index);
+			}
+			else if (IsNodeInsideAndMeshable(Index))
 			{
 				DisplayPoint(Points[Index] * DisplayScale, EVisuProperty::BluePoint, Index);
 			}
@@ -705,7 +709,7 @@ public:
 			}
 			else
 			{
-				DisplayPoint(Points[Index] * DisplayScale, IsNodeCloseToLoop(Index) ? EVisuProperty::YellowPoint : EVisuProperty::GreenPoint, Index);
+				DisplayPoint(Points[Index] * DisplayScale, IsNodeCloseToLoop(Index) ? EVisuProperty::YellowPoint : EVisuProperty::GreenPoint , Index);
 			}
 		}
 	}
diff --git a/Engine/Source/Runtime/Datasmith/CADKernel/Private/CADKernel/Mesh/Structure/GridDisplay.cpp b/Engine/Source/Runtime/Datasmith/CADKernel/Private/CADKernel/Mesh/Structure/GridDisplay.cpp
--- a/Engine/Source/Runtime/Datasmith/CADKernel/Private/CADKernel/Mesh/Structure/GridDisplay.cpp
+++ b/Engine/Source/Runtime/Datasmith/CADKernel/Private/CADKernel/Mesh/Structure/GridDisplay.cpp
@@ -316,7 +316,6 @@ void FGrid::DisplayGridPoints(EGridSpace DisplaySpace) const
 			}
 		}
 	}
-	ensureCADKernel(NbNum == CountOfInnerNodes);
 
 	{
 		F3DDebugSession _(TEXT("FGrid::FindInnerDomainPoints Outside Point"));
diff --git a/Engine/Source/Runtime/Datasmith/CADKernel/Public/CADKernel/Math/Geometry.h b/Engine/Source/Runtime/Datasmith/CADKernel/Public/CADKernel/Math/Geometry.h
--- a/Engine/Source/Runtime/Datasmith/CADKernel/Public/CADKernel/Math/Geometry.h
+++ b/Engine/Source/Runtime/Datasmith/CADKernel/Public/CADKernel/Math/Geometry.h
@@ -584,7 +584,7 @@ inline FPoint2D FindIntersectionOfSegments2D(const TSegment<FPoint2D>& SegmentAB
 }
 
 /**
- * @return false if the lines are parallele
+ * @return false if the lines are parallel
  */
 inline bool FindIntersectionOfLines2D(const TSegment<FPoint2D>& LineAB, const TSegment<FPoint2D>& LineCD, FPoint2D& OutIntersectionPoint)
 {
diff --git a/Engine/Source/Runtime/Datasmith/CADKernel/Public/CADKernel/Math/Point.h b/Engine/Source/Runtime/Datasmith/CADKernel/Public/CADKernel/Math/Point.h
--- a/Engine/Source/Runtime/Datasmith/CADKernel/Public/CADKernel/Math/Point.h
+++ b/Engine/Source/Runtime/Datasmith/CADKernel/Public/CADKernel/Math/Point.h
@@ -116,6 +116,10 @@ public:
 
 	FPoint operator/(double InvScale) const
 	{
+		if (FMath::IsNearlyZero(InvScale))
+		{
+			return *this;
+		}
 		return FPoint(X / InvScale, Y / InvScale, Z / InvScale);
 	}
 
@@ -181,6 +185,11 @@ public:
 
 	FPoint& operator/=(double InvScale)
 	{
+		if (FMath::IsNearlyZero(InvScale))
+		{
+			return *this;
+		}
+
 		X /= InvScale;
 		Y /= InvScale;
 		Z /= InvScale;
@@ -266,17 +275,28 @@ public:
 		double Norm = Length();
 		if (FMath::IsNearlyZero(Norm))
 		{
-			X = Y = Z = 0.0;
+			*this = FPoint::UnitPoint.Normalize();
 		}
 		else
 		{
-			X /= Norm;
-			Y /= Norm;
-			Z /= Norm;
+			*this /= Norm;
 		}
 		return *this;
 	}
 
+	FPoint Normalize() const 
+	{
+		double Norm = Length();
+		if (FMath::IsNearlyZero(Norm))
+		{
+			return FPoint::UnitPoint.Normalize();
+		}
+		else
+		{
+			return *this / Norm;
+		}
+	}
+
 	FPoint Middle(const FPoint& Point) const
 	{
 		return FPoint(0.5 * (X + Point.X), 0.5 * (Y + Point.Y), 0.5 * (Z + Point.Z));
@@ -602,6 +622,11 @@ public:
 
 	FPoint2D operator/(double InvScale) const
 	{
+		if (FMath::IsNearlyZero(InvScale))
+		{
+			return *this;
+		}
+
 		return FPoint2D(U / InvScale, V / InvScale);
 	}
 
@@ -633,6 +658,11 @@ public:
 
 	FPoint2D& operator/=(double InvScale)
 	{
+		if (FMath::IsNearlyZero(InvScale))
+		{
+			return *this;
+		}
+
 		U /= InvScale;
 		V /= InvScale;
 		return *this;
diff --git a/Engine/Source/Runtime/Datasmith/CADKernel/Public/CADKernel/Utils/Cache.h b/Engine/Source/Runtime/Datasmith/CADKernel/Public/CADKernel/Utils/Cache.h
--- a/Engine/Source/Runtime/Datasmith/CADKernel/Public/CADKernel/Utils/Cache.h
+++ b/Engine/Source/Runtime/Datasmith/CADKernel/Public/CADKernel/Utils/Cache.h
@@ -1,10 +1,6 @@
 // Copyright Epic Games, Inc. All Rights Reserved.
 #pragma once
 
-#include "CADKernel/Core/Types.h"
-#include "HAL/Platform.h"
-#include "Serialization/Archive.h"
-
 namespace UE::CADKernel
 {
 template<class ObjectType>
