To avoid to create triangle crossing natural iso line (a line carrying iso edges), a test is added to create triangle minimizing the number of intersection with these lines

diff --git a/Engine/Source/Runtime/Datasmith/CADKernel/Private/CADKernel/Mesh/Meshers/IsoTriangulator.cpp b/Engine/Source/Runtime/Datasmith/CADKernel/Private/CADKernel/Mesh/Meshers/IsoTriangulator.cpp
--- a/Engine/Source/Runtime/Datasmith/CADKernel/Private/CADKernel/Mesh/Meshers/IsoTriangulator.cpp
+++ b/Engine/Source/Runtime/Datasmith/CADKernel/Private/CADKernel/Mesh/Meshers/IsoTriangulator.cpp
@@ -507,6 +507,16 @@ void FIsoTriangulator::BuildInnerSegments()
 		InnerToOuterSegmentsIntersectionTool.AddSegment(Node1, Node2, Point1, Point2);
 	};
 
+	TFunction<void(const int32, const int32, const ESegmentType)> AddToInnerSegments = [&](const int32 IndexNode1, const int32 IndexNode2, const ESegmentType InType)
+	{
+		FIsoInnerNode& Node1 = *GlobalIndexToIsoInnerNodes[IndexNode1];
+		FIsoInnerNode& Node2 = *GlobalIndexToIsoInnerNodes[IndexNode2];
+		FIsoSegment& Segment = IsoSegmentFactory.New();
+		Segment.Init(Node1, Node2, InType);
+		Segment.ConnectToNode();
+		FinalInnerSegments.Add(&Segment);
+	};
+
 	TFunction<void(const int32, const int32, const ESegmentType)> BuildSegmentIfValid = [&](const int32 IndexNode1, const int32 IndexNode2, const ESegmentType InType)
 	{
 		if (Grid.IsNodeOutsideFace(IndexNode1) && Grid.IsNodeOutsideFace(IndexNode2))
@@ -532,6 +542,26 @@ void FIsoTriangulator::BuildInnerSegments()
 			return;
 		}
 
+		if (Grid.IsNodeInsideAndCloseToLoop(IndexNode1) && Grid.IsNodeInsideAndCloseToLoop(IndexNode2))
+		{
+#ifdef DEBUG_BUILDINNERSEGMENTS
+			if (bDisplay)
+			{
+				DisplaySegment(Grid.GetInner2DPoint(EGridSpace::UniformScaled, IndexNode1) * DisplayScale, Grid.GetInner2DPoint(EGridSpace::UniformScaled, IndexNode2) * DisplayScale, IndexNode1, EVisuProperty::PinkCurve);
+			}
+#endif
+			if (LoopSegmentsIntersectionTool.DoesIntersect(Grid.GetInner2DPoint(EGridSpace::UniformScaled, IndexNode1), Grid.GetInner2DPoint(EGridSpace::UniformScaled, IndexNode2))
+				|| AlmostHitsLoop(Grid.GetInner2DPoint(EGridSpace::UniformScaled, IndexNode1), Grid.GetInner2DPoint(EGridSpace::UniformScaled, IndexNode2), InType))
+			{
+				AddToInnerToOuterSegmentsIntersectionTool(IndexNode1, IndexNode2, InType);
+			}
+			else
+			{
+				AddToInnerSegments(IndexNode1, IndexNode2, InType);
+			}
+
+			return;
+		}
 
 		if (Grid.IsNodeInsideAndMeshable(IndexNode1) && Grid.IsNodeInsideAndMeshable(IndexNode2))
 		{
@@ -541,13 +571,7 @@ void FIsoTriangulator::BuildInnerSegments()
 				DisplaySegment(Grid.GetInner2DPoint(EGridSpace::UniformScaled, IndexNode1) * DisplayScale, Grid.GetInner2DPoint(EGridSpace::UniformScaled, IndexNode2) * DisplayScale, IndexNode1, EVisuProperty::BlueCurve);
 			}
 #endif
-
-			FIsoInnerNode& Node1 = *GlobalIndexToIsoInnerNodes[IndexNode1];
-			FIsoInnerNode& Node2 = *GlobalIndexToIsoInnerNodes[IndexNode2];
-			FIsoSegment& Segment = IsoSegmentFactory.New();
-			Segment.Init(Node1, Node2, InType);
-			Segment.ConnectToNode();
-			FinalInnerSegments.Add(&Segment);
+			AddToInnerSegments(IndexNode1, IndexNode2, InType);
 			return;
 		}
 
@@ -957,16 +981,16 @@ void FIsoTriangulator::FindIsoSegmentToLinkLoopToLoop()
 		Algo::Sort(SortedLoopNodesAlong, [&](const FLoopNode* LoopNode1, const FLoopNode* LoopNode2)
 			{
 				const FPoint2D& Node1Coordinates = LoopNode1->Get2DPoint(EGridSpace::UniformScaled, Grid);
-				const FPoint2D& Node2Coordinates = LoopNode2->Get2DPoint(EGridSpace::UniformScaled, Grid);
-				return (Node1Coordinates.U + (Node1Coordinates.V - VMin) * DeltaFactor) < (Node2Coordinates.U + (Node2Coordinates.V - VMin) * DeltaFactor);
+		const FPoint2D& Node2Coordinates = LoopNode2->Get2DPoint(EGridSpace::UniformScaled, Grid);
+		return (Node1Coordinates.U + (Node1Coordinates.V - VMin) * DeltaFactor) < (Node2Coordinates.U + (Node2Coordinates.V - VMin) * DeltaFactor);
 			});
 		FindIsoSegmentAlong(IsoU, IsoUCoordinates, IsoVCoordinates);
 
-		Algo::Sort(SortedLoopNodesAlong, [&](const FLoopNode* LoopNode1, const FLoopNode* LoopNode2) 
+		Algo::Sort(SortedLoopNodesAlong, [&](const FLoopNode* LoopNode1, const FLoopNode* LoopNode2)
 			{
 				const FPoint2D& Node1Coordinates = LoopNode1->Get2DPoint(EGridSpace::UniformScaled, Grid);
-				const FPoint2D& Node2Coordinates = LoopNode2->Get2DPoint(EGridSpace::UniformScaled, Grid);
-				return (Node1Coordinates.V + (Node1Coordinates.U - UMin) * DeltaFactor) < (Node2Coordinates.V + (Node2Coordinates.U - UMin) * DeltaFactor);
+		const FPoint2D& Node2Coordinates = LoopNode2->Get2DPoint(EGridSpace::UniformScaled, Grid);
+		return (Node1Coordinates.V + (Node1Coordinates.U - UMin) * DeltaFactor) < (Node2Coordinates.V + (Node2Coordinates.U - UMin) * DeltaFactor);
 			});
 		FindIsoSegmentAlong(IsoV, IsoVCoordinates, IsoUCoordinates);
 	}
@@ -1985,105 +2009,138 @@ void FIsoTriangulator::MeshCycle(const EGridSpace Space, const TArray<FIsoSegmen
 			double CandidateSlopeAtStartNode = 8.;
 			double CandidateSlopeAtEndNode = 8.;
 
-			for (FIsoNode* Node : CycleNodes)
+
+			// To avoid to create triangle crossing natural iso line (a line carrying iso edges)
+			// the iso edges connecting an inside node to an external node (vs the loop) are added to InnerToOuterSegmentsIntersectionTool
+			// the selected triangle is the triangle minimizing the number of intersection with these lines
+			int32 MaxIntersections = 1;
+			for (int32 IntersectionCountAllow = 0; IntersectionCountAllow < MaxIntersections; IntersectionCountAllow++)
 			{
-				if (Node == &StartNode)
-				{
-					continue;
-				}
-				if (Node == &EndNode)
+				for (FIsoNode* Node : CycleNodes)
 				{
-					continue;
-				}
+					if (Node == &StartNode)
+					{
+						continue;
+					}
+					if (Node == &EndNode)
+					{
+						continue;
+					}
 
-				// Check if the node is inside the sector (X) or outside (Z)
-				const FPoint2D& NodePoint2D = Node->Get2DPoint(EGridSpace::UniformScaled, Grid);
-				double SlopeAtStartNode = GetSlopeAtStartNode(StartPoint2D, NodePoint2D, StartReferenceSlope);
-				double SlopeAtEndNode = GetSlopeAtEndNode(EndPoint2D, NodePoint2D, EndReferenceSlope);
+					// Check if the node is inside the sector (X) or outside (Z)
+					const FPoint2D& NodePoint2D = Node->Get2DPoint(EGridSpace::UniformScaled, Grid);
+					double SlopeAtStartNode = GetSlopeAtStartNode(StartPoint2D, NodePoint2D, StartReferenceSlope);
+					double SlopeAtEndNode = GetSlopeAtEndNode(EndPoint2D, NodePoint2D, EndReferenceSlope);
 
-				if (Node != &PreviousNode)
-				{
-					if (SlopeAtStartNode <= EpsilonSlope || SlopeAtStartNode >= StartMaxSlope - EpsilonSlope)
+					if (Node != &PreviousNode)
 					{
-						continue;
+						if (SlopeAtStartNode <= EpsilonSlope || SlopeAtStartNode >= StartMaxSlope - EpsilonSlope)
+						{
+							continue;
+						}
 					}
-				}
 
-				if (Node != &NextNode)
-				{
-					if (SlopeAtEndNode <= EpsilonSlope || SlopeAtEndNode >= EndMaxSlope - EpsilonSlope)
+					if (Node != &NextNode)
 					{
-						continue;
+						if (SlopeAtEndNode <= EpsilonSlope || SlopeAtEndNode >= EndMaxSlope - EpsilonSlope)
+						{
+							continue;
+						}
 					}
-				}
 
 #ifdef DEBUG_FIND_BEST_TRIANGLE
-				if (TriangleIndex == 14)
-				{
-					static int32 NodeIndex = 1;
-					F3DDebugSession _(FString::Printf(TEXT("Next Node %d"), NodeIndex);
-					Grid.DisplayIsoNode(EGridSpace::UniformScaled, *Node, ++NodeIndex, EVisuProperty::YellowPoint);
-					Wait();
-				}
+					if (TriangleIndex == 14)
+					{
+						static int32 NodeIndex = 1;
+						F3DDebugSession _(FString::Printf(TEXT("Next Node %d"), NodeIndex);
+						Grid.DisplayIsoNode(EGridSpace::UniformScaled, *Node, ++NodeIndex, EVisuProperty::YellowPoint);
+						Wait();
+					}
 #endif
 
-				if (FMath::IsNearlyEqual(SlopeAtStartNode, CandidateSlopeAtStartNode, MinSlopeToNotBeAligned) && SlopeAtEndNode > CandidateSlopeAtEndNode)
-				{
-					continue;
-				}
+					if (FMath::IsNearlyEqual(SlopeAtStartNode, CandidateSlopeAtStartNode, MinSlopeToNotBeAligned) && SlopeAtEndNode > CandidateSlopeAtEndNode)
+					{
+						continue;
+					}
 
-				if (FMath::IsNearlyEqual(SlopeAtEndNode, CandidateSlopeAtEndNode, EpsilonSlope) && SlopeAtStartNode > CandidateSlopeAtStartNode)
-				{
-					continue;
-				}
+					if (FMath::IsNearlyEqual(SlopeAtEndNode, CandidateSlopeAtEndNode, EpsilonSlope) && SlopeAtStartNode > CandidateSlopeAtStartNode)
+					{
+						continue;
+					}
 
 #ifdef D3_COTANGENT_CRITERIA
 				// need more test, need to be tested with IsoscelesCriteria instead of 3D CotangentCriteria
-				const FPoint& NodePoint3D = Node->Get3DPoint(Grid);
-				FPoint NodeNormal;
-				double PointCriteria = FMath::Abs(CotangentCriteria(StartPoint3D, EndPoint3D, NodePoint3D, NodeNormal));
-				double CosAngle = FMath::Abs(ComputeCosinus(NodeNormal, Node->GetNormal(Grid)));
+					const FPoint& NodePoint3D = Node->Get3DPoint(Grid);
+					FPoint NodeNormal;
+					double PointCriteria = FMath::Abs(CotangentCriteria(StartPoint3D, EndPoint3D, NodePoint3D, NodeNormal));
+					double CosAngle = FMath::Abs(ComputeCosinus(NodeNormal, Node->GetNormal(Grid)));
 
-				// the criteria is weighted according to the cosine of the angle between the normal of the candidate triangle and the normal at the tested point
-				if (CosAngle > DOUBLE_SMALL_NUMBER)
-				{
-					PointCriteria /= CosAngle;
-				}
-				else
-				{
-					PointCriteria = HUGE_VALUE;
-				}
-#else 
-				double PointCriteria = IsoTriangulatorImpl::IsoscelesCriteria(StartPoint2D, EndPoint2D, NodePoint2D);
-#endif
-				if (
-					// the candidate triangle is inside the current candidate triangle
-					((SlopeAtStartNode < (CandidateSlopeAtStartNode + MinSlopeToNotBeAligned)) && (SlopeAtEndNode < (CandidateSlopeAtEndNode + MinSlopeToNotBeAligned)))
-					||
-					// the candidate triangle is better the current candidate triangle and doesn't contain the current candidate triangle
-					((PointCriteria < MinCriteria) && ((SlopeAtStartNode > CandidateSlopeAtStartNode) ^ (SlopeAtEndNode > CandidateSlopeAtEndNode))))
-				{
-					// check if the candidate segment is not in intersection with existing segments
-					// if the segment exist, it has already been tested
-					FIsoSegment* StartSegment = StartNode.GetSegmentConnectedTo(Node);
-					FIsoSegment* EndSegment = EndNode.GetSegmentConnectedTo(Node);
-
-					if (!StartSegment && CycleIntersectionTool.DoesIntersect(StartNode, *Node))
+					// the criteria is weighted according to the cosine of the angle between the normal of the candidate triangle and the normal at the tested point
+					if (CosAngle > DOUBLE_SMALL_NUMBER)
 					{
-						continue;
+						PointCriteria /= CosAngle;
 					}
-
-					if (!EndSegment && CycleIntersectionTool.DoesIntersect(EndNode, *Node))
+					else
 					{
-						continue;
+						PointCriteria = HUGE_VALUE;
 					}
+#else 
+					double PointCriteria = IsoTriangulatorImpl::IsoscelesCriteria(StartPoint2D, EndPoint2D, NodePoint2D);
+#endif
+					if (
+						// the candidate triangle is inside the current candidate triangle
+						((SlopeAtStartNode < (CandidateSlopeAtStartNode + MinSlopeToNotBeAligned)) && (SlopeAtEndNode < (CandidateSlopeAtEndNode + MinSlopeToNotBeAligned)))
+						||
+						// the candidate triangle is better the current candidate triangle and doesn't contain the current candidate triangle
+						((PointCriteria < MinCriteria) && ((SlopeAtStartNode > CandidateSlopeAtStartNode) ^ (SlopeAtEndNode > CandidateSlopeAtEndNode))))
+					{
+						// check if the candidate segment is not in intersection with existing segments
+						// if the segment exist, it has already been tested
+						FIsoSegment* StartSegment = StartNode.GetSegmentConnectedTo(Node);
+						FIsoSegment* EndSegment = EndNode.GetSegmentConnectedTo(Node);
 
-					MinCriteria = PointCriteria;
-					CandidatNode = Node;
-					StartToCandiatSegment = StartSegment;
-					EndToCandiatSegment = EndSegment;
-					CandidateSlopeAtStartNode = SlopeAtStartNode;
-					CandidateSlopeAtEndNode = SlopeAtEndNode;
+						if (!StartSegment && CycleIntersectionTool.DoesIntersect(StartNode, *Node))
+						{
+							continue;
+						}
+
+						if (!EndSegment && CycleIntersectionTool.DoesIntersect(EndNode, *Node))
+						{
+							continue;
+						}
+
+#ifdef DEBUG_FIND_BEST_TRIANGLE
+						if (bDisplay)
+						{
+							F3DDebugSession _(TEXT("Candidate triangle"));
+							Grid.DisplayIsoSegment(EGridSpace::UniformScaled, StartNode, *Node, 1, EVisuProperty::RedCurve);
+							Grid.DisplayIsoSegment(EGridSpace::UniformScaled, EndNode, *Node, 2, EVisuProperty::RedCurve);
+							Wait();
+						}
+#endif
+						int32 IntersectionCount = InnerToOuterSegmentsIntersectionTool.CountIntersections(StartNode, *Node);
+						IntersectionCount = FMath::Max(IntersectionCount, InnerToOuterSegmentsIntersectionTool.CountIntersections(EndNode, *Node));
+						if (IntersectionCount > IntersectionCountAllow)
+						{
+							if (IntersectionCount >= MaxIntersections)
+							{
+								MaxIntersections = IntersectionCount + 1;
+							}
+							continue;
+						}
+
+						MinCriteria = PointCriteria;
+						CandidatNode = Node;
+						StartToCandiatSegment = StartSegment;
+						EndToCandiatSegment = EndSegment;
+						CandidateSlopeAtStartNode = SlopeAtStartNode;
+						CandidateSlopeAtEndNode = SlopeAtEndNode;
+					}
+				}
+
+				if (CandidatNode != nullptr)
+				{
+					break;
 				}
 			}
 		}
@@ -2124,7 +2181,7 @@ void FIsoTriangulator::MeshCycle(const EGridSpace Space, const TArray<FIsoSegmen
 				CycleIntersectionTool.Sort();
 			}
 		}
-						};
+							};
 
 	for (int32 Index = 0; Index < SegmentStack.Num(); ++Index)
 	{
