[NO NON-PII COMMIT MESSAGE]

diff --git a/Engine/Source/Runtime/Datasmith/CADKernel/Private/CADKernel/Mesh/MeshEnum.h b/Engine/Source/Runtime/Datasmith/CADKernel/Private/CADKernel/Mesh/MeshEnum.h
--- a/Engine/Source/Runtime/Datasmith/CADKernel/Private/CADKernel/Mesh/MeshEnum.h
+++ b/Engine/Source/Runtime/Datasmith/CADKernel/Private/CADKernel/Mesh/MeshEnum.h
@@ -1,8 +1,6 @@
 // Copyright Epic Games, Inc. All Rights Reserved.
 #pragma once
 
-#include "HAL/Platform.h"
-
 namespace UE::CADKernel
 {
 /**
diff --git a/Engine/Source/Runtime/Datasmith/CADKernel/Private/CADKernel/Topo/TopologicalEdge.cpp b/Engine/Source/Runtime/Datasmith/CADKernel/Private/CADKernel/Topo/TopologicalEdge.cpp
--- a/Engine/Source/Runtime/Datasmith/CADKernel/Private/CADKernel/Topo/TopologicalEdge.cpp
+++ b/Engine/Source/Runtime/Datasmith/CADKernel/Private/CADKernel/Topo/TopologicalEdge.cpp
@@ -1182,6 +1182,32 @@ void FTopologicalEdge::Offset2D(const FPoint2D& OffsetDirection)
 	Curve->Offset2D(OffsetDirection);
 }
 
+bool FTopologicalEdge::IsConnectedTo(const FTopologicalFace* Face) const
+{
+	for (FTopologicalEdge* TwinEdge : GetTwinEntities())
+	{
+		if (TwinEdge->GetFace() == Face)
+		{
+			return true;
+		}
+	}
+	return false;
+}
+
+TArray<FTopologicalFace*> FTopologicalEdge::GetLinkedFaces() const
+{
+	TArray<FTopologicalFace*> NeighborFaces;
+	NeighborFaces.Reserve(GetTwinEntities().Num());
+
+	for (FTopologicalEdge* TwinEdge : GetTwinEntities())
+	{
+		FTopologicalFace* NeighborFace = TwinEdge->GetFace();
+		NeighborFaces.Add(NeighborFace);
+	}
+
+	return MoveTemp(NeighborFaces);
+}
+
 #ifdef CADKERNEL_DEV
 FInfoEntity& FEdgeLink::GetInfo(FInfoEntity& Info) const
 {
diff --git a/Engine/Source/Runtime/Datasmith/CADKernel/Private/CADKernel/Topo/TopologicalFace.cpp b/Engine/Source/Runtime/Datasmith/CADKernel/Private/CADKernel/Topo/TopologicalFace.cpp
--- a/Engine/Source/Runtime/Datasmith/CADKernel/Private/CADKernel/Topo/TopologicalFace.cpp
+++ b/Engine/Source/Runtime/Datasmith/CADKernel/Private/CADKernel/Topo/TopologicalFace.cpp
@@ -268,87 +268,135 @@ void FTopologicalFace::RemoveLoop(const TSharedPtr<FTopologicalLoop>& Loop)
 	}
 }
 
-void FTopologicalFace::Disjoin(TArray<FTopologicalEdge*>& NewBorderEdges)
+void FTopologicalFace::Disjoin(TArray<FTopologicalEdge*>* NewBorderEdges)
 {
-	NewBorderEdges.Reserve(NewBorderEdges.Num() + EdgeCount());
+	if(NewBorderEdges)
+	{
+		NewBorderEdges->Reserve(NewBorderEdges->Num() + EdgeCount());
+	}
+
 	for (const TSharedPtr<FTopologicalLoop>& Loop : GetLoops())
 	{
-		for (const FOrientedEdge& Edge : Loop->GetEdges())
+		for (const FOrientedEdge& OrientedEdge : Loop->GetEdges())
 		{
-			const TArray<FTopologicalEdge*> Twins = Edge.Entity->GetTwinEntities();
+			FTopologicalEdge* Edge = OrientedEdge.Entity.Get();
+			const TArray<FTopologicalEdge*> Twins = Edge->GetTwinEntities();
 			for (FTopologicalEdge* TwinEdge : Twins)
 			{
-				if (TwinEdge != Edge.Entity.Get())
+				if (NewBorderEdges && TwinEdge != Edge)
 				{
-					NewBorderEdges.Add(TwinEdge);
+					NewBorderEdges->Add(TwinEdge);
 				}
 			}
-			Edge.Entity->RemoveFromLink();
-			Edge.Entity->SetMarker1();
-			Edge.Entity->GetStartVertex()->RemoveFromLink();
-			Edge.Entity->GetEndVertex()->RemoveFromLink();
+			Edge->RemoveFromLink();
+			Edge->SetMarker1();
+			Edge->GetStartVertex()->RemoveFromLink();
+			Edge->GetEndVertex()->RemoveFromLink();
 		}
 	}
 }
 
-bool FTopologicalFace::HasSameBoundariesAs(const TSharedPtr<FTopologicalFace>& OtherFace) const
+bool FTopologicalFace::HasSameBoundariesAs(const FTopologicalFace* OtherFace) const
 {
-	int32 EdgeCount = 0;
 	for (const TSharedPtr<FTopologicalLoop>& Loop : GetLoops())
 	{
-		for (const FOrientedEdge& Edge : Loop->GetEdges())
+		for (const FOrientedEdge& OrientedEdge : Loop->GetEdges())
 		{
-			if (Edge.Entity->IsDegenerated())
+			const FTopologicalEdge& Edge = *OrientedEdge.Entity;
+			if (!Edge.IsDegenerated() && !Edge.IsBorder() && !Edge.IsConnectedTo(OtherFace))
 			{
-				continue;
+				return false;
 			}
-			Edge.Entity->GetLinkActiveEntity()->SetMarker1();
-			++EdgeCount;
 		}
 	}
+	return true;
+}
+
+bool FTopologicalFace::IsADuplicatedFace() const
+{
+	ensureCADKernel(!GetLoops().Num());
 
-	bool bSameBoundary = true;
-	int32 OtherFaceEdgeCount = 0;
-	for (const TSharedPtr<FTopologicalLoop>& Loop : OtherFace->GetLoops())
+	// Find in the adjacent faces of the first (surface or non manifold) edge, a face with the same loops
+	
+	TArray<FTopologicalFace*> AdjacentFaces;
+	const TSharedPtr<FTopologicalLoop>& Loop = GetLoops()[0];
+	for (const FOrientedEdge& OrientedEdge : Loop->GetEdges())
+	{
+		const FTopologicalEdge& Edge = *OrientedEdge.Entity;
+		if (!Edge.IsDegenerated() && !Edge.IsBorder())
+		{
+			AdjacentFaces = Edge.GetLinkedFaces();
+			break;
+		}
+	}
+
+	for (FTopologicalFace* AdjacentFace : AdjacentFaces)
+	{
+		if (AdjacentFace == this)
+		{
+			continue;
+		}
+
+		if (HasSameBoundariesAs(AdjacentFace))
+		{
+			return true;
+		}
+	}
+
+	return false;
+}
+
+
+bool FTopologicalFace::IsANonManifoldFace() const
+{
+	for (const TSharedPtr<FTopologicalLoop>& Loop : GetLoops())
 	{
-		OtherFaceEdgeCount += Loop->EdgeCount();
 		for (const FOrientedEdge& Edge : Loop->GetEdges())
 		{
-			if (Edge.Entity->IsDegenerated())
+			if (Edge.Entity->GetTwinEntityCount() > 2)
 			{
-				continue;
+				return true;
 			}
-			if (!Edge.Entity->GetLinkActiveEntity()->HasMarker1())
-			{
-				bSameBoundary = false;
-				break;
-			}
-			++OtherFaceEdgeCount;
 		}
 	}
+	return false;
+}
 
+bool FTopologicalFace::IsABorderFace() const
+{
 	for (const TSharedPtr<FTopologicalLoop>& Loop : GetLoops())
 	{
-		EdgeCount += Loop->EdgeCount();
 		for (const FOrientedEdge& Edge : Loop->GetEdges())
 		{
-			Edge.Entity->GetLinkActiveEntity()->ResetMarkers();
+			if (Edge.Entity->GetTwinEntityCount() == 1)
+			{
+				return true;
+			}
 		}
 	}
+	return false;
+}
 
-	if (EdgeCount != OtherFaceEdgeCount)
+bool FTopologicalFace::IsAFullyNonManifoldFace() const
+{
+	for (const TSharedPtr<FTopologicalLoop>& Loop : GetLoops())
 	{
-		bSameBoundary = false;
+		for (const FOrientedEdge& Edge : Loop->GetEdges())
+		{
+			if (!Edge.Entity->IsDegenerated() && Edge.Entity->GetTwinEntityCount() < 3)
+			{
+				return false;
+			}
+		}
 	}
-
-	return bSameBoundary;
+	return true;
 }
 
 const FTopologicalEdge* FTopologicalFace::GetLinkedEdge(const FTopologicalEdge& LinkedEdge) const
 {
 	for (FTopologicalEdge* TwinEdge : LinkedEdge.GetTwinEntities())
 	{
-		if (&*TwinEdge->GetLoop()->GetFace() == this)
+		if (&*TwinEdge->GetFace() == this)
 		{
 			return TwinEdge;
 		}
@@ -465,17 +513,6 @@ void FTopologicalFace::ChooseFinalDeltaUs()
 	ChooseFinalDeltas(CrossingPointDeltaMins[EIso::IsoV], CrossingPointDeltaMaxs[EIso::IsoV]);
 }
 
-// =========================================================================================================================================================================================================
-// =========================================================================================================================================================================================================
-// =========================================================================================================================================================================================================
-//
-//
-//                                                                            NOT YET REVIEWED
-//
-//
-// =========================================================================================================================================================================================================
-// =========================================================================================================================================================================================================
-// =========================================================================================================================================================================================================
 
 // Quad ==============================================================================================================================================================================================================================
 
@@ -583,7 +620,7 @@ void FTopologicalFace::DefineSurfaceType()
 						{
 							continue;
 						}
-						Neighbor = Edge->GetLoop()->GetFace();
+						Neighbor = Edge->GetFace();
 					}
 				}
 
@@ -661,8 +698,6 @@ const TSharedPtr<FTopologicalLoop> FTopologicalFace::GetExternalLoop() const
 	return TSharedPtr<FTopologicalLoop>();
 }
 
-
-
 void FFaceSubset::SetMainShell(TMap<FTopologicalShapeEntity*, int32>& ShellToFaceCount)
 {
 	if (ShellToFaceCount.Num() == 0)
diff --git a/Engine/Source/Runtime/Datasmith/CADKernel/Private/CADKernel/Topo/Topomaker.cpp b/Engine/Source/Runtime/Datasmith/CADKernel/Private/CADKernel/Topo/Topomaker.cpp
--- a/Engine/Source/Runtime/Datasmith/CADKernel/Private/CADKernel/Topo/Topomaker.cpp
+++ b/Engine/Source/Runtime/Datasmith/CADKernel/Private/CADKernel/Topo/Topomaker.cpp
@@ -845,10 +845,8 @@ void FTopomaker::RemoveThinFaces(TArray<FTopologicalEdge*>& NewBorderEdges)
 		double GapSize = 0;
 		if (Analyer.IsThinFace(GapSize))
 		{
-			Face->Disjoin(NewBorderEdges);
+			Face->Remove(&NewBorderEdges);
 			DeletedFaces.Add(Face.Get());
-			Face->Delete();
-			Face->RemoveOfHost();
 #ifdef CADKERNEL_DEV
 			Report.AddThinFace();
 #endif
diff --git a/Engine/Source/Runtime/Datasmith/CADKernel/Public/CADKernel/Math/Boundary.h b/Engine/Source/Runtime/Datasmith/CADKernel/Public/CADKernel/Math/Boundary.h
--- a/Engine/Source/Runtime/Datasmith/CADKernel/Public/CADKernel/Math/Boundary.h
+++ b/Engine/Source/Runtime/Datasmith/CADKernel/Public/CADKernel/Math/Boundary.h
@@ -117,7 +117,7 @@ struct CADKERNEL_API FLinearBoundary
 	bool IsDegenerated() const
 	{
 		double DeltaU = (Max - Min);
-		return (DeltaU < DOUBLE_KINDA_SMALL_NUMBER);
+		return (DeltaU < DOUBLE_SMALL_NUMBER);
 	}
 
 	/**
diff --git a/Engine/Source/Runtime/Datasmith/CADKernel/Public/CADKernel/Topo/TopologicalEdge.h b/Engine/Source/Runtime/Datasmith/CADKernel/Public/CADKernel/Topo/TopologicalEdge.h
--- a/Engine/Source/Runtime/Datasmith/CADKernel/Public/CADKernel/Topo/TopologicalEdge.h
+++ b/Engine/Source/Runtime/Datasmith/CADKernel/Public/CADKernel/Topo/TopologicalEdge.h
@@ -728,7 +728,7 @@ public:
 	/**
 	 * @return true if the edge is adjacent to only one surface (its carrier surface)
 	 */
-	bool IsBorder()
+	bool IsBorder() const
 	{
 		return GetTwinEntityCount() == 1;
 	}
@@ -736,11 +736,15 @@ public:
 	/**
 	 * @return true if the edge is adjacent to only two surfaces
 	 */
-	bool IsSurfacic()
+	bool IsSurfacic() const
 	{
 		return GetTwinEntityCount() == 2;
 	}
 
+	bool IsConnectedTo(const FTopologicalFace* Face) const;
+
+	TArray<FTopologicalFace*> GetLinkedFaces() const;
+
 	/**
 	 * Merge successive edges of a face in a single edge.
 	 * Edges must not be topological linked to other faces i.e. edges must be border edges
diff --git a/Engine/Source/Runtime/Datasmith/CADKernel/Public/CADKernel/Topo/TopologicalFace.h b/Engine/Source/Runtime/Datasmith/CADKernel/Public/CADKernel/Topo/TopologicalFace.h
--- a/Engine/Source/Runtime/Datasmith/CADKernel/Public/CADKernel/Topo/TopologicalFace.h
+++ b/Engine/Source/Runtime/Datasmith/CADKernel/Public/CADKernel/Topo/TopologicalFace.h
@@ -158,7 +158,6 @@ public:
 	{
 	}
 
-
 	// ======   Loop Functions   ======
 
 	void RemoveLoop(const TSharedPtr<FTopologicalLoop>& Loop);
@@ -284,16 +283,38 @@ public:
 
 	// ======   Topo Functions   ======
 
+	/** 
+	 * Delete a face i.e. disjoin the face to all its linked elements before setting it as deleted
+	 */
+	void Remove(TArray<FTopologicalEdge*>* NewBorderEdges = nullptr)
+	{
+		Disjoin(NewBorderEdges);
+		RemoveOfHost();
+		SetDeleted();
+	}
+
+	/** Has at least one non manifold edge */
+	bool IsANonManifoldFace() const;
+
+	/** Has at least one border edge */
+	bool IsABorderFace() const;
+
+	/** All its edges are non manifold */
+	bool IsAFullyNonManifoldFace() const;
+
+	/** one adjacent face shares the same loops */
+	bool IsADuplicatedFace() const;
+
 	/**
 	 * Checks if the face and the other face have the same boundaries i.e. each non degenerated edge is linked to an edge of the other face
 	 */
-	bool HasSameBoundariesAs(const TSharedPtr<FTopologicalFace>& OtherFace) const;
+	bool HasSameBoundariesAs(const FTopologicalFace* OtherFace) const;
 
 	/**
 	 * Disconnects the face of its neighbors i.e. remove topological edge and vertex link with its neighbors
 	 * @param OutNewBorderEdges the neighbors edges
 	 */
-	void Disjoin(TArray<FTopologicalEdge*>& OutNewBorderEdges);
+	void Disjoin(TArray<FTopologicalEdge*>* NewBorderEdges = nullptr);
 
 #ifdef CADKERNEL_DEV
 	virtual void FillTopologyReport(FTopologyReport& Report) const override;
