[NO NON-PII COMMIT MESSAGE]

diff --git a/Engine/Plugins/Enterprise/DatasmithCADImporter/Source/CADInterfaces/Private/TechSoftFileParserCADKernelTessellator.cpp b/Engine/Plugins/Enterprise/DatasmithCADImporter/Source/CADInterfaces/Private/TechSoftFileParserCADKernelTessellator.cpp
--- a/Engine/Plugins/Enterprise/DatasmithCADImporter/Source/CADInterfaces/Private/TechSoftFileParserCADKernelTessellator.cpp
+++ b/Engine/Plugins/Enterprise/DatasmithCADImporter/Source/CADInterfaces/Private/TechSoftFileParserCADKernelTessellator.cpp
@@ -28,8 +28,6 @@ namespace CADLibrary
 FTechSoftFileParserCADKernelTessellator::FTechSoftFileParserCADKernelTessellator(FCADFileData& InCADData, const FString& EnginePluginsPath)
 	: FTechSoftFileParser(InCADData, EnginePluginsPath)
 	, LastHostIdUsed(1 << 30)
-	, GeometricTolerance(FImportParameters::GStitchingTolerance * 10) // cm to mm
-	, ForceFactor(FImportParameters::GStitchingForceFactor)
 {
 }
 
@@ -90,8 +88,7 @@ void FTechSoftFileParserCADKernelTessellator::SewAndMesh(TArray<A3DRiRepresentat
 
 	using namespace UE::CADKernel;
 
-	const bool bTryToImproveSew = FImportParameters::bGStitchingForceSew;
-	const bool bRemoveThinSurface = FImportParameters::bGStitchingRemoveThinFaces;
+	const double GeometricTolerance = FImportParameters::GStitchingTolerance * 10; // cm to mm
 
 	FSession CADKernelSession(GeometricTolerance);
 	CADKernelSession.SetFirstNewHostId(LastHostIdUsed);
@@ -112,15 +109,17 @@ void FTechSoftFileParserCADKernelTessellator::SewAndMesh(TArray<A3DRiRepresentat
 	}
 
 	// Sew if needed
-	FTopomaker Topomaker(CADKernelSession, GeometricTolerance, ForceFactor);
-	Topomaker.Sew(bTryToImproveSew, bRemoveThinSurface);
+	UE::CADKernel::FTopomakerOptions TopomakerOptions((UE::CADKernel::ESewOption) SewOption::GetFromImportParameters(), GeometricTolerance, FImportParameters::GStitchingForceFactor);
+
+	FTopomaker Topomaker(CADKernelSession, TopomakerOptions);
+	Topomaker.Sew();
 	Topomaker.SplitIntoConnectedShells();
 	Topomaker.OrientShells();
 
 	// The Sew + SplitIntoConnectedShells change the bodies: some are deleted some are create
 	// but at the end the count of body is always <= than the initial count
 	// We need to found the unchanged bodies to link them to their FArchiveBody
-	// The new bodies will be linked to FArchiveBody of deleted bodies but the metadata of these FArchiveBody havbe to be cleaned
+	// The new bodies will be linked to FArchiveBody of deleted bodies but the metadata of these FArchiveBody have to be cleaned
 
 	int32 BodyCount = CADKernelModel.GetBodies().Num();
 
@@ -211,6 +210,7 @@ void FTechSoftFileParserCADKernelTessellator::GenerateBodyMesh(A3DRiRepresentati
 	TRACE_CPUPROFILER_EVENT_SCOPE(FTechSoftFileParserCADKernelTessellator::GenerateBodyMesh);
 	using namespace UE::CADKernel;
 
+	const double GeometricTolerance = FImportParameters::GStitchingTolerance * 10; // cm to mm
 	FSession CADKernelSession(GeometricTolerance);
 	CADKernelSession.SetFirstNewHostId(LastHostIdUsed);
 	FModel& CADKernelModel = CADKernelSession.GetModel();
@@ -226,11 +226,10 @@ void FTechSoftFileParserCADKernelTessellator::GenerateBodyMesh(A3DRiRepresentati
 
 	if (CADFileData.GetImportParameters().GetStitchingTechnique() == StitchingHeal)
 	{
-		const bool bTryToImproveSew = FImportParameters::bGStitchingForceSew;
-		const bool bRemoveThinSurface = FImportParameters::bGStitchingRemoveThinFaces;
+		UE::CADKernel::FTopomakerOptions TopomakerOptions((UE::CADKernel::ESewOption)SewOption::GetFromImportParameters(), GeometricTolerance, FImportParameters::GStitchingForceFactor);
 
-		FTopomaker Topomaker(CADKernelSession, GeometricTolerance, ForceFactor);
-		Topomaker.Sew(bTryToImproveSew, bRemoveThinSurface);
+		FTopomaker Topomaker(CADKernelSession, TopomakerOptions);
+		Topomaker.Sew();
 		Topomaker.OrientShells();
 	}
 
diff --git a/Engine/Plugins/Enterprise/DatasmithCADImporter/Source/CADInterfaces/Private/TechSoftFileParserCADKernelTessellator.h b/Engine/Plugins/Enterprise/DatasmithCADImporter/Source/CADInterfaces/Private/TechSoftFileParserCADKernelTessellator.h
--- a/Engine/Plugins/Enterprise/DatasmithCADImporter/Source/CADInterfaces/Private/TechSoftFileParserCADKernelTessellator.h
+++ b/Engine/Plugins/Enterprise/DatasmithCADImporter/Source/CADInterfaces/Private/TechSoftFileParserCADKernelTessellator.h
@@ -47,12 +47,7 @@ private:
 
 private:
 
-	FCadId LastEntityId = 1;
 	int32 LastHostIdUsed = 0;
-
-	const double GeometricTolerance = 0.01; // mm
-	const double ForceFactor = 5.;
-
 };
 
 } // ns CADLibrary
\ No newline at end of file
diff --git a/Engine/Plugins/Enterprise/DatasmithCADImporter/Source/CADInterfaces/Public/CADFileData.h b/Engine/Plugins/Enterprise/DatasmithCADImporter/Source/CADInterfaces/Public/CADFileData.h
--- a/Engine/Plugins/Enterprise/DatasmithCADImporter/Source/CADInterfaces/Public/CADFileData.h
+++ b/Engine/Plugins/Enterprise/DatasmithCADImporter/Source/CADInterfaces/Public/CADFileData.h
@@ -43,6 +43,7 @@ public:
 			GeomFileHash = HashCombine(GeomFileHash, ::GetTypeHash(FImportParameters::GStitchingTolerance));
 			GeomFileHash = HashCombine(GeomFileHash, ::GetTypeHash(FImportParameters::bGStitchingForceSew));
 			GeomFileHash = HashCombine(GeomFileHash, ::GetTypeHash(FImportParameters::bGStitchingRemoveThinFaces));
+			GeomFileHash = HashCombine(GeomFileHash, ::GetTypeHash(FImportParameters::bGStitchingRemoveDuplicatedFaces));
 			GeomFileHash = HashCombine(GeomFileHash, ::GetTypeHash(FImportParameters::GStitchingForceFactor));
 		}
 		return GeomFileHash;
diff --git a/Engine/Plugins/Enterprise/DatasmithCADImporter/Source/CADKernelSurface/Public/CADModelToCADKernelConverterBase.h b/Engine/Plugins/Enterprise/DatasmithCADImporter/Source/CADKernelSurface/Public/CADModelToCADKernelConverterBase.h
--- a/Engine/Plugins/Enterprise/DatasmithCADImporter/Source/CADKernelSurface/Public/CADModelToCADKernelConverterBase.h
+++ b/Engine/Plugins/Enterprise/DatasmithCADImporter/Source/CADKernelSurface/Public/CADModelToCADKernelConverterBase.h
@@ -42,12 +42,11 @@ public:
 		// Apply stitching if applicable
 		if(ImportParameters.GetStitchingTechnique() != StitchingNone)
 		{
-			const double JoiningTolerance = FImportParameters::GStitchingTolerance * 10.; //CM to MM
-			const double ForceFactor = FImportParameters::GStitchingForceFactor;
-			bool bForceJoining = FImportParameters::bGStitchingForceSew;
-			bool bRemoveThinFaces = FImportParameters::bGStitchingRemoveThinFaces;
-			UE::CADKernel::FTopomaker Topomaker(CADKernelSession, JoiningTolerance, ForceFactor);
-			Topomaker.Sew(bForceJoining, bRemoveThinFaces);
+			const double StitchingTolerance = FImportParameters::GStitchingTolerance * 10.; //CM to MM
+			UE::CADKernel::FTopomakerOptions TopomakerOptions((UE::CADKernel::ESewOption)SewOption::GetFromImportParameters(), StitchingTolerance, FImportParameters::GStitchingForceFactor);
+
+			UE::CADKernel::FTopomaker Topomaker(CADKernelSession, TopomakerOptions);
+			Topomaker.Sew();
 			Topomaker.OrientShells();
 		}
 
diff --git a/Engine/Plugins/Enterprise/DatasmithCADImporter/Source/CADTools/Private/CADOption.cpp b/Engine/Plugins/Enterprise/DatasmithCADImporter/Source/CADTools/Private/CADOption.cpp
--- a/Engine/Plugins/Enterprise/DatasmithCADImporter/Source/CADTools/Private/CADOption.cpp
+++ b/Engine/Plugins/Enterprise/DatasmithCADImporter/Source/CADTools/Private/CADOption.cpp
@@ -90,6 +90,15 @@ Only available with CADKernel.\n\
 Default value of RemoveThinFaces is true\n"),
 ECVF_Default);
 
+bool FImportParameters::bGStitchingRemoveDuplicatedFaces = true;
+FAutoConsoleVariableRef GCADTranslatorStitchingRemoveDuplicatedFaces(
+	TEXT("ds.CADTranslator.Stitching.RemoveDuplicatedFaces"),
+	FImportParameters::bGStitchingRemoveDuplicatedFaces,
+	TEXT("During the welding process, duplicated faces (i.e. faces with the same loops) are removed. \n\
+Only available with CADKernel.\n\
+Default value of RemoveDuplicatedFaces is true\n"),
+ECVF_Default);
+
 float FImportParameters::GUnitScale = 1.f;
 FAutoConsoleVariableRef GCADTranslatorUnitScale(
 	TEXT("ds.CADTranslator.UnitScale"),
diff --git a/Engine/Plugins/Enterprise/DatasmithCADImporter/Source/CADTools/Public/CADOptions.h b/Engine/Plugins/Enterprise/DatasmithCADImporter/Source/CADTools/Public/CADOptions.h
--- a/Engine/Plugins/Enterprise/DatasmithCADImporter/Source/CADTools/Public/CADOptions.h
+++ b/Engine/Plugins/Enterprise/DatasmithCADImporter/Source/CADTools/Public/CADOptions.h
@@ -19,6 +19,19 @@ namespace CADLibrary
 		StitchingSew,
 	};
 
+	enum class ESewOption : uint8  // Same as UE::CADKernel::ESewOption
+	{
+		None = 0x00u,	// No flags.
+
+		ForceJoining = 0x01u,
+		RemoveThinFaces = 0x02u,
+		RemoveDuplicatedFaces = 0x04u,
+
+		All = 0x07u
+	};
+
+	ENUM_CLASS_FLAGS(ESewOption);
+
 	class FImportParameters
 	{
 	private:
@@ -39,6 +52,7 @@ namespace CADLibrary
 		CADTOOLS_API static float GStitchingTolerance;
 		CADTOOLS_API static bool bGStitchingForceSew;
 		CADTOOLS_API static bool bGStitchingRemoveThinFaces;
+		CADTOOLS_API static bool bGStitchingRemoveDuplicatedFaces;
 		CADTOOLS_API static float GStitchingForceFactor;
 		CADTOOLS_API static float GUnitScale;
 		CADTOOLS_API static float GMeshingParameterFactor;
@@ -73,7 +87,7 @@ namespace CADLibrary
 			{
 				Hash = HashCombine(Hash, ::GetTypeHash(Param));
 			}
-			for (bool Param : {bGStitchingForceSew, bGStitchingRemoveThinFaces, bGDisableCADKernelTessellation, bGPreferJtFileEmbeddedTessellation})
+			for (bool Param : {bGStitchingForceSew, bGStitchingRemoveThinFaces, bGStitchingRemoveDuplicatedFaces, bGDisableCADKernelTessellation, bGPreferJtFileEmbeddedTessellation})
 			{
 				Hash = HashCombine(Hash, ::GetTypeHash(Param));
 			}
@@ -98,6 +112,7 @@ namespace CADLibrary
 			Ar << ImportParameters.GStitchingTolerance;
 			Ar << ImportParameters.bGStitchingForceSew;
 			Ar << ImportParameters.bGStitchingRemoveThinFaces;
+			Ar << ImportParameters.bGStitchingRemoveDuplicatedFaces;
 			Ar << ImportParameters.GStitchingForceFactor;
 			Ar << ImportParameters.GMaxMaterialCountPerMesh;
 			return Ar;
@@ -136,6 +151,33 @@ namespace CADLibrary
 		CADTOOLS_API friend uint32 GetTypeHash(const FImportParameters& ImportParameters);
 	};
 
+	namespace SewOption
+	{
+
+	static ESewOption GetFromImportParameters()
+	{
+		ESewOption Option = ESewOption::None;
+
+		if (FImportParameters::bGStitchingForceSew)
+		{
+			Option |= ESewOption::ForceJoining;
+		}
+
+		if (FImportParameters::bGStitchingRemoveThinFaces)
+		{
+			Option |= ESewOption::RemoveThinFaces;
+		}
+
+		if (FImportParameters::bGStitchingRemoveDuplicatedFaces)
+		{
+			Option |= ESewOption::RemoveDuplicatedFaces;
+		}
+
+		return Option;
+	}
+
+	} // ns SewOption
+
 	inline FString BuildCadCachePath(const TCHAR* CachePath, uint32 FileHash)
 	{
 		FString FileName = FString::Printf(TEXT("UEx%08x"), FileHash) + TEXT(".prc");
diff --git a/Engine/Plugins/Enterprise/DatasmithCADImporter/Source/DatasmithCADTranslator/Private/DatasmithCADTranslator.cpp b/Engine/Plugins/Enterprise/DatasmithCADImporter/Source/DatasmithCADTranslator/Private/DatasmithCADTranslator.cpp
--- a/Engine/Plugins/Enterprise/DatasmithCADImporter/Source/DatasmithCADTranslator/Private/DatasmithCADTranslator.cpp
+++ b/Engine/Plugins/Enterprise/DatasmithCADImporter/Source/DatasmithCADTranslator/Private/DatasmithCADTranslator.cpp
@@ -177,9 +177,10 @@ bool FDatasmithCADTranslator::LoadScene(TSharedRef<IDatasmithScene> DatasmithSce
 	if (!FImportParameters::bGDisableCADKernelTessellation)
 	{
 		UE_LOG(LogCADTranslator, Display, TEXT("     - Stitching Options:"));
-		UE_LOG(LogCADTranslator, Display, TEXT("         - ForceSew:        %s"), ImportParameters.bGStitchingForceSew ? TEXT("True") : TEXT("False"));
-		UE_LOG(LogCADTranslator, Display, TEXT("         - RemoveThinFaces: %s"), ImportParameters.bGStitchingRemoveThinFaces ? TEXT("True") : TEXT("False"));
-		UE_LOG(LogCADTranslator, Display, TEXT("         - ForceFactor:     %f"), ImportParameters.GStitchingForceFactor);
+		UE_LOG(LogCADTranslator, Display, TEXT("         - ForceSew:              %s"), ImportParameters.bGStitchingForceSew ? TEXT("True") : TEXT("False"));
+		UE_LOG(LogCADTranslator, Display, TEXT("         - RemoveThinFaces:       %s"), ImportParameters.bGStitchingRemoveThinFaces ? TEXT("True") : TEXT("False"));
+		UE_LOG(LogCADTranslator, Display, TEXT("         - RemoveDuplicatedFaces: %s"), ImportParameters.bGStitchingRemoveDuplicatedFaces ? TEXT("True") : TEXT("False"));
+		UE_LOG(LogCADTranslator, Display, TEXT("         - ForceFactor:           %f"), ImportParameters.GStitchingForceFactor);
 	}
 
 	switch (FileDescriptor.GetFileFormat())
diff --git a/Engine/Source/Runtime/Datasmith/CADKernel/Private/CADKernel/Topo/Topomaker.cpp b/Engine/Source/Runtime/Datasmith/CADKernel/Private/CADKernel/Topo/Topomaker.cpp
--- a/Engine/Source/Runtime/Datasmith/CADKernel/Private/CADKernel/Topo/Topomaker.cpp
+++ b/Engine/Source/Runtime/Datasmith/CADKernel/Private/CADKernel/Topo/Topomaker.cpp
@@ -22,6 +22,24 @@
 namespace UE::CADKernel
 {
 
+namespace SewOption
+{
+static bool IsForceJoining(ESewOption SewOptions)
+{
+	return (SewOptions & ESewOption::ForceJoining) == ESewOption::ForceJoining;
+}
+
+static bool IsRemoveThinFaces(ESewOption SewOptions)
+{
+	return (SewOptions & ESewOption::RemoveThinFaces) == ESewOption::RemoveThinFaces;
+}
+
+static bool IsRemoveDuplicatedFaces(ESewOption SewOptions)
+{
+	return (SewOptions & ESewOption::RemoveDuplicatedFaces) == ESewOption::RemoveDuplicatedFaces;
+}
+} // namespace
+
 namespace TopomakerTools
 {
 
@@ -398,10 +416,10 @@ void MergeCoincidentEdges(TArray<TSharedPtr<FTopologicalVertex>>& VerticesToProc
 
 } // namespace TopomakerTools
 
-FTopomaker::FTopomaker(FSession& InSession, double InTolerance, double InForceFactor)
+FTopomaker::FTopomaker(FSession& InSession, const FTopomakerOptions& InOptions)
 	: Session(InSession)
 {
-	SetTolerance(InTolerance, InForceFactor);
+	SetTolerance(InOptions);
 
 	int32 ShellCount = 0;
 	for (const TSharedPtr<FBody>& Body : Session.GetModel().GetBodies())
@@ -421,17 +439,17 @@ FTopomaker::FTopomaker(FSession& InSession, double InTolerance, double InForceFa
 	InitFaces();
 }
 
-FTopomaker::FTopomaker(FSession& InSession, const TArray<TSharedPtr<FTopologicalFace>>& InFaces, double InTolerance, double InForceFactor)
+FTopomaker::FTopomaker(FSession& InSession, const TArray<TSharedPtr<FTopologicalFace>>& InFaces, const FTopomakerOptions& InOptions)
 	: Session(InSession)
 	, Faces(InFaces)
 {
-	SetTolerance(InTolerance, InForceFactor);
+	SetTolerance(InOptions);
 }
 
-FTopomaker::FTopomaker(FSession& InSession, const TArray<TSharedPtr<FShell>>& InShells, double InTolerance, double InForceFactor)
+FTopomaker::FTopomaker(FSession& InSession, const TArray<TSharedPtr<FShell>>& InShells, const FTopomakerOptions& InOptions)
 	: Session(InSession)
 {
-	SetTolerance(InTolerance, InForceFactor);
+	SetTolerance(InOptions);
 
 	Shells.Reserve(InShells.Num());
 	for (const TSharedPtr<FShell>& Shell : InShells)
@@ -515,7 +533,7 @@ void FTopomaker::EmptyShells()
 	}
 }
 
-void FTopomaker::Sew(bool bForceJoining, bool bRemoveThinFaces)
+void FTopomaker::Sew()
 {
 	FTimePoint StartJoinTime = FChrono::Now();
 
@@ -528,7 +546,7 @@ void FTopomaker::Sew(bool bForceJoining, bool bRemoveThinFaces)
 	// basic case: merge pair of edges connected together at both extremities i.e. pair of edges with same extremity active vertices.
 	TopomakerTools::MergeCoincidentEdges(BorderVertices, EdgeLengthTolerance);
 
-	if (bRemoveThinFaces)
+	if (SewOption::IsRemoveThinFaces(SewOptions))
 	{
 		TArray<FTopologicalEdge*> NewBorderEdges;
 		RemoveThinFaces(NewBorderEdges);
@@ -536,7 +554,7 @@ void FTopomaker::Sew(bool bForceJoining, bool bRemoveThinFaces)
 
 	MergeUnconnectedSuccessiveEdges();
 
-	if (bForceJoining)
+	if (SewOption::IsForceJoining(SewOptions))
 	{
 		BorderVertices.Empty(BorderVertices.Num());
 		GetBorderVertices(BorderVertices);
@@ -545,12 +563,19 @@ void FTopomaker::Sew(bool bForceJoining, bool bRemoveThinFaces)
 		CheckSelfConnectedEdge(LargeEdgeLengthTolerance, BorderVertices);
 	}
 
+	const bool bForceJoining = SewOption::IsForceJoining(SewOptions);
+
 	// re process with new edges from MergeUnconnectedSuccessiveEdges and new merged vertices (if bForceJoining)
 	TopomakerTools::MergeCoincidentEdges(BorderVertices, bForceJoining ? LargeEdgeLengthTolerance : EdgeLengthTolerance);
 
 	// advance case: two partially coincident edges connected at one extremity i.e. coincident along the shortest edge
 	TopomakerTools::StitchParallelEdges(BorderVertices, bForceJoining ? SewToleranceToForceJoin : SewTolerance, bForceJoining ? LargeEdgeLengthTolerance : EdgeLengthTolerance);
 
+	if (SewOption::IsRemoveDuplicatedFaces(SewOptions))
+	{
+		RemoveDuplicatedFaces();
+	}
+
 #ifdef CADKERNEL_DEV
 	Report.SewDuration = FChrono::Elapse(StartJoinTime);
 	FChrono::PrintClockElapse(EVerboseLevel::Log, TEXT(""), TEXT("Sew"), Report.SewDuration);
@@ -1199,4 +1224,112 @@ void FTopomaker::OrientShells()
 #endif
 }
 
+void FTopomaker::RemoveDuplicatedFaces()
+{
+	FTimePoint StartTime = FChrono::Now();
+
+#ifdef DEBUG_REMOVE_DUPLICATED_FACES
+	F3DDebugSession _(TEXT("DOUBLE FACE"));
+#endif
+
+	TArray<FTopologicalFace*> NonManifoldFaces;
+	NonManifoldFaces.Reserve(Faces.Num() / 10);
+	for (TSharedPtr<FTopologicalFace>& FacePtr : Faces)
+	{
+		FTopologicalFace& Face = *FacePtr;
+		if (Face.IsDeleted())
+		{
+			continue;
+		}
+
+		if (Face.IsANonManifoldFace())
+		{
+			if (Face.IsAFullyNonManifoldFace())
+			{
+				if (Face.IsADuplicatedFace())
+				{
+#ifdef DEBUG_REMOVE_DUPLICATED_FACES
+					F3DDebugSession A(*FString::Printf(TEXT("Duplicated 1 %d"), Face.GetId()));
+					Display(Face);
+#endif
+					Face.Delete();
+#ifdef CADKERNEL_DEV
+					Report.AddDuplicatedFace();
+#endif
+				}
+			}
+			else
+			{
+				NonManifoldFaces.Add(&Face);
+			}
+		}
+	}
+
+
+	// Step 2: Process of Face with NonManifold and without border edges
+	for (FTopologicalFace* Face : NonManifoldFaces)
+	{
+		if (Face->IsDeleted() || !Face->IsANonManifoldFace())
+		{
+			continue;
+		}
+
+		if (!Face->IsABorderFace())
+		{
+			if (Face->IsADuplicatedFace())
+			{
+#ifdef DEBUG_REMOVE_DUPLICATED_FACES
+				F3DDebugSession A(*FString::Printf(TEXT("Duplicated 2 %d"), Face->GetId()));
+				Display(*Face);
+#endif
+				Face->Delete();
+#ifdef CADKERNEL_DEV
+				Report.AddDuplicatedFace();
+#endif
+			}
+		}
+	}
+
+	// Step 3: Process of the remaining non manifold Face
+	for (FTopologicalFace* Face : NonManifoldFaces)
+	{
+		if (Face->IsDeleted() || !Face->IsANonManifoldFace())
+		{
+			continue;
+		}
+
+		if (Face->IsADuplicatedFace())
+		{
+#ifdef DEBUG_REMOVE_DUPLICATED_FACES
+			F3DDebugSession A(*FString::Printf(TEXT("Border %d"), Face->GetId()));
+			Display(*Face);
+#endif
+			Face->Delete();
+#ifdef CADKERNEL_DEV
+			Report.AddNearlyDuplicatedFace();
+#endif
+		}
+	}
+
+#ifdef DEBUG_REMOVE_DUPLICATED_FACES
+	// Step 4: display the remaining Face
+	for (FTopologicalFace* Face : NonManifoldFaces)
+	{
+		if (Face->IsDeleted() || !Face->IsANonManifoldFace())
+		{
+			continue;
+		}
+
+		F3DDebugSession A(*FString::Printf(TEXT("Remaining %d"), Face->GetId()));
+		Display(*Face);
+	}
+#endif
+
+#ifdef CADKERNEL_DEV
+	Report.OrientationDuration = FChrono::Elapse(StartTime);
+	FChrono::PrintClockElapse(EVerboseLevel::Log, TEXT(""), TEXT("RemoveDuplicatedFaces"), Report.RemoveDuplicatedFacesDuration);
+#endif
+
+}
+
 }
diff --git a/Engine/Source/Runtime/Datasmith/CADKernel/Public/CADKernel/Topo/Topomaker.h b/Engine/Source/Runtime/Datasmith/CADKernel/Public/CADKernel/Topo/Topomaker.h
--- a/Engine/Source/Runtime/Datasmith/CADKernel/Public/CADKernel/Topo/Topomaker.h
+++ b/Engine/Source/Runtime/Datasmith/CADKernel/Public/CADKernel/Topo/Topomaker.h
@@ -17,6 +17,33 @@ class FTopologicalEdge;
 class FTopologicalFace;
 class FTopologicalVertex;
 
+enum class ESewOption : uint8
+{
+	None = 0x00u,	// No flags.
+
+	ForceJoining = 0x01u,
+	RemoveThinFaces = 0x02u, 
+	RemoveDuplicatedFaces = 0x04u,
+
+	All = 0x07u
+};
+
+ENUM_CLASS_FLAGS(ESewOption);
+
+struct FTopomakerOptions
+{
+	ESewOption SewOptions;
+	double Tolerance;
+	double ForceJoinFactor;
+
+	FTopomakerOptions(ESewOption InSewOptions, double InTolerance, double InForceJoinFactor)
+		: SewOptions(InSewOptions)
+		, Tolerance(InTolerance)
+		, ForceJoinFactor(InForceJoinFactor)
+	{}
+};
+
+
 class CADKERNEL_API FTopomaker
 {
 
@@ -33,14 +60,15 @@ protected:
 
 public:
 
-	FTopomaker(FSession& InSession, double InTolerance, double InForceFactor);
-	FTopomaker(FSession& InSession, const TArray<TSharedPtr<FShell>>& Shells, double InTolerance, double InForceFactor);
-	FTopomaker(FSession& InSession, const TArray<TSharedPtr<FTopologicalFace>>& Surfaces, double InTolerance, double InForceFactor);
+	FTopomaker(FSession& InSession, const FTopomakerOptions& InOptions);
+	FTopomaker(FSession& InSession, const TArray<TSharedPtr<FShell>>& Shells, const FTopomakerOptions& InOptions);
+	FTopomaker(FSession& InSession, const TArray<TSharedPtr<FTopologicalFace>>& Surfaces, const FTopomakerOptions& InOptions);
 
-	void SetTolerance(double InTolerance, double InForceFactor)
+	void SetTolerance(const FTopomakerOptions Options)
 	{
-		Tolerance = InTolerance;
-		ForceJoinFactor = InForceFactor;
+		SewOptions = Options.SewOptions;
+		Tolerance = Options.Tolerance;
+		ForceJoinFactor = Options.ForceJoinFactor;
 
 		SewTolerance = Tolerance * UE_DOUBLE_SQRT_2;
 
@@ -58,7 +86,7 @@ public:
 		ThinFaceWidth = Tolerance * ForceJoinFactor;
 	}
 
-	void Sew(bool bForceJoining, bool bRemoveThinFaces);
+	void Sew();
 
 	/**
 	 * Check topology of each body
@@ -105,6 +133,8 @@ private:
 
 	void RemoveEmptyShells();
 
+	void RemoveDuplicatedFaces();
+
 	/**
 	 * Return an array of active vertices.
 	 */
@@ -139,6 +169,8 @@ private:
 	 */
 	void MergeUnconnectedSuccessiveEdges();
 
+	ESewOption SewOptions;
+
 	double Tolerance;
 	double ForceJoinFactor;
 
