Fix Build an optimal InnerToOuterSegmentsIntersectionTool

diff --git a/Engine/Source/Runtime/Datasmith/CADKernel/Private/CADKernel/Mesh/Meshers/IsoTriangulator.cpp b/Engine/Source/Runtime/Datasmith/CADKernel/Private/CADKernel/Mesh/Meshers/IsoTriangulator.cpp
--- a/Engine/Source/Runtime/Datasmith/CADKernel/Private/CADKernel/Mesh/Meshers/IsoTriangulator.cpp
+++ b/Engine/Source/Runtime/Datasmith/CADKernel/Private/CADKernel/Mesh/Meshers/IsoTriangulator.cpp
@@ -519,26 +519,27 @@ void FIsoTriangulator::BuildInnerSegments()
 
 	TFunction<void(const int32, const int32, const ESegmentType)> BuildSegmentIfValid = [&](const int32 IndexNode1, const int32 IndexNode2, const ESegmentType InType)
 	{
-		if (Grid.IsNodeOutsideFace(IndexNode1) && Grid.IsNodeOutsideFace(IndexNode2))
+		if (Grid.IsNodeOusideFaceButClose(IndexNode1) && Grid.IsNodeOusideFaceButClose(IndexNode2))
 		{
 #ifdef DEBUG_BUILDINNERSEGMENTS
 			if (bDisplay)
 			{
-				DisplaySegment(Grid.GetInner2DPoint(EGridSpace::UniformScaled, IndexNode1) * DisplayScale, Grid.GetInner2DPoint(EGridSpace::UniformScaled, IndexNode2) * DisplayScale, IndexNode1, EVisuProperty::RedCurve);
+				DisplaySegment(Grid.GetInner2DPoint(EGridSpace::UniformScaled, IndexNode1) * DisplayScale, Grid.GetInner2DPoint(EGridSpace::UniformScaled, IndexNode2) * DisplayScale, IndexNode1, EVisuProperty::PinkCurve);
 			}
 #endif
+			AddToInnerToOuterSegmentsIntersectionTool(IndexNode1, IndexNode2, InType);
 			return;
 		}
 
-		if (Grid.IsNodeCloseToButOusideFace(IndexNode1) && Grid.IsNodeCloseToButOusideFace(IndexNode2))
+		if (Grid.IsNodeOutsideFace(IndexNode1) && Grid.IsNodeOutsideFace(IndexNode2))
 		{
 #ifdef DEBUG_BUILDINNERSEGMENTS
 			if (bDisplay)
 			{
-				DisplaySegment(Grid.GetInner2DPoint(EGridSpace::UniformScaled, IndexNode1) * DisplayScale, Grid.GetInner2DPoint(EGridSpace::UniformScaled, IndexNode2) * DisplayScale, IndexNode1, EVisuProperty::PinkCurve);
+				F3DDebugSession _(bDisplay, TEXT("Segment"));
+				DisplaySegment(Grid.GetInner2DPoint(EGridSpace::UniformScaled, IndexNode1) * DisplayScale, Grid.GetInner2DPoint(EGridSpace::UniformScaled, IndexNode2) * DisplayScale, IndexNode1, EVisuProperty::RedCurve);
 			}
 #endif
-			AddToInnerToOuterSegmentsIntersectionTool(IndexNode1, IndexNode2, InType);
 			return;
 		}
 
diff --git a/Engine/Source/Runtime/Datasmith/CADKernel/Private/CADKernel/Mesh/Structure/Grid.h b/Engine/Source/Runtime/Datasmith/CADKernel/Private/CADKernel/Mesh/Structure/Grid.h
--- a/Engine/Source/Runtime/Datasmith/CADKernel/Private/CADKernel/Mesh/Structure/Grid.h
+++ b/Engine/Source/Runtime/Datasmith/CADKernel/Private/CADKernel/Mesh/Structure/Grid.h
@@ -363,7 +363,7 @@ public:
 	/**
 	 * @return true if the bits EPointMarker::IsInside and EPointMarker::IsInsideButTooCloseToLoop are booth equal to 0 and EPointMarker::IsCloseToLoop  is equal to 1
 	 */
-	const bool IsNodeCloseToButOusideFace(int32 Index) const
+	const bool IsNodeOusideFaceButClose(int32 Index) const
 	{
 		constexpr ENodeMarker IsInsideAndClose = ENodeMarker::IsInside | ENodeMarker::IsInsideButTooCloseToLoop | ENodeMarker::IsCloseToLoop;
 		return (NodeMarkers[Index] & IsInsideAndClose) == ENodeMarker::IsCloseToLoop;
