In the case of CAD file with huge face, if the file is imported with "Max Edge Length" criterion set, the resulting mesh can be too big. For example, a face of 100,000 mm x 100,000 mm with a "Max Edge Length" of 10 mm will generate a mesh with 2e8 triangles.

A rule is added to disable the "Max Edge Length" criterion if the resulting mesh will be too huge.

diff --git a/Engine/Source/Runtime/Datasmith/CADKernel/Private/CADKernel/Mesh/Criteria/SizeCriterion.cpp b/Engine/Source/Runtime/Datasmith/CADKernel/Private/CADKernel/Mesh/Criteria/SizeCriterion.cpp
--- a/Engine/Source/Runtime/Datasmith/CADKernel/Private/CADKernel/Mesh/Criteria/SizeCriterion.cpp
+++ b/Engine/Source/Runtime/Datasmith/CADKernel/Private/CADKernel/Mesh/Criteria/SizeCriterion.cpp
@@ -9,6 +9,18 @@
 namespace UE::CADKernel
 {
 
+/**
+ * In some cases as a huge plan, the MaxSizeCriteria can generate a lot of triangles e.g. a 100m side plan with a MaxSizeCriteria of 3cm will need 3e3 elements by side so 2e7 triangles.
+ * This kind of case is most of the time a forgotten sketch body than a real wanted body. This mesh could make the process extremely long or simply crash all the system.
+ * The idea is to not cancel the mesh of the body in the case it will really expected but to avoid the generation of a huge mesh with unwanted hundreds of millions of triangles.
+ * So if MaxSizeCriteria will generate a huge mesh, this criteria is abandoned.
+ * The chosen limit value is 3000 elements by side
+ */
+constexpr int32 GetMaxElementCountPerSide()
+{
+	return 3000;
+}
+
 void FMinSizeCriterion::ApplyOnEdgeParameters(FTopologicalEdge& Edge, const TArray<double>& Coordinates, const TArray<FCurvePoint>& Points) const
 {
 	double NumericPrecision = Edge.GetTolerance3D();
@@ -35,6 +47,12 @@ void FMaxSizeCriterion::ApplyOnEdgeParameters(FTopologicalEdge& Edge, const TArr
 		return;
 	}
 
+	const double ElementCount = Edge.Length() / Size;
+	if (ElementCount > GetMaxElementCountPerSide())
+	{
+		return;
+	}
+
 	ApplyOnParameters(Coordinates, Points, Edge.GetDeltaUMaxs(), [](double NewValue, double& AbacusValue)
 		{
 			if (NewValue < AbacusValue)
@@ -79,6 +97,12 @@ void FMaxSizeCriterion::UpdateDelta(double InDeltaU, double InUSag, double InDia
 		return;
 	}
 
+	const double ElementCount = ChordLength / Size;
+	if (ElementCount > GetMaxElementCountPerSide())
+	{
+		return;
+	}
+
 	double DeltaU = InDeltaU * Size / ChordLength;
 	if (DeltaU < OutSagDeltaUMax)
 	{
