Preparatory work: Refactoring ApplyCriteria step by encapsulating all steps of ApplyCriteria in CriteriaGrid methods.

diff --git a/Engine/Source/Runtime/Datasmith/CADKernel/Private/CADKernel/Mesh/Criteria/CriteriaGrid.cpp b/Engine/Source/Runtime/Datasmith/CADKernel/Private/CADKernel/Mesh/Criteria/CriteriaGrid.cpp
--- a/Engine/Source/Runtime/Datasmith/CADKernel/Private/CADKernel/Mesh/Criteria/CriteriaGrid.cpp
+++ b/Engine/Source/Runtime/Datasmith/CADKernel/Private/CADKernel/Mesh/Criteria/CriteriaGrid.cpp
@@ -70,11 +70,27 @@ FCriteriaGrid::FCriteriaGrid(FTopologicalFace& InFace)
 void FCriteriaGrid::ApplyCriteria(const TArray<TSharedPtr<FCriterion>>& Criteria) const
 {
 	TArray<double>& DeltaUMaxArray = Face.GetCrossingPointDeltaMaxs(EIso::IsoU);
-	TArray<double>& DeltaUMiniArray = Face.GetCrossingPointDeltaMins(EIso::IsoU);
+	TArray<double>& DeltaUMinArray = Face.GetCrossingPointDeltaMins(EIso::IsoU);
 	TArray<double>& DeltaVMaxArray = Face.GetCrossingPointDeltaMaxs(EIso::IsoV);
 	TArray<double>& DeltaVMinArray = Face.GetCrossingPointDeltaMins(EIso::IsoV);
 	FSurfaceCurvature& SurfaceCurvature = Face.GetCurvatures();
 
+	double ElementLengthMin = DOUBLE_BIG_NUMBER;
+
+	TFunction<void(const double, const double, const double)> ElementLength = [&](const double CoordNextMinusCoord, const double Length, const double DeltaMax)
+	{
+		if (CoordNextMinusCoord < DOUBLE_SMALL_NUMBER)
+		{
+			return Length;
+		}
+		const double ElemLength = Length * DeltaMax / CoordNextMinusCoord;
+		if (ElementLengthMin > ElemLength)
+		{
+			ElementLengthMin = ElemLength;
+		}
+		return ElemLength;
+	};
+
 	for (int32 IndexV = 0; IndexV < GetCoordinateCount(EIso::IsoV) - 1; ++IndexV)
 	{
 		for (int32 IndexU = 0; IndexU < GetCoordinateCount(EIso::IsoU) - 1; ++IndexU)
@@ -89,17 +105,28 @@ void FCriteriaGrid::ApplyCriteria(const TArray<TSharedPtr<FCriterion>>& Criteria
 
 			// Evaluate Sag
 			double LengthU;
-			double SagU = FCriterion::EvaluateSag(Point_U0_Vm, Point_U1_Vm, Point_Um_Vm, LengthU);
+			const double SagU = FCriterion::EvaluateSag(Point_U0_Vm, Point_U1_Vm, Point_Um_Vm, LengthU);
 			double LengthV;
-			double SagV = FCriterion::EvaluateSag(Point_Um_V0, Point_Um_V1, Point_Um_Vm, LengthV);
+			const double SagV = FCriterion::EvaluateSag(Point_Um_V0, Point_Um_V1, Point_Um_Vm, LengthV);
 			double LengthUV;
-			double SagUV = FCriterion::EvaluateSag(Point_U0_V0, Point_U1_V1, Point_Um_Vm, LengthUV);
+			const double SagUV = FCriterion::EvaluateSag(Point_U0_V0, Point_U1_V1, Point_Um_Vm, LengthUV);
+
+			double& DeltaUMin = DeltaUMinArray[IndexU];
+			double& DeltaUMax = DeltaUMaxArray[IndexU];
+			double& DeltaVMin = DeltaVMinArray[IndexV];
+			double& DeltaVMax = DeltaVMaxArray[IndexV];
+
+			const double UNextMinusU = GetCoordinate(EIso::IsoU, IndexU + 1) - GetCoordinate(EIso::IsoU, IndexU);
+			const double VNextMinusV = GetCoordinate(EIso::IsoV, IndexV + 1) - GetCoordinate(EIso::IsoV, IndexV);
 
 			for (const TSharedPtr<FCriterion>& Criterion : Criteria)
 			{
-				Criterion->UpdateDelta((GetCoordinate(EIso::IsoU, IndexU + 1) - GetCoordinate(EIso::IsoU, IndexU)), SagU, SagUV, SagV, LengthU, LengthUV, DeltaUMaxArray[IndexU], DeltaUMiniArray[IndexU], SurfaceCurvature[EIso::IsoU]);
-				Criterion->UpdateDelta((GetCoordinate(EIso::IsoV, IndexV + 1) - GetCoordinate(EIso::IsoV, IndexV)), SagV, SagUV, SagU, LengthV, LengthUV, DeltaVMaxArray[IndexV], DeltaVMinArray[IndexV], SurfaceCurvature[EIso::IsoV]);
+				Criterion->UpdateDelta(UNextMinusU, SagU, SagUV, SagV, LengthU, LengthUV, DeltaUMax, DeltaUMin, SurfaceCurvature[EIso::IsoU]);
+				Criterion->UpdateDelta(VNextMinusV, SagV, SagUV, SagU, LengthV, LengthUV, DeltaVMax, DeltaVMin, SurfaceCurvature[EIso::IsoV]);
 			}
+
+			ElementLength(UNextMinusU, LengthU, DeltaUMax);
+			ElementLength(VNextMinusV, LengthV, DeltaVMax);
 		}
 	}
 
@@ -115,6 +142,9 @@ void FCriteriaGrid::ApplyCriteria(const TArray<TSharedPtr<FCriterion>>& Criteria
 		DeltaVMaxArray[0] = (DeltaVMaxArray[0] + DeltaVMaxArray[1] * 2) * AThird;
 		DeltaVMaxArray.Last() = (DeltaVMaxArray.Last() + DeltaVMaxArray[DeltaVMaxArray.Num() - 2] * 2) * AThird;
 	}
+
+	Face.SetEstimatedMinimalElementLength(ElementLengthMin);
+	Face.SetApplyCriteriaMarker();
 }
 
 void FCriteriaGrid::Display() const
diff --git a/Engine/Source/Runtime/Datasmith/CADKernel/Private/CADKernel/Mesh/Criteria/CriteriaGrid.h b/Engine/Source/Runtime/Datasmith/CADKernel/Private/CADKernel/Mesh/Criteria/CriteriaGrid.h
--- a/Engine/Source/Runtime/Datasmith/CADKernel/Private/CADKernel/Mesh/Criteria/CriteriaGrid.h
+++ b/Engine/Source/Runtime/Datasmith/CADKernel/Private/CADKernel/Mesh/Criteria/CriteriaGrid.h
@@ -4,7 +4,6 @@
 #include "CADKernel/Core/Types.h"
 #include "CADKernel/Geo/GeoPoint.h"
 #include "CADKernel/Geo/Sampling/SurfacicSampling.h"
-#include "CADKernel/Math/Point.h"
 
 namespace UE::CADKernel
 {
diff --git a/Engine/Source/Runtime/Datasmith/CADKernel/Private/CADKernel/Mesh/Criteria/Criterion.cpp b/Engine/Source/Runtime/Datasmith/CADKernel/Private/CADKernel/Mesh/Criteria/Criterion.cpp
--- a/Engine/Source/Runtime/Datasmith/CADKernel/Private/CADKernel/Mesh/Criteria/Criterion.cpp
+++ b/Engine/Source/Runtime/Datasmith/CADKernel/Private/CADKernel/Mesh/Criteria/Criterion.cpp
@@ -17,6 +17,7 @@ namespace UE::CADKernel
 void FCriterion::ApplyOnEdgeParameters(FTopologicalEdge& Edge, const TArray<double>& Coordinates, const TArray<FCurvePoint>& Points) const
 {
 	TArray<double>& DeltaUMaxs = Edge.GetDeltaUMaxs();
+	TArray<double>& DeltaUMins = Edge.GetDeltaUMins();
 
 	double NumericPrecision = Edge.GetTolerance3D();
 	NumericPrecision /= 10;
@@ -29,31 +30,24 @@ void FCriterion::ApplyOnEdgeParameters(FTopologicalEdge& Edge, const TArray<doub
 	int32 PreviousIndex = 0;
 	for (int32 Index = 1; Index < Coordinates.Num(); Index++)
 	{
-		int32 MiddleIndex = PreviousIndex + Index;
-		double ChordLength;
-		double Sag = EvaluateSag(Points[2 * PreviousIndex].Point, Points[2 * Index].Point, Points[MiddleIndex].Point, ChordLength);
+		const int32 MiddleIndex = PreviousIndex + Index;
+		const double DeltaU = Coordinates[Index] - Coordinates[PreviousIndex];
 
-		double AiAm = Points[2 * PreviousIndex].Point.Distance(Points[MiddleIndex].Point);
-		double AmAi2 = Points[2 * Index].Point.Distance(Points[MiddleIndex].Point);
-		double AiAi2 = Points[2 * PreviousIndex].Point.Distance(Points[2 * Index].Point);
+		const FPoint Start = Points[2 * PreviousIndex].Point;
+		const FPoint End = Points[2 * Index].Point;
+		const FPoint Middle = Points[MiddleIndex].Point;
 
-		double DeltaU = Coordinates[Index] - Coordinates[PreviousIndex];
+		double ChordLength;
+		const double SagFromStart = EvaluateSag(Start, End, Middle, ChordLength);
+		const double SagFromEnd = EvaluateSag(End, Start, Middle, ChordLength);
+		const double Sag = FMath::Max(SagFromStart, SagFromEnd);
 
 		if (ChordLength > NumericPrecision && Sag > DOUBLE_SMALL_NUMBER)
 		{
-			double DeltaUMax = ComputeDeltaU(ChordLength, DeltaU, Sag);
-			if (DeltaUMax < DeltaUMaxs[PreviousIndex]) DeltaUMaxs[PreviousIndex] = DeltaUMax;
+			const double NewDeltaUMax = ComputeDeltaU(ChordLength, DeltaU, Sag);
+			UpdateWithUMaxValue(NewDeltaUMax, DeltaUMaxs[PreviousIndex], DeltaUMins[PreviousIndex]);
 		}
 
-		Sag = EvaluateSag(Points[2 * Index].Point, Points[2 * PreviousIndex].Point, Points[MiddleIndex].Point, ChordLength);
-		if (ChordLength > NumericPrecision && Sag > DOUBLE_SMALL_NUMBER)
-		{
-			double DeltaUMax = ComputeDeltaU(ChordLength, DeltaU, Sag);
-			if (DeltaUMax < DeltaUMaxs[PreviousIndex])
-			{
-				DeltaUMaxs[PreviousIndex] = DeltaUMax;
-			}
-		}
 		PreviousIndex = Index;
 	}
 }
diff --git a/Engine/Source/Runtime/Datasmith/CADKernel/Private/CADKernel/Mesh/Criteria/SizeCriterion.cpp b/Engine/Source/Runtime/Datasmith/CADKernel/Private/CADKernel/Mesh/Criteria/SizeCriterion.cpp
--- a/Engine/Source/Runtime/Datasmith/CADKernel/Private/CADKernel/Mesh/Criteria/SizeCriterion.cpp
+++ b/Engine/Source/Runtime/Datasmith/CADKernel/Private/CADKernel/Mesh/Criteria/SizeCriterion.cpp
@@ -29,12 +29,9 @@ void FMinSizeCriterion::ApplyOnEdgeParameters(FTopologicalEdge& Edge, const TArr
 		return;
 	}
 
-	ApplyOnParameters(Coordinates, Points, Edge.GetDeltaUMins(), [](double NewValue, double& AbacusValue)
+	ApplyOnParameters(Coordinates, Points, Edge.GetDeltaUMaxs(), Edge.GetDeltaUMins(), [&](double NewMaxValue, double& OutDeltaUMax, double& OutDeltaUMin)
 		{
-			if (NewValue > AbacusValue)
-			{
-				AbacusValue = NewValue;
-			}
+			UpdateWithUMinValue(NewMaxValue, OutDeltaUMax, OutDeltaUMin);
 		});
 }
 
@@ -53,44 +50,37 @@ void FMaxSizeCriterion::ApplyOnEdgeParameters(FTopologicalEdge& Edge, const TArr
 		return;
 	}
 
-	ApplyOnParameters(Coordinates, Points, Edge.GetDeltaUMaxs(), [](double NewValue, double& AbacusValue)
+	ApplyOnParameters(Coordinates, Points, Edge.GetDeltaUMaxs(), Edge.GetDeltaUMins(), [&](double NewMaxValue, double& OutDeltaUMax, double& OutDeltaUMin)
 		{
-			if (NewValue < AbacusValue)
-			{
-				AbacusValue = NewValue;
-			}
+			UpdateWithUMaxValue(NewMaxValue, OutDeltaUMax, OutDeltaUMin);
 		});
 }
 
-void FSizeCriterion::ApplyOnParameters(const TArray<double>& Coordinates, const TArray<FCurvePoint>& Points, TArray<double>& DeltaUArray, TFunction<void(double, double&)> Compare) const
+void FSizeCriterion::ApplyOnParameters(const TArray<double>& Coordinates, const TArray<FCurvePoint>& Points, TArray<double>& DeltaUMaxs, TArray<double>& DeltaUMins, TFunction<void(double, double&, double&)> UpdateDeltaU) const
 {
-	double DeltaUMax = Coordinates[Coordinates.Num() - 1] - Coordinates[0];
-
-	for (int32 PIndex = 1; PIndex < Coordinates.Num(); PIndex++)
+	for (int32 Index = 1; Index < Coordinates.Num(); Index++)
 	{
-		double DeltaU = Coordinates[PIndex] - Coordinates[PIndex - 1];
-		double Length = Points[2 * (PIndex - 1)].Point.Distance(Points[2 * PIndex].Point);
+		const int32 PreviousIndex = Index - 1;
+		const double DeltaCoordinate = Coordinates[Index] - Coordinates[PreviousIndex];
+		const double ChordLength = Points[2 * PreviousIndex].Point.Distance(Points[2 * Index].Point);
 
-		DeltaU = (Length > 0) ? DeltaU * Size / Length : DeltaUMax;
-		Compare(DeltaU, DeltaUArray[PIndex - 1]);
+		const double NewDeltaU = ComputeSizeCriterionValue(DeltaCoordinate, ChordLength);
+		UpdateDeltaU(NewDeltaU, DeltaUMaxs[Index - 1], DeltaUMins[Index - 1]);
 	}
 }
 
-void FMinSizeCriterion::UpdateDelta(double InDeltaU, double InUSag, double InDiagonalSag, double InVSag, double ChordLength, double DiagonalLength, double& OutSagDeltaUMax, double& OutSagDeltaUMin, FIsoCurvature& SurfaceCurvature) const
+void FMinSizeCriterion::UpdateDelta(double InDeltaCoordinate, double InUSag, double InDiagonalSag, double InVSag, double ChordLength, double DiagonalLength, double& OutSagDeltaUMax, double& OutSagDeltaUMin, FIsoCurvature& SurfaceCurvature) const
 {
 	if (ChordLength < DOUBLE_KINDA_SMALL_NUMBER)
 	{
 		return;
 	}
 
-	double DeltaU = InDeltaU * Size / ChordLength;
-	if (DeltaU > OutSagDeltaUMin)
-	{
-		OutSagDeltaUMin = DeltaU;
-	}
+	const double NewDeltaU = ComputeSizeCriterionValue(InDeltaCoordinate, ChordLength);
+	UpdateWithUMinValue(NewDeltaU, OutSagDeltaUMax, OutSagDeltaUMin);
 }
 
-void FMaxSizeCriterion::UpdateDelta(double InDeltaU, double InUSag, double InDiagonalSag, double InVSag, double ChordLength, double DiagonalLength, double& OutSagDeltaUMax, double& OutSagDeltaUMin, FIsoCurvature& SurfaceCurvature) const
+void FMaxSizeCriterion::UpdateDelta(double InDeltaCoordinate, double InUSag, double InDiagonalSag, double InVSag, double ChordLength, double DiagonalLength, double& OutSagDeltaUMax, double& OutSagDeltaUMin, FIsoCurvature& SurfaceCurvature) const
 {
 	if (ChordLength < DOUBLE_KINDA_SMALL_NUMBER)
 	{
@@ -103,11 +93,8 @@ void FMaxSizeCriterion::UpdateDelta(double InDeltaU, double InUSag, double InDia
 		return;
 	}
 
-	double DeltaU = InDeltaU * Size / ChordLength;
-	if (DeltaU < OutSagDeltaUMax)
-	{
-		OutSagDeltaUMax = DeltaU;
-	}
+	const double NewDeltaU = ComputeSizeCriterionValue(InDeltaCoordinate, ChordLength);
+	UpdateWithUMaxValue(NewDeltaU, OutSagDeltaUMax, OutSagDeltaUMin);
 }
 
 } // namespace UE::CADKernel
diff --git a/Engine/Source/Runtime/Datasmith/CADKernel/Private/CADKernel/Mesh/Criteria/SizeCriterion.h b/Engine/Source/Runtime/Datasmith/CADKernel/Private/CADKernel/Mesh/Criteria/SizeCriterion.h
--- a/Engine/Source/Runtime/Datasmith/CADKernel/Private/CADKernel/Mesh/Criteria/SizeCriterion.h
+++ b/Engine/Source/Runtime/Datasmith/CADKernel/Private/CADKernel/Mesh/Criteria/SizeCriterion.h
@@ -24,6 +24,10 @@ protected:
 	{
 	}
 
+	double ComputeSizeCriterionValue(double InDeltaU, double ChordLength) const
+	{
+		return  (ChordLength > UE_DOUBLE_KINDA_SMALL_NUMBER) ? InDeltaU * Size / ChordLength : UE_DOUBLE_BIG_NUMBER;
+	}
 
 public:
 
@@ -50,7 +54,7 @@ public:
 		return 0;
 	}
 
-	void ApplyOnParameters(const TArray<double>& TabU, const TArray<FCurvePoint>& tabPt, TArray<double>& tabDeltaU, TFunction<void(double, double&)> Compare) const;
+	void ApplyOnParameters(const TArray<double>& Coordinates, const TArray<FCurvePoint>& Points, TArray<double>& DeltaUMax, TArray<double>& DeltaUMins, TFunction<void(double, double&, double&)> UpdateDeltaU) const;
 };
 
 class FMinSizeCriterion : public FSizeCriterion
@@ -70,6 +74,11 @@ public:
 	}
 
 	virtual void ApplyOnEdgeParameters(FTopologicalEdge& Edge, const TArray<double>& Coordinates, const TArray<FCurvePoint>& Points) const override;
+
+	/**
+	 * MinSizeCrtierion sets DeltaUMin.
+	 * @see FCriterion::UpdateWithUMinValue
+	 */
 	virtual void UpdateDelta(double InDeltaU, double InUSag, double InDiagonalSag, double InVSag, double ChordLength, double DiagonalLength, double& OutSagDeltaUMax, double& OutSagDeltaUMin, FIsoCurvature& SurfaceCurvature) const override;
 };
 
@@ -89,6 +98,11 @@ public:
 	}
 
 	virtual void ApplyOnEdgeParameters(FTopologicalEdge& Edge, const TArray<double>& Coordinates, const TArray<FCurvePoint>& Points) const override;
+
+	/**
+	 * MaxSizeCriterion sets DeltaUMax.
+	 * @see FCriterion::UpdateWithUMaxValue
+	 */
 	virtual void UpdateDelta(double InDeltaU, double InUSag, double InDiagonalSag, double InVSag, double ChordLength, double DiagonalLength, double& OutSagDeltaUMax, double& OutSagDeltaUMin, FIsoCurvature& SurfaceCurvature) const override;
 };
 
diff --git a/Engine/Source/Runtime/Datasmith/CADKernel/Private/CADKernel/Mesh/Meshers/ParametricMesher.cpp b/Engine/Source/Runtime/Datasmith/CADKernel/Private/CADKernel/Mesh/Meshers/ParametricMesher.cpp
--- a/Engine/Source/Runtime/Datasmith/CADKernel/Private/CADKernel/Mesh/Meshers/ParametricMesher.cpp
+++ b/Engine/Source/Runtime/Datasmith/CADKernel/Private/CADKernel/Mesh/Meshers/ParametricMesher.cpp
@@ -84,32 +84,45 @@ void FParametricMesher::MeshEntities(TArray<FTopologicalShapeEntity*>& InEntitie
 		Face->ResetMarkers();
 	}
 
+	PreMeshingTasks();
 	MeshEntities();
 }
 
-void FParametricMesher::MeshEntities()
+void FParametricMesher::PreMeshingTasks()
 {
+	TRACE_CPUPROFILER_EVENT_SCOPE(ParametricMesher::PreMeshingTasks);
+
 	FTimePoint StartTime = FChrono::Now();
 	FTimePoint ApplyCriteriaStartTime = FChrono::Now();
 
 	FProgress ProgressBar(Faces.Num() * 2, TEXT("Meshing Entities : Apply Surface Criteria"));
 
-	// ============================================================================================================
-	//      Apply Surface Criteria
-	// ============================================================================================================
-
 	for (FTopologicalFace* Face : Faces)
 	{
 		FProgress _(1, TEXT("Meshing Entities : Apply Surface Criteria"));
 
-		if(Face == nullptr || Face->IsDeleted())
+		if (Face == nullptr || Face->IsDeleted() || Face->IsMeshed())
 		{
 			continue;
 		}
 
 		ApplyFaceCriteria(*Face);
+		Face->ComputeSurfaceSideProperties();
 	}
 
+#ifdef CADKERNEL_DEV
+	MesherReport.Chronos.ApplyCriteriaDuration = FChrono::Elapse(ApplyCriteriaStartTime);
+#endif
+}
+
+void FParametricMesher::MeshEntities()
+{
+	TRACE_CPUPROFILER_EVENT_SCOPE(ParametricMesher::MeshEntities);
+
+	FTimePoint StartTime = FChrono::Now();
+
+	FProgress ProgressBar(Faces.Num() * 2, TEXT("Meshing Entities : Find quad surfaces"));
+
 #ifdef CADKERNEL_DEV
 	MesherReport.Chronos.ApplyCriteriaDuration = FChrono::Elapse(ApplyCriteriaStartTime);
 #endif
@@ -117,7 +130,7 @@ void FParametricMesher::MeshEntities()
 	FTimePoint MeshingStartTime = FChrono::Now();
 
 	// ============================================================================================================
-	//      Find quad surfaces 
+	//      Find and sort quad surfaces 
 	// ============================================================================================================
 
 	TArray<FCostToFace> QuadTrimmedSurfaceSet;
@@ -149,7 +162,6 @@ void FParametricMesher::MeshEntities()
 	MesherReport.Chronos.GlobalDuration = FChrono::Elapse(StartTime);
 #endif
 
-	//Chronos.PrintTimeElapse();
 }
 
 void FParametricMesher::ApplyFaceCriteria(FTopologicalFace& Face)
@@ -166,9 +178,6 @@ void FParametricMesher::ApplyFaceCriteria(FTopologicalFace& Face)
 	}
 
 	Grid.ApplyCriteria(GetMeshModel().GetCriteria());
-
-	Face.ChooseFinalDeltaUs();
-	Face.SetApplyCriteria();
 }
 
 void FParametricMesher::ApplyEdgeCriteria(FTopologicalEdge& Edge)
@@ -198,9 +207,8 @@ void FParametricMesher::ApplyEdgeCriteria(FTopologicalEdge& Edge)
 		Criterion->ApplyOnEdgeParameters(Edge, CrossingPointUs, Points3D);
 	}
 
-	Edge.ChooseFinalDeltaUs();
-	Edge.SetApplyCriteria();
-	ActiveEdge.SetApplyCriteria();
+	Edge.SetApplyCriteriaMarker();
+	ActiveEdge.SetApplyCriteriaMarker();
 }
 
 void FParametricMesher::Mesh(FTopologicalFace& Face)
@@ -780,21 +788,13 @@ void FParametricMesher::GenerateEdgeElements(FTopologicalEdge& Edge)
 
 void FParametricMesher::IsolateQuadFace(TArray<FCostToFace>& QuadSurfaces, TArray<FTopologicalFace*>& OtherSurfaces) const
 {
+	TRACE_CPUPROFILER_EVENT_SCOPE(ParametricMesher::IsolateQuadFace);
+
 	TArray<FTopologicalFace*> FlatQuadsAndTriangles;
 	FlatQuadsAndTriangles.Reserve(Faces.Num());
 	QuadSurfaces.Reserve(Faces.Num() * 2);
 	OtherSurfaces.Reserve(Faces.Num());
 
-	for (FTopologicalFace* Face : Faces)
-	{
-		if (Face == nullptr || Face->IsDeleted() || Face->IsMeshed())
-		{
-			continue;
-		}
-
-		Face->ComputeSurfaceSideProperties();
-	}
-
 	for (FTopologicalFace* FacePtr : Faces)
 	{
 		if (FacePtr == nullptr)
@@ -843,36 +843,34 @@ void FParametricMesher::IsolateQuadFace(TArray<FCostToFace>& QuadSurfaces, TArra
 #ifdef DEBUG_ISOLATEQUADFACE
 	if (QuadSurfaces.Num() > 0)
 	{
-		Open3DDebugSession(TEXT("Quad Entities"));
+		F3DDebugSession A(TEXT("Quad Entities"));
 		for (const FCostToFace& Quad : QuadSurfaces)
 		{
-			//FGraphicSession _(FString::Printf(TEXT("Face %d %f"), Quad.Face.GetId(), Quad.Cost));
-			Display(Quad.Face);
+			F3DDebugSession _(FString::Printf(TEXT("Face %d %f"), Quad.Face->GetId(), Quad.Cost));
+			Display(*Quad.Face);
 		}
-		Close3DDebugSession();
 	}
 
 	if (FlatQuadsAndTriangles.Num() > 0)
 	{
-		Open3DDebugSession(TEXT("Flat Quads & Triangles"));
-		for (const TSharedPtr<FTopologicalFace>& Face : FlatQuadsAndTriangles)
+		F3DDebugSession A(TEXT("Flat Quads & Triangles"));
+		for (const FTopologicalFace* Face : FlatQuadsAndTriangles)
 		{
-			//FGraphicSession _(FString::Printf(TEXT("Face %d %f"), Face.GetId()));
-			Display(Face);
+			F3DDebugSession _(FString::Printf(TEXT("Face %d"), Face->GetId()));
+			Display(*Face);
 		}
-		Close3DDebugSession();
 	}
 
 	if (OtherSurfaces.Num() > 0)
 	{
-		Open3DDebugSession(TEXT("Other Entities"));
-		for (const TSharedPtr<FTopologicalFace>& Face : OtherSurfaces)
+		F3DDebugSession A(TEXT("Other Entities"));
+		for (const FTopologicalFace* Face : OtherSurfaces)
 		{
-			//FGraphicSession _(FString::Printf(TEXT("Face %d %f"), Face.GetId()));
-			Display(StaticCastSharedPtr<FTopologicalFace>(Face));
+			F3DDebugSession _(FString::Printf(TEXT("Face %d"), Face->GetId()));
+			Display(*Face);
 		}
-		Close3DDebugSession();
 	}
+	Wait();
 #endif
 }
 
@@ -1039,9 +1037,9 @@ void FParametricMesher::LinkQuadSurfaceForMesh(TArray<FCostToFace>& QuadTrimmedS
 
 void FParametricMesher::MeshSurfaceByFront(TArray<FCostToFace>& QuadTrimmedSurfaceSet)
 {
-	// Processed3 : Surfaces that have to be meshed are set Processed3
-	// Processed1 : Surfaces added in CandidateSurfacesForMesh
-	// Processed2 : Surfaces added in SecondChoiceOfCandidateSurfacesForMesh
+	// Marker3 : Surfaces that have to be meshed are set Marker3
+	// Marker1 : Surfaces added in CandidateSurfacesForMesh
+	// Marker2 : Surfaces added in SecondChoiceOfCandidateSurfacesForMesh
 
 	FMessage::Printf(EVerboseLevel::Debug, TEXT("Start MeshSurfaceByFront\n"));
 
@@ -1065,17 +1063,19 @@ void FParametricMesher::MeshSurfaceByFront(TArray<FCostToFace>& QuadTrimmedSurfa
 	TFunction<void(FTopologicalFace&)> MeshFace = [&](FTopologicalFace& Face)
 	{
 #ifdef DISPLAYDEBUGMESHFACEBYFACESTEP
-		Open3DDebugSession(TEXT("Surface " + Utils::ToFString(Face->GetId()));
-		Display(Face);
-		Close3DDebugSession();
+		{
+			F3DDebugSession A(FString::Printf(TEXT("Surface %d"), Face.GetId()));
+			Display(Face);
+		}
 #endif
 		Mesh(Face);
 
 #ifdef DISPLAYDEBUGMESHFACEBYFACESTEP
-		Open3DDebugSession(TEXT("Mesh " + Utils::ToFString(Face.GetId()));
-		DisplayMesh((FFaceMesh&)Face.GetOrCreateMesh(GetMeshModel()));
-		Close3DDebugSession();
-		//Wait();
+		{
+			F3DDebugSession A(FString::Printf(TEXT("Mesh %d"), Face.GetId()));
+			DisplayMesh(*Face.GetOrCreateMesh(GetMeshModel()));
+			//Wait();
+		}
 #endif
 
 		if (Face.HasMarker1())
@@ -1087,12 +1087,13 @@ void FParametricMesher::MeshSurfaceByFront(TArray<FCostToFace>& QuadTrimmedSurfa
 			SecondChoiceOfCandidateFacesForMesh.RemoveSingle(&Face);
 		}
 
-		const TSharedPtr<FTopologicalLoop> Loop = Face.GetLoops()[0];
+		const TSharedPtr<FTopologicalLoop>& Loop = Face.GetLoops()[0];
 		for (const FOrientedEdge& OrientedEdge : Loop->GetEdges())
 		{
-			const TSharedPtr<FTopologicalEdge> Edge = OrientedEdge.Entity;
-			Edge->SetMarker1(); // tmp for debug
-			for (FTopologicalEdge* NextEdge : Edge->GetTwinEntities())
+			const FTopologicalEdge& Edge = *OrientedEdge.Entity;
+			Edge.SetMarker1();
+
+			for (FTopologicalEdge* NextEdge : Edge.GetTwinEntities())
 			{
 				if (NextEdge->HasMarker1())
 				{
@@ -1113,9 +1114,11 @@ void FParametricMesher::MeshSurfaceByFront(TArray<FCostToFace>& QuadTrimmedSurfa
 				{
 					continue;
 				}
+
 				int32 SideIndex = NextFace->GetSideIndex(*NextEdge);
 				if (SideIndex == -1)
 				{
+					// The face is not a quad
 					continue;
 				}
 
diff --git a/Engine/Source/Runtime/Datasmith/CADKernel/Private/CADKernel/Topo/TopologicalEdge.cpp b/Engine/Source/Runtime/Datasmith/CADKernel/Private/CADKernel/Topo/TopologicalEdge.cpp
--- a/Engine/Source/Runtime/Datasmith/CADKernel/Private/CADKernel/Topo/TopologicalEdge.cpp
+++ b/Engine/Source/Runtime/Datasmith/CADKernel/Private/CADKernel/Topo/TopologicalEdge.cpp
@@ -693,17 +693,6 @@ TSharedRef<FEdgeMesh> FTopologicalEdge::GetOrCreateMesh(FModelMesh& ShellMesh)
 	return Mesh.ToSharedRef();
 }
 
-void FTopologicalEdge::ChooseFinalDeltaUs()
-{
-	for (int32 Index = 0; Index < CrossingPointDeltaUMins.Num(); ++Index)
-	{
-		if (CrossingPointDeltaUMins[Index] > CrossingPointDeltaUMaxs[Index])
-		{
-			CrossingPointDeltaUMaxs[Index] = CrossingPointDeltaUMins[Index];
-		}
-	}
-}
-
 TSharedPtr<FTopologicalEdge> FTopologicalEdge::CreateEdgeByMergingEdges(TArray<FOrientedEdge>& Edges, const TSharedRef<FTopologicalVertex> StartVertex, const TSharedRef<FTopologicalVertex> EndVertex)
 {
 	// Make merged 2d Nurbs ===================================================
diff --git a/Engine/Source/Runtime/Datasmith/CADKernel/Private/CADKernel/Topo/TopologicalFace.cpp b/Engine/Source/Runtime/Datasmith/CADKernel/Private/CADKernel/Topo/TopologicalFace.cpp
--- a/Engine/Source/Runtime/Datasmith/CADKernel/Private/CADKernel/Topo/TopologicalFace.cpp
+++ b/Engine/Source/Runtime/Datasmith/CADKernel/Private/CADKernel/Topo/TopologicalFace.cpp
@@ -314,7 +314,10 @@ bool FTopologicalFace::HasSameBoundariesAs(const FTopologicalFace* OtherFace) co
 
 bool FTopologicalFace::IsADuplicatedFace() const
 {
-	ensureCADKernel(!GetLoops().Num());
+	if (!GetLoops().Num())
+	{
+		return false;
+	}
 
 	// Find in the adjacent faces of the first (surface or non manifold) edge, a face with the same loops
 	
@@ -495,25 +498,6 @@ void FTopologicalFace::InitDeltaUs()
 	CrossingPointDeltaMaxs[EIso::IsoV].Init(HUGE_VALUE, CrossingCoordinates[EIso::IsoV].Num() - 1);
 }
 
-void FTopologicalFace::ChooseFinalDeltaUs()
-{
-	TFunction<void(const TArray<double>&, TArray<double>&)> ChooseFinalDeltas = [](const TArray<double>& DeltaUMins, TArray<double>& DeltaUMaxs)
-	{
-		int32 Index = 0;
-		for (; Index < DeltaUMins.Num(); ++Index)
-		{
-			if (DeltaUMins[Index] > DeltaUMaxs[Index])
-			{
-				DeltaUMaxs[Index] = DeltaUMins[Index];
-			}
-		}
-	};
-
-	ChooseFinalDeltas(CrossingPointDeltaMins[EIso::IsoU], CrossingPointDeltaMaxs[EIso::IsoU]);
-	ChooseFinalDeltas(CrossingPointDeltaMins[EIso::IsoV], CrossingPointDeltaMaxs[EIso::IsoV]);
-}
-
-
 // Quad ==============================================================================================================================================================================================================================
 
 double FTopologicalFace::GetQuadCriteria()
diff --git a/Engine/Source/Runtime/Datasmith/CADKernel/Public/CADKernel/Mesh/Criteria/Criterion.h b/Engine/Source/Runtime/Datasmith/CADKernel/Public/CADKernel/Mesh/Criteria/Criterion.h
--- a/Engine/Source/Runtime/Datasmith/CADKernel/Public/CADKernel/Mesh/Criteria/Criterion.h
+++ b/Engine/Source/Runtime/Datasmith/CADKernel/Public/CADKernel/Mesh/Criteria/Criterion.h
@@ -112,16 +112,26 @@ public:
 
 	static double DefaultValue(ECriterion type);
 
-	virtual void UpdateDelta(double InDeltaU, double InUSag, double InDiagonalSag, double InVSag, double ChordLength, double DiagonalLength, double& OutSagDeltaUMax, double& OutSagDeltaUMin, FIsoCurvature& SurfaceCurvature) const
+	/**
+	 * According to the criterion, either DeltaUMin or DeltaUMax is set e.g MinSizeCrtierion sets DeltaUMin when MaxSizeCriterion sets DeltaUMax.
+	 * The implicit rule is that DeltaUMin is predominant over DeltaUMax i.e.: 
+	 *	    UpdateWithUMaxValue: 
+	 *			DeltaUMax = Max(DeltaUMin, Min(DeltaUMax, NewDeltaUMax))
+	 *
+	 *		 UpdateWithUMinValue: 
+	 *			DeltaUMin = Max(DeltaUMin, NewDeltaUMin)
+	 *			DeltaUMax = Max(DeltaUMax, NewDeltaUMin)
+	 */
+	virtual void UpdateDelta(double InDeltaU, double InUSag, double InDiagonalSag, double InVSag, double ChordLength, double DiagonalLength, double& OutDeltaUMax, double& OutDeltaUMin, FIsoCurvature& SurfaceCurvature) const
 	{
 		// the component according to U of the sag along the diagonal = (diagonal sag - V Sag) * U Length / diag length
-		double DiagonalSagU = FMath::Abs(InDiagonalSag - InVSag) * ChordLength / DiagonalLength;
+		const double DiagonalSagU = FMath::Abs(InDiagonalSag - InVSag) * ChordLength / DiagonalLength;
 		InUSag = FMath::Max(InUSag, DiagonalSagU);
 
 		if (InUSag > DOUBLE_SMALL_NUMBER)
 		{
-			double DeltaUMax = ComputeDeltaU(ChordLength, InDeltaU, InUSag);
-			OutSagDeltaUMax = (OutSagDeltaUMax < DeltaUMax) ? OutSagDeltaUMax : DeltaUMax;
+			const double CriterionDeltaUMax = ComputeDeltaU(ChordLength, InDeltaU, InUSag);
+			UpdateWithUMaxValue(CriterionDeltaUMax, OutDeltaUMax, OutDeltaUMin);
 		}
 	}
 
@@ -133,6 +143,40 @@ protected:
 		return 0.0;
 	};
 
+	/**
+	 * DeltaUMax = Max(DeltaUMin, Min(DeltaUMax, NewDeltaUMax))
+	 */
+	void UpdateWithUMaxValue(double NewMaxValue, double& OutDeltaUMax, const double& OutDeltaUMin) const 
+	{
+		if (NewMaxValue < OutDeltaUMax)
+		{
+			if (NewMaxValue < OutDeltaUMin)
+			{
+				OutDeltaUMax = OutDeltaUMin;
+			}
+			else
+			{
+				OutDeltaUMax = NewMaxValue;
+			}
+		}
+	}
+
+	/**
+	 * DeltaUMin = Max(DeltaUMin, NewDeltaUMin)
+	 * DeltaUMin cannot be smaller than DeltaUMax, so:
+	 *    DeltaUMax = Max(DeltaUMax, NewDeltaUMin)
+	 */
+	void UpdateWithUMinValue(double NewMinValue, double& OutDeltaUMax, double& OutDeltaUMin) const
+	{
+		if (OutDeltaUMin < NewMinValue)
+		{
+			OutDeltaUMin = NewMinValue;
+		}
+		if (OutDeltaUMax < NewMinValue)
+		{
+			OutDeltaUMax = NewMinValue;
+		}
+	}
 };
 
 } // namespace UE::CADKernel
diff --git a/Engine/Source/Runtime/Datasmith/CADKernel/Public/CADKernel/Mesh/Meshers/ParametricMesher.h b/Engine/Source/Runtime/Datasmith/CADKernel/Public/CADKernel/Mesh/Meshers/ParametricMesher.h
--- a/Engine/Source/Runtime/Datasmith/CADKernel/Public/CADKernel/Mesh/Meshers/ParametricMesher.h
+++ b/Engine/Source/Runtime/Datasmith/CADKernel/Public/CADKernel/Mesh/Meshers/ParametricMesher.h
@@ -99,6 +99,10 @@ public:
 
 protected:
 
+	/**
+	 * ApplyFaceCriteria, ComputeSurfaceSideProperties
+	 */
+	void PreMeshingTasks();
 	void MeshEntities();
 
 
diff --git a/Engine/Source/Runtime/Datasmith/CADKernel/Public/CADKernel/Topo/TopologicalEdge.h b/Engine/Source/Runtime/Datasmith/CADKernel/Public/CADKernel/Topo/TopologicalEdge.h
--- a/Engine/Source/Runtime/Datasmith/CADKernel/Public/CADKernel/Topo/TopologicalEdge.h
+++ b/Engine/Source/Runtime/Datasmith/CADKernel/Public/CADKernel/Topo/TopologicalEdge.h
@@ -458,8 +458,6 @@ public:
 		CrossingPointDeltaUMaxs.Init(2.0 * (GetEndCurvilinearCoordinates() - GetStartCurvilinearCoordinates()), Size - 1);
 	}
 
-	void ChooseFinalDeltaUs();
-
 	const TArray<double>& GetCrossingPointUs() const
 	{
 		return CrossingPointUs;
diff --git a/Engine/Source/Runtime/Datasmith/CADKernel/Public/CADKernel/Topo/TopologicalEntity.h b/Engine/Source/Runtime/Datasmith/CADKernel/Public/CADKernel/Topo/TopologicalEntity.h
--- a/Engine/Source/Runtime/Datasmith/CADKernel/Public/CADKernel/Topo/TopologicalEntity.h
+++ b/Engine/Source/Runtime/Datasmith/CADKernel/Public/CADKernel/Topo/TopologicalEntity.h
@@ -41,7 +41,7 @@ public:
 		return ((States & EHaveStates::IsApplyCriteria) == EHaveStates::IsApplyCriteria);
 	}
 
-	virtual void SetApplyCriteria() const
+	virtual void SetApplyCriteriaMarker() const
 	{
 		States |= EHaveStates::IsApplyCriteria;
 	}
diff --git a/Engine/Source/Runtime/Datasmith/CADKernel/Public/CADKernel/Topo/TopologicalFace.h b/Engine/Source/Runtime/Datasmith/CADKernel/Public/CADKernel/Topo/TopologicalFace.h
--- a/Engine/Source/Runtime/Datasmith/CADKernel/Public/CADKernel/Topo/TopologicalFace.h
+++ b/Engine/Source/Runtime/Datasmith/CADKernel/Public/CADKernel/Topo/TopologicalFace.h
@@ -337,9 +337,6 @@ public:
 
 	void InitDeltaUs();
 
-	void ChooseFinalDeltaUs();
-
-
 	const TArray<double>& GetCuttingCoordinatesAlongIso(EIso Iso) const
 	{
 		return MeshCuttingCoordinates[Iso];
@@ -477,6 +474,8 @@ private:
 	FSurfaceCurvature Curvatures;
 	EQuadType QuadType = EQuadType::Unset;
 
+	double EstimatedMinimalElementLength = DOUBLE_BIG_NUMBER;
+
 public:
 	void ComputeQuadCriteria();
 	double GetQuadCriteria();
@@ -582,6 +581,19 @@ public:
 		}
 	}
 
+	void SetEstimatedMinimalElementLength(double Value)
+	{
+		EstimatedMinimalElementLength = Value;
+	}
+
+	/**
+	 * @return the minimal size of the mesh elements according to meshing criteria
+	 */
+	double GetEstimatedMinimalElementLength() const
+	{
+		return EstimatedMinimalElementLength;
+	}
+
 };
 
 struct FFaceSubset
