Refactor FParametricMesher::ApplyFaceCriteria
Move methods from FCriteriaGrid to FTopologicalFace as this is  FTopologicalFace that is updated

diff --git a/Engine/Plugins/Enterprise/DatasmithCADImporter/Source/CADInterfaces/Private/TechSoftBridge.cpp b/Engine/Plugins/Enterprise/DatasmithCADImporter/Source/CADInterfaces/Private/TechSoftBridge.cpp
--- a/Engine/Plugins/Enterprise/DatasmithCADImporter/Source/CADInterfaces/Private/TechSoftBridge.cpp
+++ b/Engine/Plugins/Enterprise/DatasmithCADImporter/Source/CADInterfaces/Private/TechSoftBridge.cpp
@@ -248,7 +248,7 @@ UE::CADKernel::FBody* FTechSoftBridge::AddBody(A3DRiBrepModel* A3DBRepModel, TMa
 
 	if(Body->FaceCount() == 0)
 	{
-		Body->SetDeleted();
+		Body->Delete();
 		return nullptr;
 	}
 
@@ -893,7 +893,7 @@ void FTechSoftBridge::AddFace(const A3DTopoFace* A3DFace, UE::CADKernel::EOrient
 			Report->FailedFaceCount++;
 #endif
 			Face->SetAsDegenerated();
-			Face->SetDeleted();
+			Face->Delete();
 			return;
 		}
 		else
diff --git a/Engine/Source/Runtime/Datasmith/CADKernel/Private/CADKernel/Mesh/Criteria/CriteriaGrid.cpp b/Engine/Source/Runtime/Datasmith/CADKernel/Private/CADKernel/Mesh/Criteria/CriteriaGrid.cpp
--- a/Engine/Source/Runtime/Datasmith/CADKernel/Private/CADKernel/Mesh/Criteria/CriteriaGrid.cpp
+++ b/Engine/Source/Runtime/Datasmith/CADKernel/Private/CADKernel/Mesh/Criteria/CriteriaGrid.cpp
@@ -14,23 +14,15 @@
 namespace UE::CADKernel
 {
 
-const FPoint& FCriteriaGrid::GetPoint(int32 UIndex, int32 VIndex, bool bIsInternalU, bool bIsInternalV, FVector3f* OutNormal) const
-{
-	int32 Index = GetIndex(UIndex, VIndex, bIsInternalU, bIsInternalV);
-	ensureCADKernel(Grid.Points3D.IsValidIndex(Index));
-
-	if (OutNormal)
-	{
-		*OutNormal = Grid.Normals[Index];
-	}
-	return Grid.Points3D[Index];
-}
-
-void FCriteriaGrid::Init()
+void FCriteriaGrid::Init(const FTopologicalFace& Face)
 {
 	TFunction<void(const TArray<double>&, TArray<double>&)> ComputeMiddlePointsCoordinates = [](const TArray<double>& Tab, TArray<double>& Tab2)
 	{
-		Tab2.SetNum(Tab.Num() * 2 - 1);
+		const int32 TabSize = Tab.Num();
+		ensure(TabSize);
+
+		Tab2.SetNum(TabSize * 2 - 1);
+
 		Tab2[0] = Tab[0];
 		for (int32 Index = 1; Index < Tab.Num(); Index++)
 		{
@@ -47,106 +39,16 @@ void FCriteriaGrid::Init()
 	TrueUcoorindateCount = CoordinateGrid2[EIso::IsoU].Num();
 }
 
-FCriteriaGrid::FCriteriaGrid(FTopologicalFace& InFace)
-	: Face(InFace)
-	, CoordinateGrid(InFace.GetCrossingPointCoordinates())
+FCriteriaGrid::FCriteriaGrid(const FTopologicalFace& InFace)
+	: CoordinateGrid(InFace.GetCrossingPointCoordinates())
 {
-	Face.Presample();
-
-	constexpr int32 MaxGrid = 1000000;
-	if (CoordinateGrid[EIso::IsoU].Num() * CoordinateGrid[EIso::IsoV].Num() > MaxGrid)
-	{
-		Face.SetDeleted();
-		return;
-	}
-
-	Face.InitDeltaUs();
-	Init();
+	Init(InFace);
 #ifdef DISPLAY_CRITERIA_GRID
 	Display();
 #endif
 }
 
-void FCriteriaGrid::ApplyCriteria(const TArray<TSharedPtr<FCriterion>>& Criteria) const
-{
-	TArray<double>& DeltaUMaxArray = Face.GetCrossingPointDeltaMaxs(EIso::IsoU);
-	TArray<double>& DeltaUMinArray = Face.GetCrossingPointDeltaMins(EIso::IsoU);
-	TArray<double>& DeltaVMaxArray = Face.GetCrossingPointDeltaMaxs(EIso::IsoV);
-	TArray<double>& DeltaVMinArray = Face.GetCrossingPointDeltaMins(EIso::IsoV);
-	FSurfaceCurvature& SurfaceCurvature = Face.GetCurvatures();
-
-	double ElementLengthMin = DOUBLE_BIG_NUMBER;
-
-	TFunction<void(const double, const double, const double)> ElementLength = [&](const double CoordNextMinusCoord, const double Length, const double DeltaMax)
-	{
-		if (CoordNextMinusCoord < DOUBLE_SMALL_NUMBER)
-		{
-			return Length;
-		}
-		const double ElemLength = Length * DeltaMax / CoordNextMinusCoord;
-		if (ElementLengthMin > ElemLength)
-		{
-			ElementLengthMin = ElemLength;
-		}
-		return ElemLength;
-	};
-
-	for (int32 IndexV = 0; IndexV < GetCoordinateCount(EIso::IsoV) - 1; ++IndexV)
-	{
-		for (int32 IndexU = 0; IndexU < GetCoordinateCount(EIso::IsoU) - 1; ++IndexU)
-		{
-			const FPoint& Point_U0_V0 = GetPoint(IndexU, IndexV);
-			const FPoint& Point_U1_V1 = GetPoint(IndexU + 1, IndexV + 1);
-			const FPoint& Point_Um_V0 = GetIntermediateU(IndexU, IndexV);
-			const FPoint& Point_Um_V1 = GetIntermediateU(IndexU, IndexV+1);
-			const FPoint& Point_U0_Vm = GetIntermediateV(IndexU, IndexV);
-			const FPoint& Point_U1_Vm = GetIntermediateV(IndexU+1, IndexV);
-			const FPoint& Point_Um_Vm = GetIntermediateUV(IndexU, IndexV);
-
-			// Evaluate Sag
-			double LengthU;
-			const double SagU = FCriterion::EvaluateSag(Point_U0_Vm, Point_U1_Vm, Point_Um_Vm, LengthU);
-			double LengthV;
-			const double SagV = FCriterion::EvaluateSag(Point_Um_V0, Point_Um_V1, Point_Um_Vm, LengthV);
-			double LengthUV;
-			const double SagUV = FCriterion::EvaluateSag(Point_U0_V0, Point_U1_V1, Point_Um_Vm, LengthUV);
-
-			double& DeltaUMin = DeltaUMinArray[IndexU];
-			double& DeltaUMax = DeltaUMaxArray[IndexU];
-			double& DeltaVMin = DeltaVMinArray[IndexV];
-			double& DeltaVMax = DeltaVMaxArray[IndexV];
-
-			const double UNextMinusU = GetCoordinate(EIso::IsoU, IndexU + 1) - GetCoordinate(EIso::IsoU, IndexU);
-			const double VNextMinusV = GetCoordinate(EIso::IsoV, IndexV + 1) - GetCoordinate(EIso::IsoV, IndexV);
-
-			for (const TSharedPtr<FCriterion>& Criterion : Criteria)
-			{
-				Criterion->UpdateDelta(UNextMinusU, SagU, SagUV, SagV, LengthU, LengthUV, DeltaUMax, DeltaUMin, SurfaceCurvature[EIso::IsoU]);
-				Criterion->UpdateDelta(VNextMinusV, SagV, SagUV, SagU, LengthV, LengthUV, DeltaVMax, DeltaVMin, SurfaceCurvature[EIso::IsoV]);
-			}
-
-			ElementLength(UNextMinusU, LengthU, DeltaUMax);
-			ElementLength(VNextMinusV, LengthV, DeltaVMax);
-		}
-	}
-
-	// Delta of the extremities are smooth to avoid big disparity 
-	if (DeltaUMaxArray.Num() > 2)
-	{
-		DeltaUMaxArray[0] = (DeltaUMaxArray[0] + DeltaUMaxArray[1] * 2) * AThird;
-		DeltaUMaxArray.Last() = (DeltaUMaxArray.Last() + DeltaUMaxArray[DeltaUMaxArray.Num() - 2] * 2) * AThird;
-	}
-
-	if (DeltaVMaxArray.Num() > 2)
-	{
-		DeltaVMaxArray[0] = (DeltaVMaxArray[0] + DeltaVMaxArray[1] * 2) * AThird;
-		DeltaVMaxArray.Last() = (DeltaVMaxArray.Last() + DeltaVMaxArray[DeltaVMaxArray.Num() - 2] * 2) * AThird;
-	}
-
-	Face.SetEstimatedMinimalElementLength(ElementLengthMin);
-	Face.SetApplyCriteriaMarker();
-}
-
+#ifdef DISPLAY_CRITERIA_GRID
 void FCriteriaGrid::Display() const
 {
 	F3DDebugSession _(TEXT("Grid"));
@@ -195,22 +97,6 @@ void FCriteriaGrid::Display() const
 		}
 	}
 
-	{
-		F3DDebugSession A(TEXT("Loop 3D"));
-		for (const TSharedPtr<FTopologicalLoop>& Loop : Face.GetLoops())
-		{
-			UE::CADKernel::Display(*Loop);
-		}
-	}
-
-	{
-		F3DDebugSession A(TEXT("Loop 2D"));
-		for (const TSharedPtr<FTopologicalLoop>& Loop : Face.GetLoops())
-		{
-			UE::CADKernel::Display2D(*Loop);
-		}
-	}
-
 	{
 		F3DDebugSession A(TEXT("CriteriaGrid Point 2D"));
 		for (int32 VIndex = 0; VIndex < GetCoordinateCount(EIso::IsoV); ++VIndex)
@@ -249,6 +135,7 @@ void FCriteriaGrid::Display() const
 
 	Wait();
 }
+#endif
 
 } // namespace UE::CADKernel
 
diff --git a/Engine/Source/Runtime/Datasmith/CADKernel/Private/CADKernel/Mesh/Criteria/CriteriaGrid.h b/Engine/Source/Runtime/Datasmith/CADKernel/Private/CADKernel/Mesh/Criteria/CriteriaGrid.h
--- a/Engine/Source/Runtime/Datasmith/CADKernel/Private/CADKernel/Mesh/Criteria/CriteriaGrid.h
+++ b/Engine/Source/Runtime/Datasmith/CADKernel/Private/CADKernel/Mesh/Criteria/CriteriaGrid.h
@@ -8,12 +8,10 @@
 namespace UE::CADKernel
 {
 class FTopologicalFace;
-class FCriterion;
 
 class FCriteriaGrid
 {
 protected:
-	FTopologicalFace& Face;
 
 	const FCoordinateGrid& CoordinateGrid;
 	int32 TrueUcoorindateCount;
@@ -34,18 +32,23 @@ protected:
 		return TrueVIndex * TrueUcoorindateCount + TrueUIndex;
 	}
 
-	const FPoint& GetPoint(int32 UIndex, int32 VIndex, bool bIsInterU, bool bIsInterV, FVector3f* normal = nullptr) const;
-	void Init();
-
-public:
-	FCriteriaGrid(FTopologicalFace& Surface);
-
-	FTopologicalFace& GetSurface()
+	const FPoint& GetPoint(int32 UIndex, int32 VIndex, bool bIsInternalU, bool bIsInternalV, FVector3f* OutNormal = nullptr) const
 	{
-		return Face;
+		int32 Index = GetIndex(UIndex, VIndex, bIsInternalU, bIsInternalV);
+		ensureCADKernel(Grid.Points3D.IsValidIndex(Index));
+
+		if (OutNormal)
+		{
+			*OutNormal = Grid.Normals[Index];
+		}
+		return Grid.Points3D[Index];
 	}
 
-	void ApplyCriteria(const TArray<TSharedPtr<FCriterion>>& GetCriteria) const;
+	void Init(const FTopologicalFace& Face);
+
+public:
+
+	FCriteriaGrid(const FTopologicalFace& InFace);
 
 	double GetCoordinate(EIso Iso, int32 ind) const
 	{
@@ -92,7 +95,10 @@ public:
 		return GetPoint(iU, iV, true, true, normal);
 	}
 
+#ifdef DISPLAY_CRITERIA_GRID
 	void Display() const;
+#endif
+
 };
 }
 
diff --git a/Engine/Source/Runtime/Datasmith/CADKernel/Private/CADKernel/Mesh/Criteria/SizeCriterion.h b/Engine/Source/Runtime/Datasmith/CADKernel/Private/CADKernel/Mesh/Criteria/SizeCriterion.h
--- a/Engine/Source/Runtime/Datasmith/CADKernel/Private/CADKernel/Mesh/Criteria/SizeCriterion.h
+++ b/Engine/Source/Runtime/Datasmith/CADKernel/Private/CADKernel/Mesh/Criteria/SizeCriterion.h
@@ -26,7 +26,7 @@ protected:
 
 	double ComputeSizeCriterionValue(double InDeltaU, double ChordLength) const
 	{
-		return  (ChordLength > UE_DOUBLE_KINDA_SMALL_NUMBER) ? InDeltaU * Size / ChordLength : UE_DOUBLE_BIG_NUMBER;
+		return  (ChordLength > DOUBLE_KINDA_SMALL_NUMBER) ? InDeltaU * Size / ChordLength : DOUBLE_BIG_NUMBER;
 	}
 
 public:
diff --git a/Engine/Source/Runtime/Datasmith/CADKernel/Private/CADKernel/Mesh/Meshers/ParametricMesher.cpp b/Engine/Source/Runtime/Datasmith/CADKernel/Private/CADKernel/Mesh/Meshers/ParametricMesher.cpp
--- a/Engine/Source/Runtime/Datasmith/CADKernel/Private/CADKernel/Mesh/Meshers/ParametricMesher.cpp
+++ b/Engine/Source/Runtime/Datasmith/CADKernel/Private/CADKernel/Mesh/Meshers/ParametricMesher.cpp
@@ -124,7 +124,7 @@ void FParametricMesher::MeshEntities()
 	FProgress ProgressBar(Faces.Num() * 2, TEXT("Meshing Entities : Find quad surfaces"));
 
 #ifdef CADKERNEL_DEV
-	MesherReport.Chronos.ApplyCriteriaDuration = FChrono::Elapse(ApplyCriteriaStartTime);
+	MesherReport.Chronos.ApplyCriteriaDuration = FChrono::Elapse(StartTime);
 #endif
 
 	FTimePoint MeshingStartTime = FChrono::Now();
@@ -171,13 +171,16 @@ void FParametricMesher::ApplyFaceCriteria(FTopologicalFace& Face)
 		return;
 	}
 
-	FCriteriaGrid Grid(Face);
-	if (Face.IsDeleted())
+	if (!Face.ComputeCriteriaGridSampling())
 	{
+		// The face is considered as degenerate, the face is delete, the process is canceled
 		return;
 	}
 
-	Grid.ApplyCriteria(GetMeshModel().GetCriteria());
+	FCriteriaGrid Grid(Face);
+
+	Face.InitDeltaUs();
+	Face.ApplyCriteria(GetMeshModel().GetCriteria(), Grid);
 }
 
 void FParametricMesher::ApplyEdgeCriteria(FTopologicalEdge& Edge)
diff --git a/Engine/Source/Runtime/Datasmith/CADKernel/Private/CADKernel/Topo/Shell.cpp b/Engine/Source/Runtime/Datasmith/CADKernel/Private/CADKernel/Topo/Shell.cpp
--- a/Engine/Source/Runtime/Datasmith/CADKernel/Private/CADKernel/Topo/Shell.cpp
+++ b/Engine/Source/Runtime/Datasmith/CADKernel/Private/CADKernel/Topo/Shell.cpp
@@ -49,13 +49,14 @@ FShell::FShell(const TArray<TSharedPtr<FTopologicalFace>>& InTopologicalFaces, c
 	}
 }
 
-void FShell::Empty(int32 NewSize)
+void FShell::Empty()
 {
 	for(FOrientedFace& Face : TopologicalFaces)
 	{
 		Face.Entity->ResetHost();
 	}
-	TopologicalFaces.Empty(NewSize);
+	TopologicalFaces.Empty();
+	FTopologicalShapeEntity::Empty();
 }
 
 void FShell::Add(TArray<FTopologicalFace*> Faces)
@@ -70,7 +71,6 @@ void FShell::Add(TArray<FTopologicalFace*> Faces)
 	}
 }
 
-
 void FShell::Add(TSharedRef<FTopologicalFace> InTopologicalFace, EOrientation Orientation)
 {
 	TSharedPtr<FTopologicalFace> Face = InTopologicalFace;
diff --git a/Engine/Source/Runtime/Datasmith/CADKernel/Private/CADKernel/Topo/TopologicalEdge.cpp b/Engine/Source/Runtime/Datasmith/CADKernel/Private/CADKernel/Topo/TopologicalEdge.cpp
--- a/Engine/Source/Runtime/Datasmith/CADKernel/Private/CADKernel/Topo/TopologicalEdge.cpp
+++ b/Engine/Source/Runtime/Datasmith/CADKernel/Private/CADKernel/Topo/TopologicalEdge.cpp
@@ -356,7 +356,7 @@ void FTopologicalEdge::Link(FTopologicalEdge& Twin)
 	MakeLink(Twin);
 }
 
-void FTopologicalEdge::Delete()
+void FTopologicalEdge::Empty()
 {
 	StartVertex->RemoveConnectedEdge(*this);
 	StartVertex->DeleteIfIsolated();
@@ -374,7 +374,8 @@ void FTopologicalEdge::Delete()
 	Curve.Reset();
 	Loop = nullptr;
 	Mesh.Reset();
-	SetDeleted();
+
+	TLinkable<FTopologicalEdge, FEdgeLink>::Empty();
 }
 
 FTopologicalFace* FTopologicalEdge::GetFace() const
diff --git a/Engine/Source/Runtime/Datasmith/CADKernel/Private/CADKernel/Topo/TopologicalFace.cpp b/Engine/Source/Runtime/Datasmith/CADKernel/Private/CADKernel/Topo/TopologicalFace.cpp
--- a/Engine/Source/Runtime/Datasmith/CADKernel/Private/CADKernel/Topo/TopologicalFace.cpp
+++ b/Engine/Source/Runtime/Datasmith/CADKernel/Private/CADKernel/Topo/TopologicalFace.cpp
@@ -10,6 +10,8 @@
 #include "CADKernel/Geo/Sampler/SamplerOnChord.h"
 #include "CADKernel/Geo/Sampling/Polyline.h"
 #include "CADKernel/Geo/Surfaces/Surface.h"
+#include "CADKernel/Mesh/Criteria/CriteriaGrid.h"
+#include "CADKernel/Mesh/Criteria/Criterion.h"
 #include "CADKernel/Mesh/Structure/FaceMesh.h"
 #include "CADKernel/Mesh/Structure/Grid.h"
 #include "CADKernel/Topo/Shell.h"
@@ -43,12 +45,6 @@ void FTopologicalFace::ComputeBoundary() const
 	Boundary.SetReady();
 }
 
-void FTopologicalFace::Presample()
-{
-	const FSurfacicBoundary& FaceBoundaries = GetBoundary();
-	CarrierSurface->Presample(FaceBoundaries, CrossingCoordinates);
-}
-
 #ifdef DEBUG_GET_BBOX
 #include "CADKernel/Math/Aabb.h"
 #endif
@@ -264,7 +260,7 @@ void FTopologicalFace::RemoveLoop(const TSharedPtr<FTopologicalLoop>& Loop)
 
 	if (Loops.Num() == 0)
 	{
-		SetDeleted();
+		Delete();
 	}
 }
 
@@ -489,6 +485,8 @@ TSharedRef<FFaceMesh> FTopologicalFace::GetOrCreateMesh(FModelMesh& MeshModel)
 	return Mesh.ToSharedRef();
 }
 
+// Meshing parameters ==============================================================================================================================================================================================================================
+
 void FTopologicalFace::InitDeltaUs()
 {
 	CrossingPointDeltaMins[EIso::IsoU].Init(DOUBLE_SMALL_NUMBER, CrossingCoordinates[EIso::IsoU].Num() - 1);
@@ -498,6 +496,104 @@ void FTopologicalFace::InitDeltaUs()
 	CrossingPointDeltaMaxs[EIso::IsoV].Init(HUGE_VALUE, CrossingCoordinates[EIso::IsoV].Num() - 1);
 }
 
+bool FTopologicalFace::ComputeCriteriaGridSampling()
+{
+	const FSurfacicBoundary& FaceBoundaries = GetBoundary();
+	CarrierSurface->Presample(FaceBoundaries, CrossingCoordinates);
+
+	constexpr int32 MaxGrid = 1000000;
+	if (CrossingCoordinates[EIso::IsoU].Num() * CrossingCoordinates[EIso::IsoV].Num() > MaxGrid)
+	{
+		// The sampling of the surface is huge. This is probably due to a degenerated carrier surface
+		// The face is removed
+		Remove();
+		return false;
+	}
+
+	return true;
+}
+
+void FTopologicalFace::ApplyCriteria(const TArray<TSharedPtr<FCriterion>>& Criteria, const FCriteriaGrid& Grid)
+{
+	TArray<double>& DeltaUMaxArray = CrossingPointDeltaMaxs[EIso::IsoU];
+	TArray<double>& DeltaUMinArray = CrossingPointDeltaMins[EIso::IsoU];
+	TArray<double>& DeltaVMaxArray = CrossingPointDeltaMaxs[EIso::IsoV];
+	TArray<double>& DeltaVMinArray = CrossingPointDeltaMins[EIso::IsoV];
+	FSurfaceCurvature& SurfaceCurvature = GetCurvatures();
+
+	double ElementLengthMin = DOUBLE_BIG_NUMBER;
+
+	TFunction<void(const double, const double, const double)> ElementLength = [&](const double CoordNextMinusCoord, const double Length, const double DeltaMax)
+	{
+		if (CoordNextMinusCoord < DOUBLE_SMALL_NUMBER)
+		{
+			return Length;
+		}
+		const double ElemLength = Length * DeltaMax / CoordNextMinusCoord;
+		if (ElementLengthMin > ElemLength)
+		{
+			ElementLengthMin = ElemLength;
+		}
+		return ElemLength;
+	};
+
+	for (int32 IndexV = 0; IndexV < Grid.GetCoordinateCount(EIso::IsoV) - 1; ++IndexV)
+	{
+		for (int32 IndexU = 0; IndexU < Grid.GetCoordinateCount(EIso::IsoU) - 1; ++IndexU)
+		{
+			const FPoint& Point_U0_V0 = Grid.GetPoint(IndexU, IndexV);
+			const FPoint& Point_U1_V1 = Grid.GetPoint(IndexU + 1, IndexV + 1);
+			const FPoint& Point_Um_V0 = Grid.GetIntermediateU(IndexU, IndexV);
+			const FPoint& Point_Um_V1 = Grid.GetIntermediateU(IndexU, IndexV + 1);
+			const FPoint& Point_U0_Vm = Grid.GetIntermediateV(IndexU, IndexV);
+			const FPoint& Point_U1_Vm = Grid.GetIntermediateV(IndexU + 1, IndexV);
+			const FPoint& Point_Um_Vm = Grid.GetIntermediateUV(IndexU, IndexV);
+
+			// Evaluate Sag
+			double LengthU;
+			const double SagU = FCriterion::EvaluateSag(Point_U0_Vm, Point_U1_Vm, Point_Um_Vm, LengthU);
+			double LengthV;
+			const double SagV = FCriterion::EvaluateSag(Point_Um_V0, Point_Um_V1, Point_Um_Vm, LengthV);
+			double LengthUV;
+			const double SagUV = FCriterion::EvaluateSag(Point_U0_V0, Point_U1_V1, Point_Um_Vm, LengthUV);
+
+			double& DeltaUMin = DeltaUMinArray[IndexU];
+			double& DeltaUMax = DeltaUMaxArray[IndexU];
+			double& DeltaVMin = DeltaVMinArray[IndexV];
+			double& DeltaVMax = DeltaVMaxArray[IndexV];
+
+			const double UNextMinusU = Grid.GetCoordinate(EIso::IsoU, IndexU + 1) - Grid.GetCoordinate(EIso::IsoU, IndexU);
+			const double VNextMinusV = Grid.GetCoordinate(EIso::IsoV, IndexV + 1) - Grid.GetCoordinate(EIso::IsoV, IndexV);
+
+			for (const TSharedPtr<FCriterion>& Criterion : Criteria)
+			{
+				Criterion->UpdateDelta(UNextMinusU, SagU, SagUV, SagV, LengthU, LengthUV, DeltaUMax, DeltaUMin, SurfaceCurvature[EIso::IsoU]);
+				Criterion->UpdateDelta(VNextMinusV, SagV, SagUV, SagU, LengthV, LengthUV, DeltaVMax, DeltaVMin, SurfaceCurvature[EIso::IsoV]);
+			}
+
+			ElementLength(UNextMinusU, LengthU, DeltaUMax);
+			ElementLength(VNextMinusV, LengthV, DeltaVMax);
+		}
+	}
+
+	// Delta of the extremities are smooth to avoid big disparity 
+	if (DeltaUMaxArray.Num() > 2)
+	{
+		DeltaUMaxArray[0] = (DeltaUMaxArray[0] + DeltaUMaxArray[1] * 2) * AThird;
+		DeltaUMaxArray.Last() = (DeltaUMaxArray.Last() + DeltaUMaxArray[DeltaUMaxArray.Num() - 2] * 2) * AThird;
+	}
+
+	if (DeltaVMaxArray.Num() > 2)
+	{
+		DeltaVMaxArray[0] = (DeltaVMaxArray[0] + DeltaVMaxArray[1] * 2) * AThird;
+		DeltaVMaxArray.Last() = (DeltaVMaxArray.Last() + DeltaVMaxArray[DeltaVMaxArray.Num() - 2] * 2) * AThird;
+	}
+
+	SetEstimatedMinimalElementLength(ElementLengthMin);
+	SetApplyCriteriaMarker();
+}
+
+
 // Quad ==============================================================================================================================================================================================================================
 
 double FTopologicalFace::GetQuadCriteria()
diff --git a/Engine/Source/Runtime/Datasmith/CADKernel/Private/CADKernel/Topo/Topomaker.cpp b/Engine/Source/Runtime/Datasmith/CADKernel/Private/CADKernel/Topo/Topomaker.cpp
--- a/Engine/Source/Runtime/Datasmith/CADKernel/Private/CADKernel/Topo/Topomaker.cpp
+++ b/Engine/Source/Runtime/Datasmith/CADKernel/Private/CADKernel/Topo/Topomaker.cpp
@@ -1133,7 +1133,7 @@ void FTopomaker::SplitIntoConnectedShells()
 		if (FaceSubset.MainShell != nullptr)
 		{
 			FShell* Shell = (FShell*)FaceSubset.MainShell;
-			Shell->Empty(FaceSubset.Faces.Num());
+			Shell->Empty();
 			Shell->Add(FaceSubset.Faces);
 		}
 		else
diff --git a/Engine/Source/Runtime/Datasmith/CADKernel/Public/CADKernel/Core/Entity.h b/Engine/Source/Runtime/Datasmith/CADKernel/Public/CADKernel/Core/Entity.h
--- a/Engine/Source/Runtime/Datasmith/CADKernel/Public/CADKernel/Core/Entity.h
+++ b/Engine/Source/Runtime/Datasmith/CADKernel/Public/CADKernel/Core/Entity.h
@@ -61,7 +61,13 @@ public:
 
 	virtual void Delete()
 	{
-		SetDeleted();
+		Empty();
+		SetDeletedMarker();
+	}
+
+	virtual void Empty()
+	{
+
 	}
 
 	template<typename OtherEntity, typename... InArgTypes>
diff --git a/Engine/Source/Runtime/Datasmith/CADKernel/Public/CADKernel/Core/HaveStates.h b/Engine/Source/Runtime/Datasmith/CADKernel/Public/CADKernel/Core/HaveStates.h
--- a/Engine/Source/Runtime/Datasmith/CADKernel/Public/CADKernel/Core/HaveStates.h
+++ b/Engine/Source/Runtime/Datasmith/CADKernel/Public/CADKernel/Core/HaveStates.h
@@ -116,7 +116,7 @@ public:
 		return ((States & EHaveStates::IsDeleted) == EHaveStates::IsDeleted);
 	}
 
-	void SetDeleted() const
+	void SetDeletedMarker() const
 	{
 		States |= EHaveStates::IsDeleted;
 	}
diff --git a/Engine/Source/Runtime/Datasmith/CADKernel/Public/CADKernel/Topo/Body.h b/Engine/Source/Runtime/Datasmith/CADKernel/Public/CADKernel/Topo/Body.h
--- a/Engine/Source/Runtime/Datasmith/CADKernel/Public/CADKernel/Topo/Body.h
+++ b/Engine/Source/Runtime/Datasmith/CADKernel/Public/CADKernel/Topo/Body.h
@@ -78,7 +78,7 @@ public:
 
 	virtual void Remove(const FTopologicalShapeEntity* ShellToRemove) override;
 
-	void Empty()
+	virtual void Empty() override
 	{
 		Shells.Empty();
 	}
diff --git a/Engine/Source/Runtime/Datasmith/CADKernel/Public/CADKernel/Topo/Model.h b/Engine/Source/Runtime/Datasmith/CADKernel/Public/CADKernel/Topo/Model.h
--- a/Engine/Source/Runtime/Datasmith/CADKernel/Public/CADKernel/Topo/Model.h
+++ b/Engine/Source/Runtime/Datasmith/CADKernel/Public/CADKernel/Topo/Model.h
@@ -85,9 +85,10 @@ public:
 		Bodies.Add(InBody);
 	}
 
-	void Empty()
+	virtual void Empty() override
 	{
 		Bodies.Empty();
+		FTopologicalShapeEntity::Empty();
 	}
 
 	void RemoveBody(TSharedPtr<FBody> InBody)
diff --git a/Engine/Source/Runtime/Datasmith/CADKernel/Public/CADKernel/Topo/Shell.h b/Engine/Source/Runtime/Datasmith/CADKernel/Public/CADKernel/Topo/Shell.h
--- a/Engine/Source/Runtime/Datasmith/CADKernel/Public/CADKernel/Topo/Shell.h
+++ b/Engine/Source/Runtime/Datasmith/CADKernel/Public/CADKernel/Topo/Shell.h
@@ -89,7 +89,7 @@ public:
 		ResetMarkersRecursivelyOnEntities((TArray<TOrientedEntity<FEntity>>&) TopologicalFaces);
 	}
 
-	void Empty(int32 NewSize = 0);
+	virtual void Empty() override;
 
 	void Add(TSharedRef<FTopologicalFace> InTopologicalFace, EOrientation InOrientation);
 	void Add(TArray<FTopologicalFace*> Faces);
diff --git a/Engine/Source/Runtime/Datasmith/CADKernel/Public/CADKernel/Topo/TopologicalEdge.h b/Engine/Source/Runtime/Datasmith/CADKernel/Public/CADKernel/Topo/TopologicalEdge.h
--- a/Engine/Source/Runtime/Datasmith/CADKernel/Public/CADKernel/Topo/TopologicalEdge.h
+++ b/Engine/Source/Runtime/Datasmith/CADKernel/Public/CADKernel/Topo/TopologicalEdge.h
@@ -139,7 +139,7 @@ public:
 
 	virtual ~FTopologicalEdge() = default;
 
-	virtual void Delete() override;
+	virtual void Empty() override;
 
 	virtual void Serialize(FCADKernelArchive& Ar) override
 	{
diff --git a/Engine/Source/Runtime/Datasmith/CADKernel/Public/CADKernel/Topo/TopologicalFace.h b/Engine/Source/Runtime/Datasmith/CADKernel/Public/CADKernel/Topo/TopologicalFace.h
--- a/Engine/Source/Runtime/Datasmith/CADKernel/Public/CADKernel/Topo/TopologicalFace.h
+++ b/Engine/Source/Runtime/Datasmith/CADKernel/Public/CADKernel/Topo/TopologicalFace.h
@@ -18,7 +18,10 @@
 
 namespace UE::CADKernel
 {
+
 class FCADKernelArchive;
+class FCriteriaGrid;
+class FCriterion;
 class FCurve;
 class FDatabase;
 class FModelMesh;
@@ -267,14 +270,6 @@ public:
 
 	void EvaluateGrid(FGrid& Grid) const;
 
-	// ======   Sample Functions   ======
-
-	/**
-	 * Generate a pre-sampling of the surface saved in CrossingCoordinate.
-	 * This sampling is light enough to allow a fast computation of the grid, precise enough to compute accurately meshing criteria
-	 */
-	void Presample();
-
 	/**
 	 * Update the bounding box with an approximation of the surface based of n iso curves 
 	 * @param ApproximationFactor: the factor apply to the geometric tolerance to defined the SAG error of the ISO
@@ -290,7 +285,27 @@ public:
 	{
 		Disjoin(NewBorderEdges);
 		RemoveOfHost();
-		SetDeleted();
+		Delete();
+	}
+
+	virtual void Empty() override
+	{
+#ifndef CADKERNEL_DEV
+		CarrierSurface.Reset();
+		Loops.Empty();
+#endif
+		Mesh.Reset();
+
+		MeshCuttingCoordinates.Empty();
+		CrossingCoordinates.Empty();
+		CrossingPointDeltaMins.Empty();
+		CrossingPointDeltaMaxs.Empty();
+
+		SurfaceCorners.Empty();
+		StartSideIndices.Empty();
+		SideProperties.Empty();
+
+		FTopologicalShapeEntity::Empty();
 	}
 
 	/** Has at least one non manifold edge */
@@ -372,27 +387,28 @@ public:
 		return CrossingCoordinates[Iso];
 	}
 
+	/**
+	 * Generate a pre-sampling of the surface saved in CrossingCoordinate.
+	 * This sampling is light enough to allow a fast computation of the grid, precise enough to compute accurately meshing criteria
+	 * @return false if the face is considered as degenerate and so removed
+	 */
+	bool ComputeCriteriaGridSampling();
+		
+	/**
+	 * Compute CrossingPointDeltaMins and CrossingPointDeltaMaxs respecting the meshing criteria
+	 */
+	void ApplyCriteria(const TArray<TSharedPtr<FCriterion>>& Criteria, const FCriteriaGrid& Grid);
 
 	const TArray<double>& GetCrossingPointDeltaMins(EIso Iso) const
 	{
 		return CrossingPointDeltaMins[Iso];
 	}
 
-	TArray<double>& GetCrossingPointDeltaMins(EIso Iso)
-	{
-		return CrossingPointDeltaMins[Iso];
-	}
-
 	const TArray<double>& GetCrossingPointDeltaMaxs(EIso Iso) const
 	{
 		return CrossingPointDeltaMaxs[Iso];
 	}
 
-	TArray<double>& GetCrossingPointDeltaMaxs(EIso Iso)
-	{
-		return CrossingPointDeltaMaxs[Iso];
-	}
-
 	// ======   State, Type Functions   ======
 
 	const EQuadType GetQuadType() const
diff --git a/Engine/Source/Runtime/Datasmith/CADKernel/Public/CADKernel/Topo/TopologicalLink.h b/Engine/Source/Runtime/Datasmith/CADKernel/Public/CADKernel/Topo/TopologicalLink.h
--- a/Engine/Source/Runtime/Datasmith/CADKernel/Public/CADKernel/Topo/TopologicalLink.h
+++ b/Engine/Source/Runtime/Datasmith/CADKernel/Public/CADKernel/Topo/TopologicalLink.h
@@ -50,11 +50,10 @@ public:
 		SerializeIdents(Ar, TwinEntities, false);
 	}
 
-	virtual void Delete()
+	virtual void Empty() override
 	{
 		TwinEntities.Empty();
 		ActiveEntity = nullptr;
-		SetDeleted();
 	}
 
 	const EntityType* GetActiveEntity() const
diff --git a/Engine/Source/Runtime/Datasmith/CADKernel/Public/CADKernel/Topo/TopologicalLoop.h b/Engine/Source/Runtime/Datasmith/CADKernel/Public/CADKernel/Topo/TopologicalLoop.h
--- a/Engine/Source/Runtime/Datasmith/CADKernel/Public/CADKernel/Topo/TopologicalLoop.h
+++ b/Engine/Source/Runtime/Datasmith/CADKernel/Public/CADKernel/Topo/TopologicalLoop.h
@@ -79,6 +79,12 @@ public:
 
 	~FTopologicalLoop() = default;
 
+	virtual void Empty() override
+	{
+		Edges.Empty();
+		FTopologicalEntity::Empty();
+	}
+
 	static TSharedPtr<FTopologicalLoop> Make(const TArray<TSharedPtr<FTopologicalEdge>>& EdgeList, const TArray<EOrientation>& EdgeDirections, const bool bIsExternalLoop, const double GeometricTolerance);
 
 	void DeleteLoopEdges();
diff --git a/Engine/Source/Runtime/Datasmith/CADKernel/Public/CADKernel/Topo/TopologicalShapeEntity.h b/Engine/Source/Runtime/Datasmith/CADKernel/Public/CADKernel/Topo/TopologicalShapeEntity.h
--- a/Engine/Source/Runtime/Datasmith/CADKernel/Public/CADKernel/Topo/TopologicalShapeEntity.h
+++ b/Engine/Source/Runtime/Datasmith/CADKernel/Public/CADKernel/Topo/TopologicalShapeEntity.h
@@ -25,6 +25,12 @@ public:
 		SerializeIdent(Ar, &HostedBy);
 	}
 
+	virtual void Empty() override 
+	{
+		HostedBy = nullptr;
+		FTopologicalEntity::Empty();
+	}
+
 	const FMetadataDictionary& GetMetadataDictionary() const
 	{
 		return Dictionary;
diff --git a/Engine/Source/Runtime/Datasmith/CADKernel/Public/CADKernel/Topo/TopologicalVertex.h b/Engine/Source/Runtime/Datasmith/CADKernel/Public/CADKernel/Topo/TopologicalVertex.h
--- a/Engine/Source/Runtime/Datasmith/CADKernel/Public/CADKernel/Topo/TopologicalVertex.h
+++ b/Engine/Source/Runtime/Datasmith/CADKernel/Public/CADKernel/Topo/TopologicalVertex.h
@@ -168,10 +168,17 @@ public:
 				}
 				TopologicalLink.Reset();
 			}
-			SetDeleted();
+			Delete();
 		}
 	}
 
+	virtual void Empty() override
+	{
+		ConnectedEdges.Empty();
+		Mesh.Reset();
+		TLinkable<FTopologicalVertex, FVertexLink>::Empty();
+	}
+
 	bool IsBorderVertex();
 
 	void AddConnectedEdge(FTopologicalEdge& Edge);
diff --git a/Engine/Source/Runtime/Datasmith/CADKernel/Public/CADKernel/Topo/Topomaker.h b/Engine/Source/Runtime/Datasmith/CADKernel/Public/CADKernel/Topo/Topomaker.h
--- a/Engine/Source/Runtime/Datasmith/CADKernel/Public/CADKernel/Topo/Topomaker.h
+++ b/Engine/Source/Runtime/Datasmith/CADKernel/Public/CADKernel/Topo/Topomaker.h
@@ -30,6 +30,33 @@ enum class ESewOption : uint8
 
 ENUM_CLASS_FLAGS(ESewOption);
 
+namespace SewOption
+{
+
+static ESewOption GetFromOptions(bool bGStitchingForceSew, bool bGStitchingRemoveThinFaces, bool bGStitchingRemoveDuplicatedFaces)
+{
+	ESewOption Option = ESewOption::None;
+
+	if (bGStitchingForceSew)
+	{
+		Option |= ESewOption::ForceJoining;
+	}
+
+	if (bGStitchingRemoveThinFaces)
+	{
+		Option |= ESewOption::RemoveThinFaces;
+	}
+
+	if (bGStitchingRemoveDuplicatedFaces)
+	{
+		Option |= ESewOption::RemoveDuplicatedFaces;
+	}
+
+	return Option;
+}
+
+}
+
 struct FTopomakerOptions
 {
 	ESewOption SewOptions;
