[CADKernel] Multithreading  Apply Surface Criteria

diff --git a/Engine/Source/Runtime/Datasmith/CADKernel/Private/CADKernel/Mesh/Meshers/ParametricMesher.cpp b/Engine/Source/Runtime/Datasmith/CADKernel/Private/CADKernel/Mesh/Meshers/ParametricMesher.cpp
--- a/Engine/Source/Runtime/Datasmith/CADKernel/Private/CADKernel/Mesh/Meshers/ParametricMesher.cpp
+++ b/Engine/Source/Runtime/Datasmith/CADKernel/Private/CADKernel/Mesh/Meshers/ParametricMesher.cpp
@@ -16,6 +16,8 @@
 #include "CADKernel/Topo/TopologicalFace.h"
 #include "CADKernel/Utils/Util.h"
 
+#include "Async/ParallelFor.h"
+
 #ifdef CADKERNEL_DEV
 #include "CADKernel/Mesh/Meshers/MesherReport.h"
 #endif
@@ -97,18 +99,19 @@ void FParametricMesher::PreMeshingTasks()
 
 	FProgress ProgressBar(Faces.Num() * 2, TEXT("Meshing Entities : Apply Surface Criteria"));
 
-	for (FTopologicalFace* Face : Faces)
-	{
-		FProgress _(1, TEXT("Meshing Entities : Apply Surface Criteria"));
+	const TArray<TSharedPtr<FCriterion>>& Criteria = GetMeshModel().GetCriteria();
 
-		if (Face == nullptr || Face->IsDeleted() || Face->IsMeshed())
+	TArray<FTopologicalFace*>& LocalFaces = Faces;
+	ParallelFor(Faces.Num(), [&LocalFaces, &Criteria](int32 Index)
 		{
-			continue;
-		}
-
-		ApplyFaceCriteria(*Face);
-		Face->ComputeSurfaceSideProperties();
-	}
+			FTopologicalFace* Face = LocalFaces[Index];
+			if (Face == nullptr || Face->IsDeleted() || Face->IsMeshed())
+			{
+				return;
+			}
+			ApplyFaceCriteria(*Face, Criteria);
+			Face->ComputeSurfaceSideProperties();
+		});
 
 #ifdef CADKERNEL_DEV
 	MesherReport.Chronos.ApplyCriteriaDuration = FChrono::Elapse(ApplyCriteriaStartTime);
@@ -164,7 +167,7 @@ void FParametricMesher::MeshEntities()
 
 }
 
-void FParametricMesher::ApplyFaceCriteria(FTopologicalFace& Face)
+void FParametricMesher::ApplyFaceCriteria(FTopologicalFace& Face, const TArray<TSharedPtr<FCriterion>>& Criteria)
 {
 	if (Face.IsApplyCriteria())
 	{
@@ -180,7 +183,7 @@ void FParametricMesher::ApplyFaceCriteria(FTopologicalFace& Face)
 	FCriteriaGrid Grid(Face);
 
 	Face.InitDeltaUs();
-	Face.ApplyCriteria(GetMeshModel().GetCriteria(), Grid);
+	Face.ApplyCriteria(Criteria, Grid);
 }
 
 void FParametricMesher::ApplyEdgeCriteria(FTopologicalEdge& Edge)
diff --git a/Engine/Source/Runtime/Datasmith/CADKernel/Public/CADKernel/Mesh/Meshers/ParametricMesher.h b/Engine/Source/Runtime/Datasmith/CADKernel/Public/CADKernel/Mesh/Meshers/ParametricMesher.h
--- a/Engine/Source/Runtime/Datasmith/CADKernel/Public/CADKernel/Mesh/Meshers/ParametricMesher.h
+++ b/Engine/Source/Runtime/Datasmith/CADKernel/Public/CADKernel/Mesh/Meshers/ParametricMesher.h
@@ -112,7 +112,7 @@ protected:
 	void MeshSurfaceByFront(TArray<FCostToFace>& QuadTrimmedSurfaceSet);
 
 	void ApplyEdgeCriteria(FTopologicalEdge& Edge);
-	void ApplyFaceCriteria(FTopologicalFace& Face);
+	static void ApplyFaceCriteria(FTopologicalFace& Face, const TArray<TSharedPtr<FCriterion>>& Criteria);
 
 	/**
 	 * Generate Edge Elements on active edge from Edge cutting points
