Fix Delete Entity
Fix FTopologicalEdge::ReplaceEdgeVertex (avoid ensure in FVertexLink::DefineActiveEntity())
Fix StitchParallelEdges

diff --git a/Engine/Source/Runtime/Datasmith/CADKernel/Private/CADKernel/Topo/Shell.cpp b/Engine/Source/Runtime/Datasmith/CADKernel/Private/CADKernel/Topo/Shell.cpp
--- a/Engine/Source/Runtime/Datasmith/CADKernel/Private/CADKernel/Topo/Shell.cpp
+++ b/Engine/Source/Runtime/Datasmith/CADKernel/Private/CADKernel/Topo/Shell.cpp
@@ -395,7 +395,7 @@ int32 FShell::Orient()
 
 		for (FTopologicalFace* Face : Subshell)
 		{
-			if (Face == nullptr)
+			if (Face == nullptr || Face->IsDeleted())
 			{
 				continue;
 			}
diff --git a/Engine/Source/Runtime/Datasmith/CADKernel/Private/CADKernel/Topo/TopologicalEdge.cpp b/Engine/Source/Runtime/Datasmith/CADKernel/Private/CADKernel/Topo/TopologicalEdge.cpp
--- a/Engine/Source/Runtime/Datasmith/CADKernel/Private/CADKernel/Topo/TopologicalEdge.cpp
+++ b/Engine/Source/Runtime/Datasmith/CADKernel/Private/CADKernel/Topo/TopologicalEdge.cpp
@@ -358,13 +358,18 @@ void FTopologicalEdge::Link(FTopologicalEdge& Twin)
 
 void FTopologicalEdge::Empty()
 {
-	StartVertex->RemoveConnectedEdge(*this);
-	StartVertex->DeleteIfIsolated();
-	EndVertex->RemoveConnectedEdge(*this);
-	EndVertex->DeleteIfIsolated();
-
-	StartVertex.Reset();
-	EndVertex.Reset();
+	if (StartVertex.IsValid())
+	{
+		StartVertex->RemoveConnectedEdge(*this);
+		StartVertex->DeleteIfIsolated();
+		StartVertex.Reset();
+	}
+	if (EndVertex.IsValid())
+	{
+		EndVertex->RemoveConnectedEdge(*this);
+		EndVertex->DeleteIfIsolated();
+		EndVertex.Reset();
+	}
 
 	if (TopologicalLink.IsValid())
 	{
@@ -945,7 +950,10 @@ void FTopologicalEdge::ReplaceEdgeVertex(bool bIsStartVertex, TSharedRef<FTopolo
 	NewVertex->AddConnectedEdge(*this);
 
 	TSharedPtr<FTopologicalVertex>& OldVertex = bIsStartVertex ? StartVertex : EndVertex;
-	OldVertex->Link(*NewVertex);
+	if (OldVertex->GetTwinEntityCount() > 1)
+	{
+		OldVertex->Link(*NewVertex);
+	}
 
 	OldVertex->RemoveConnectedEdge(*this);
 
diff --git a/Engine/Source/Runtime/Datasmith/CADKernel/Private/CADKernel/Topo/TopologicalVertex.cpp b/Engine/Source/Runtime/Datasmith/CADKernel/Private/CADKernel/Topo/TopologicalVertex.cpp
--- a/Engine/Source/Runtime/Datasmith/CADKernel/Private/CADKernel/Topo/TopologicalVertex.cpp
+++ b/Engine/Source/Runtime/Datasmith/CADKernel/Private/CADKernel/Topo/TopologicalVertex.cpp
@@ -46,7 +46,7 @@ bool FTopologicalVertex::IsBorderVertex()
 {
 	for (FTopologicalVertex* Vertex : GetTwinEntities())
 	{
-		if (Vertex != nullptr)
+		if (Vertex != nullptr && !Vertex->IsDeleted())
 		{
 			for (const FTopologicalEdge* Edge : Vertex->GetDirectConnectedEdges())
 			{
diff --git a/Engine/Source/Runtime/Datasmith/CADKernel/Private/CADKernel/Topo/Topomaker.cpp b/Engine/Source/Runtime/Datasmith/CADKernel/Private/CADKernel/Topo/Topomaker.cpp
--- a/Engine/Source/Runtime/Datasmith/CADKernel/Private/CADKernel/Topo/Topomaker.cpp
+++ b/Engine/Source/Runtime/Datasmith/CADKernel/Private/CADKernel/Topo/Topomaker.cpp
@@ -307,6 +307,11 @@ void StitchParallelEdges(TArray<TSharedPtr<FTopologicalVertex>>& VerticesToProce
 						VerticesToProcess.Add(NewVertex);
 					}
 				}
+
+				if (!Vertex.IsValid() || Vertex->IsDeleted() || !Vertex->IsBorderVertex())
+				{
+					break;
+				}
 			}
 		}
 	}
diff --git a/Engine/Source/Runtime/Datasmith/CADKernel/Public/CADKernel/Topo/Linkable.h b/Engine/Source/Runtime/Datasmith/CADKernel/Public/CADKernel/Topo/Linkable.h
--- a/Engine/Source/Runtime/Datasmith/CADKernel/Public/CADKernel/Topo/Linkable.h
+++ b/Engine/Source/Runtime/Datasmith/CADKernel/Public/CADKernel/Topo/Linkable.h
@@ -19,6 +19,11 @@ protected:
 public:
 	TLinkable() = default;
 
+	virtual ~TLinkable() override
+	{
+		Empty();
+	}
+
 	void Finalize()
 	{
 		ensureCADKernel(!TopologicalLink.IsValid());
@@ -31,6 +36,11 @@ public:
 		SerializeIdent(Ar, TopologicalLink);
 	}
 
+	virtual void Empty() override
+	{
+		TopologicalLink.Reset();
+	}
+
 	const TSharedRef<const EntityType> GetLinkActiveEntity() const
 	{
 		ensureCADKernel(TopologicalLink.IsValid());
diff --git a/Engine/Source/Runtime/Datasmith/CADKernel/Public/CADKernel/Topo/TopologicalEdge.h b/Engine/Source/Runtime/Datasmith/CADKernel/Public/CADKernel/Topo/TopologicalEdge.h
--- a/Engine/Source/Runtime/Datasmith/CADKernel/Public/CADKernel/Topo/TopologicalEdge.h
+++ b/Engine/Source/Runtime/Datasmith/CADKernel/Public/CADKernel/Topo/TopologicalEdge.h
@@ -137,7 +137,10 @@ public:
 	 */
 	static TSharedPtr<FTopologicalEdge> ReturnIfValid(TSharedRef<FTopologicalEdge>& InEdge, bool bCheckVertices);
 
-	virtual ~FTopologicalEdge() = default;
+	virtual ~FTopologicalEdge() override
+	{
+		Empty();
+	}
 
 	virtual void Empty() override;
 
diff --git a/Engine/Source/Runtime/Datasmith/CADKernel/Public/CADKernel/Topo/TopologicalFace.h b/Engine/Source/Runtime/Datasmith/CADKernel/Public/CADKernel/Topo/TopologicalFace.h
--- a/Engine/Source/Runtime/Datasmith/CADKernel/Public/CADKernel/Topo/TopologicalFace.h
+++ b/Engine/Source/Runtime/Datasmith/CADKernel/Public/CADKernel/Topo/TopologicalFace.h
@@ -99,6 +99,11 @@ protected:
 
 public:
 
+	virtual ~FTopologicalFace() override
+	{
+		Empty();
+	}
+
 	virtual void Serialize(FCADKernelArchive& Ar) override
 	{
 		FTopologicalShapeEntity::Serialize(Ar);
@@ -292,10 +297,15 @@ public:
 	{
 #ifndef CADKERNEL_DEV
 		CarrierSurface.Reset();
+		for (TSharedPtr<FTopologicalLoop>& Loop : Loops)
+		{
+			Loop->Empty();
+		}
 		Loops.Empty();
 #endif
-		Mesh.Reset();
+		RemoveOfHost();
 
+		Mesh.Reset();
 		MeshCuttingCoordinates.Empty();
 		CrossingCoordinates.Empty();
 		CrossingPointDeltaMins.Empty();
diff --git a/Engine/Source/Runtime/Datasmith/CADKernel/Public/CADKernel/Topo/TopologicalLink.h b/Engine/Source/Runtime/Datasmith/CADKernel/Public/CADKernel/Topo/TopologicalLink.h
--- a/Engine/Source/Runtime/Datasmith/CADKernel/Public/CADKernel/Topo/TopologicalLink.h
+++ b/Engine/Source/Runtime/Datasmith/CADKernel/Public/CADKernel/Topo/TopologicalLink.h
@@ -35,6 +35,10 @@ protected:
 	}
 
 public:
+	virtual ~TTopologicalLink<EntityType>() override
+	{
+		Empty();
+	}
 
 	virtual void Serialize(FCADKernelArchive& Ar) override
 	{
@@ -70,7 +74,7 @@ public:
 
 	int32 GetTwinEntityNum() const
 	{
-		return (int32)TwinEntities.Num();
+		return TwinEntities.Num();
 	}
 
 	const TArray<EntityType*>& GetTwinEntities() const
diff --git a/Engine/Source/Runtime/Datasmith/CADKernel/Public/CADKernel/Topo/TopologicalLoop.h b/Engine/Source/Runtime/Datasmith/CADKernel/Public/CADKernel/Topo/TopologicalLoop.h
--- a/Engine/Source/Runtime/Datasmith/CADKernel/Public/CADKernel/Topo/TopologicalLoop.h
+++ b/Engine/Source/Runtime/Datasmith/CADKernel/Public/CADKernel/Topo/TopologicalLoop.h
@@ -81,6 +81,10 @@ public:
 
 	virtual void Empty() override
 	{
+		for (FOrientedEdge& Edge : Edges)
+		{
+			Edge.Entity->Empty();
+		}
 		Edges.Empty();
 		FTopologicalEntity::Empty();
 	}
diff --git a/Engine/Source/Runtime/Datasmith/CADKernel/Public/CADKernel/Topo/TopologicalShapeEntity.h b/Engine/Source/Runtime/Datasmith/CADKernel/Public/CADKernel/Topo/TopologicalShapeEntity.h
--- a/Engine/Source/Runtime/Datasmith/CADKernel/Public/CADKernel/Topo/TopologicalShapeEntity.h
+++ b/Engine/Source/Runtime/Datasmith/CADKernel/Public/CADKernel/Topo/TopologicalShapeEntity.h
@@ -133,7 +133,11 @@ public:
 
 	void RemoveOfHost()
 	{
-		GetHost()->Remove(this);
+		if (FTopologicalShapeEntity* Host = GetHost())
+		{
+			Host->Remove(this);
+			ResetHost();
+		}
 	}
 
 	virtual void Remove(const FTopologicalShapeEntity*) = 0;
diff --git a/Engine/Source/Runtime/Datasmith/CADKernel/Public/CADKernel/Topo/TopologicalVertex.h b/Engine/Source/Runtime/Datasmith/CADKernel/Public/CADKernel/Topo/TopologicalVertex.h
--- a/Engine/Source/Runtime/Datasmith/CADKernel/Public/CADKernel/Topo/TopologicalVertex.h
+++ b/Engine/Source/Runtime/Datasmith/CADKernel/Public/CADKernel/Topo/TopologicalVertex.h
@@ -36,6 +36,11 @@ protected:
 
 public:
 
+	virtual ~FTopologicalVertex() override
+	{
+		Empty();
+	}
+
 	static TSharedRef<FTopologicalVertex> Make(const FPoint& InCoordinate)
 	{
 		TSharedRef<FTopologicalVertex> Vertex = FEntity::MakeShared<FTopologicalVertex>(InCoordinate);
