+  Reduce the control of intersection with iso line

https://horde.devtools.epicgames.com/job/640216cd5f3b94d292a9508f

diff --git a/Engine/Plugins/Enterprise/DatasmithCADImporter/Source/CADTools/Private/CADOption.cpp b/Engine/Plugins/Enterprise/DatasmithCADImporter/Source/CADTools/Private/CADOption.cpp
--- a/Engine/Plugins/Enterprise/DatasmithCADImporter/Source/CADTools/Private/CADOption.cpp
+++ b/Engine/Plugins/Enterprise/DatasmithCADImporter/Source/CADTools/Private/CADOption.cpp
@@ -81,14 +81,15 @@ Only available with CADKernel.\n\
 Default value of ForceFactor is 5.\n"),
 ECVF_Default);
 
-bool FImportParameters::bGStitchingRemoveThinFaces = true;
-FAutoConsoleVariableRef GCADTranslatorStitchingRemoveThinFaces(
+bool FImportParameters::bGStitchingRemoveThinFaces = false;
+/*FAutoConsoleVariableRef GCADTranslatorStitchingRemoveThinFaces(
 	TEXT("ds.CADTranslator.Stitching.RemoveThinFaces"),
 	FImportParameters::bGStitchingRemoveThinFaces,
 	TEXT("During the welding process, Thin faces are removed before the stiching process. The width of the thin faces is equal to the force sew tolerance.\n\
 Only available with CADKernel.\n\
 Default value of RemoveThinFaces is true\n"),
 ECVF_Default);
+*/
 
 bool FImportParameters::bGStitchingRemoveDuplicatedFaces = true;
 FAutoConsoleVariableRef GCADTranslatorStitchingRemoveDuplicatedFaces(
diff --git a/Engine/Source/Runtime/Datasmith/CADKernel/Private/CADKernel/Mesh/Meshers/IsoTriangulator.cpp b/Engine/Source/Runtime/Datasmith/CADKernel/Private/CADKernel/Mesh/Meshers/IsoTriangulator.cpp
--- a/Engine/Source/Runtime/Datasmith/CADKernel/Private/CADKernel/Mesh/Meshers/IsoTriangulator.cpp
+++ b/Engine/Source/Runtime/Datasmith/CADKernel/Private/CADKernel/Mesh/Meshers/IsoTriangulator.cpp
@@ -1771,21 +1771,43 @@ void FIsoTriangulator::MeshCycle(const EGridSpace Space, const TArray<FIsoSegmen
 	SegmentStack.Reserve(5 * NodeCycleNum);
 
 	{
+		using BoolIter = TArray<bool>::RangedForConstIteratorType;
+		using SegmentIter = TArray<FIsoSegment*>::RangedForConstIteratorType;
+
+		TFunction<bool(FIsoSegment*, EIsoSegmentStates)> HasTriangleOn = [](FIsoSegment* Segment, EIsoSegmentStates Side) -> bool
+		{
+			if (Segment->HasTriangleOn(Side))
+			{
+#ifdef CADKERNEL_DEV
+				Wait();
+				ensureCADKernel(false);
+#endif
+				return true;
+			}
+			Segment->SetHasTriangleOn(Side);
+			return false;
+		};
+
+
 		// Get cycle's nodes and set segments as they have a triangle outside the cycle (to don't try to mesh outside the cycle)
-		auto SegmentOrientation = CycleOrientation.begin();
-		for (auto Segment = Cycle.begin(); Segment != Cycle.end(); ++Segment, ++SegmentOrientation)
+		BoolIter SegmentOrientation = CycleOrientation.begin();
+		for (SegmentIter Segment = Cycle.begin(); Segment != Cycle.end(); ++Segment, ++SegmentOrientation)
 		{
 			if (*SegmentOrientation)
 			{
 				CycleNodes.Add(&(*Segment)->GetFirstNode());
-				ensureCADKernel(!(*Segment)->HasTriangleOnRight());
-				(*Segment)->SetHasTriangleOnRight();
+				if (HasTriangleOn(*Segment, EIsoSegmentStates::RightTriangle))
+				{
+					return;
+				}
 			}
 			else
 			{
 				CycleNodes.Add(&(*Segment)->GetSecondNode());
-				ensureCADKernel(!(*Segment)->HasTriangleOnLeft());
-				(*Segment)->SetHasTriangleOnLeft();
+				if (HasTriangleOn(*Segment, EIsoSegmentStates::LeftTriangle))
+				{
+					return;
+				}
 			}
 		}
 
@@ -2010,12 +2032,20 @@ void FIsoTriangulator::MeshCycle(const EGridSpace Space, const TArray<FIsoSegmen
 			double CandidateSlopeAtStartNode = 8.;
 			double CandidateSlopeAtEndNode = 8.;
 
-
 			// To avoid to create triangle crossing natural iso line (a line carrying iso edges)
 			// the iso edges connecting an inside node to an external node (vs the loop) are added to InnerToOuterSegmentsIntersectionTool
 			// the selected triangle is the triangle minimizing the number of intersection with these lines
+			int32 MaxMaxIntersections = 1;
+			for (FIsoNode* Node : CycleNodes)
+			{
+				if (Node->IsALoopNode())
+				{
+					MaxMaxIntersections++;
+				}
+			}
+
 			int32 MaxIntersections = 1;
-			for (int32 IntersectionCountAllow = 0; IntersectionCountAllow < MaxIntersections; IntersectionCountAllow++)
+			for (int32 IntersectionCountAllow = 0; IntersectionCountAllow < MaxIntersections && IntersectionCountAllow <= MaxMaxIntersections; IntersectionCountAllow++)
 			{
 				for (FIsoNode* Node : CycleNodes)
 				{
@@ -2119,15 +2149,18 @@ void FIsoTriangulator::MeshCycle(const EGridSpace Space, const TArray<FIsoSegmen
 							Wait();
 						}
 #endif
-						int32 IntersectionCount = InnerToOuterSegmentsIntersectionTool.CountIntersections(StartNode, *Node);
-						IntersectionCount = FMath::Max(IntersectionCount, InnerToOuterSegmentsIntersectionTool.CountIntersections(EndNode, *Node));
-						if (IntersectionCount > IntersectionCountAllow)
+						if (IntersectionCountAllow < MaxMaxIntersections)
 						{
-							if (IntersectionCount >= MaxIntersections)
+							int32 IntersectionCount = InnerToOuterSegmentsIntersectionTool.CountIntersections(StartNode, *Node);
+							IntersectionCount = FMath::Max(IntersectionCount, InnerToOuterSegmentsIntersectionTool.CountIntersections(EndNode, *Node));
+							if (IntersectionCount > IntersectionCountAllow)
 							{
-								MaxIntersections = IntersectionCount + 1;
+								if (IntersectionCount >= MaxIntersections)
+								{
+									MaxIntersections = IntersectionCount + 1;
+								}
+								continue;
 							}
-							continue;
 						}
 
 						MinCriteria = PointCriteria;
@@ -3153,13 +3186,15 @@ void FIsoTriangulator::TryToConnectTwoSubLoopsWithTheMostIsoSegment(FCell& Cell,
 #endif
 
 	const double FlatSlope = 0.10; // ~5 deg: The segment must make an angle less than 10 deg with the Iso
-	double MinSlope = FlatSlope + DOUBLE_SMALL_NUMBER;// 0.25; // ~15 deg: The segment must make an angle less than 10 deg with the Iso
 
 	for (FLoopNode* CandidateA : LoopA)
 	{
 		FLoopNode* CandidateB = nullptr;
 		const FPoint2D& ACoordinates = CandidateA->Get2DPoint(EGridSpace::UniformScaled, Grid);
 
+		double MinSlope = FlatSlope + DOUBLE_SMALL_NUMBER;// 0.25; // ~15 deg: The segment must make an angle less than 10 deg with the Iso
+		double MinLengthSquare = HUGE_VALUE;
+
 		for (FLoopNode* NodeB : LoopB)
 		{
 			const FPoint2D& BCoordinates = NodeB->Get2DPoint(EGridSpace::UniformScaled, Grid);
@@ -3168,6 +3203,16 @@ void FIsoTriangulator::TryToConnectTwoSubLoopsWithTheMostIsoSegment(FCell& Cell,
 			if (Slope < MinSlope)
 			{
 				MinSlope = Slope;
+				// If the slope of the candidate segments is nearly zero, then select the shortest
+				if (MinSlope < DOUBLE_KINDA_SMALL_NUMBER)
+				{
+					double DistanceSquare = BCoordinates.SquareDistance(ACoordinates);
+					if (DistanceSquare > MinLengthSquare)
+					{
+						continue;
+					}
+					MinLengthSquare = DistanceSquare;
+				}
 				CandidateB = NodeB;
 			}
 		}
diff --git a/Engine/Source/Runtime/Datasmith/CADKernel/Private/CADKernel/Mesh/Meshers/IsoTriangulator/IsoSegment.h b/Engine/Source/Runtime/Datasmith/CADKernel/Private/CADKernel/Mesh/Meshers/IsoTriangulator/IsoSegment.h
--- a/Engine/Source/Runtime/Datasmith/CADKernel/Private/CADKernel/Mesh/Meshers/IsoTriangulator/IsoSegment.h
+++ b/Engine/Source/Runtime/Datasmith/CADKernel/Private/CADKernel/Mesh/Meshers/IsoTriangulator/IsoSegment.h
@@ -173,6 +173,16 @@ public:
 		States |= EIsoSegmentStates::RightCycle;
 	}
 
+	bool HasTriangleOn(EIsoSegmentStates Side)
+	{
+		return (States & Side) == Side;
+	}
+
+	void SetHasTriangleOn(EIsoSegmentStates Side)
+	{
+		States |= Side;
+	}
+
 	bool HasTriangleOnLeft()
 	{
 		return (States & EIsoSegmentStates::LeftTriangle) == EIsoSegmentStates::LeftTriangle;
diff --git a/Engine/Source/Runtime/Datasmith/CADKernel/Private/CADKernel/Mesh/Structure/Grid.h b/Engine/Source/Runtime/Datasmith/CADKernel/Private/CADKernel/Mesh/Structure/Grid.h
--- a/Engine/Source/Runtime/Datasmith/CADKernel/Private/CADKernel/Mesh/Structure/Grid.h
+++ b/Engine/Source/Runtime/Datasmith/CADKernel/Private/CADKernel/Mesh/Structure/Grid.h
@@ -238,12 +238,21 @@ protected:
 	 */
 	void ComputeNewCoordinate(const TArray<FPoint2D>& NewGrid, int32 IndexU, int32 IndexV, const FPoint2D& InPoint, FPoint2D& OutNewScaledPoint) const
 	{
+		const FPoint2D& PointU0V0 = NewGrid[(IndexV + 0) * CuttingCount[EIso::IsoU] + (IndexU + 0)];
+		const FPoint2D& PointU1V0 = NewGrid[(IndexV + 0) * CuttingCount[EIso::IsoU] + (IndexU + 1)];
+		const FPoint2D& PointU0V1 = NewGrid[(IndexV + 1) * CuttingCount[EIso::IsoU] + (IndexU + 0)];
+		const FPoint2D& PointU1V1 = NewGrid[(IndexV + 1) * CuttingCount[EIso::IsoU] + (IndexU + 1)];
+
+		const double U1MinusU0 = CuttingCoordinates[EIso::IsoU][IndexU + 1] - CuttingCoordinates[EIso::IsoU][IndexU];
+		const double U0MinusU = CuttingCoordinates[EIso::IsoU][IndexU] - InPoint.U;
+		const double V1MinusV0 = CuttingCoordinates[EIso::IsoV][IndexV + 1] - CuttingCoordinates[EIso::IsoV][IndexV];
+		const double V0MinusV = CuttingCoordinates[EIso::IsoV][IndexV] - InPoint.V;
+
 		OutNewScaledPoint =
-			NewGrid[(IndexV + 0) * CuttingCount[EIso::IsoU] + (IndexU + 0)] * (CuttingCoordinates[EIso::IsoU][IndexU + 1] - InPoint.U) * (CuttingCoordinates[EIso::IsoV][IndexV + 1] - InPoint.V) -
-			NewGrid[(IndexV + 0) * CuttingCount[EIso::IsoU] + (IndexU + 1)] * (CuttingCoordinates[EIso::IsoU][IndexU + 0] - InPoint.U) * (CuttingCoordinates[EIso::IsoV][IndexV + 1] - InPoint.V) -
-			NewGrid[(IndexV + 1) * CuttingCount[EIso::IsoU] + (IndexU + 0)] * (CuttingCoordinates[EIso::IsoU][IndexU + 1] - InPoint.U) * (CuttingCoordinates[EIso::IsoV][IndexV + 0] - InPoint.V) +
-			NewGrid[(IndexV + 1) * CuttingCount[EIso::IsoU] + (IndexU + 1)] * (CuttingCoordinates[EIso::IsoU][IndexU + 0] - InPoint.U) * (CuttingCoordinates[EIso::IsoV][IndexV + 0] - InPoint.V);
-		OutNewScaledPoint /= (CuttingCoordinates[EIso::IsoU][IndexU + 1] - CuttingCoordinates[EIso::IsoU][IndexU + 0]) * (CuttingCoordinates[EIso::IsoV][IndexV + 1] - CuttingCoordinates[EIso::IsoV][IndexV + 0]);
+			PointU0V0 +
+			(PointU0V0 - PointU1V0) * (U0MinusU / U1MinusU0) +
+			(PointU0V0 - PointU0V1) * (V0MinusV / V1MinusV0) +
+			(PointU0V0 - PointU0V1 + PointU1V1 - PointU1V0) * ((U0MinusU * V0MinusV) / (U1MinusU0 * V1MinusV0));		
 	};
 
 	/**
