Fix warning V1053: Calling the 'Empty' virtual function in the destructor may lead to unexpected result at runtime.
#rnx

diff --git a/Engine/Source/Runtime/Datasmith/CADKernel/Public/CADKernel/Topo/Linkable.h b/Engine/Source/Runtime/Datasmith/CADKernel/Public/CADKernel/Topo/Linkable.h
--- a/Engine/Source/Runtime/Datasmith/CADKernel/Public/CADKernel/Topo/Linkable.h
+++ b/Engine/Source/Runtime/Datasmith/CADKernel/Public/CADKernel/Topo/Linkable.h
@@ -21,7 +21,7 @@ public:
 
 	virtual ~TLinkable() override
 	{
-		Empty();
+		TLinkable::Empty();
 	}
 
 	void Finalize()
diff --git a/Engine/Source/Runtime/Datasmith/CADKernel/Public/CADKernel/Topo/TopologicalEdge.h b/Engine/Source/Runtime/Datasmith/CADKernel/Public/CADKernel/Topo/TopologicalEdge.h
--- a/Engine/Source/Runtime/Datasmith/CADKernel/Public/CADKernel/Topo/TopologicalEdge.h
+++ b/Engine/Source/Runtime/Datasmith/CADKernel/Public/CADKernel/Topo/TopologicalEdge.h
@@ -139,7 +139,7 @@ public:
 
 	virtual ~FTopologicalEdge() override
 	{
-		Empty();
+		FTopologicalEdge::Empty();
 	}
 
 	virtual void Empty() override;
diff --git a/Engine/Source/Runtime/Datasmith/CADKernel/Public/CADKernel/Topo/TopologicalFace.h b/Engine/Source/Runtime/Datasmith/CADKernel/Public/CADKernel/Topo/TopologicalFace.h
--- a/Engine/Source/Runtime/Datasmith/CADKernel/Public/CADKernel/Topo/TopologicalFace.h
+++ b/Engine/Source/Runtime/Datasmith/CADKernel/Public/CADKernel/Topo/TopologicalFace.h
@@ -101,7 +101,7 @@ public:
 
 	virtual ~FTopologicalFace() override
 	{
-		Empty();
+		FTopologicalFace::Empty();
 	}
 
 	virtual void Serialize(FCADKernelArchive& Ar) override
diff --git a/Engine/Source/Runtime/Datasmith/CADKernel/Public/CADKernel/Topo/TopologicalLink.h b/Engine/Source/Runtime/Datasmith/CADKernel/Public/CADKernel/Topo/TopologicalLink.h
--- a/Engine/Source/Runtime/Datasmith/CADKernel/Public/CADKernel/Topo/TopologicalLink.h
+++ b/Engine/Source/Runtime/Datasmith/CADKernel/Public/CADKernel/Topo/TopologicalLink.h
@@ -35,9 +35,9 @@ protected:
 	}
 
 public:
-	virtual ~TTopologicalLink<EntityType>() override
+	virtual ~TTopologicalLink() override
 	{
-		Empty();
+		TTopologicalLink::Empty();
 	}
 
 	virtual void Serialize(FCADKernelArchive& Ar) override
diff --git a/Engine/Source/Runtime/Datasmith/CADKernel/Public/CADKernel/Topo/TopologicalVertex.h b/Engine/Source/Runtime/Datasmith/CADKernel/Public/CADKernel/Topo/TopologicalVertex.h
--- a/Engine/Source/Runtime/Datasmith/CADKernel/Public/CADKernel/Topo/TopologicalVertex.h
+++ b/Engine/Source/Runtime/Datasmith/CADKernel/Public/CADKernel/Topo/TopologicalVertex.h
@@ -38,7 +38,7 @@ public:
 
 	virtual ~FTopologicalVertex() override
 	{
-		Empty();
+		FTopologicalVertex::Empty();
 	}
 
 	static TSharedRef<FTopologicalVertex> Make(const FPoint& InCoordinate)
