+ Fix apply criteria
 + Improve delete entity

diff --git a/Engine/Source/Runtime/Datasmith/CADKernel/Private/CADKernel/Topo/Shell.cpp b/Engine/Source/Runtime/Datasmith/CADKernel/Private/CADKernel/Topo/Shell.cpp
--- a/Engine/Source/Runtime/Datasmith/CADKernel/Private/CADKernel/Topo/Shell.cpp
+++ b/Engine/Source/Runtime/Datasmith/CADKernel/Private/CADKernel/Topo/Shell.cpp
@@ -49,13 +49,18 @@ FShell::FShell(const TArray<TSharedPtr<FTopologicalFace>>& InTopologicalFaces, c
 	}
 }
 
-void FShell::Empty()
+void FShell::RemoveFaces()
 {
-	for(FOrientedFace& Face : TopologicalFaces)
+	for (FOrientedFace& Face : TopologicalFaces)
 	{
 		Face.Entity->ResetHost();
 	}
 	TopologicalFaces.Empty();
+}
+
+void FShell::Empty()
+{
+	RemoveFaces();
 	FTopologicalShapeEntity::Empty();
 }
 
diff --git a/Engine/Source/Runtime/Datasmith/CADKernel/Private/CADKernel/Topo/TopologicalFace.cpp b/Engine/Source/Runtime/Datasmith/CADKernel/Private/CADKernel/Topo/TopologicalFace.cpp
--- a/Engine/Source/Runtime/Datasmith/CADKernel/Private/CADKernel/Topo/TopologicalFace.cpp
+++ b/Engine/Source/Runtime/Datasmith/CADKernel/Private/CADKernel/Topo/TopologicalFace.cpp
@@ -537,9 +537,9 @@ void FTopologicalFace::ApplyCriteria(const TArray<TSharedPtr<FCriterion>>& Crite
 		return ElemLength;
 	};
 
-	for (int32 IndexV = 0; IndexV < Grid.GetCoordinateCount(EIso::IsoV) - 1; ++IndexV)
+	for (int32 IndexV = 0; IndexV < CrossingCoordinates[EIso::IsoV].Num() - 1; ++IndexV)
 	{
-		for (int32 IndexU = 0; IndexU < Grid.GetCoordinateCount(EIso::IsoU) - 1; ++IndexU)
+		for (int32 IndexU = 0; IndexU < CrossingCoordinates[EIso::IsoU].Num() - 1; ++IndexU)
 		{
 			const FPoint& Point_U0_V0 = Grid.GetPoint(IndexU, IndexV);
 			const FPoint& Point_U1_V1 = Grid.GetPoint(IndexU + 1, IndexV + 1);
@@ -562,8 +562,8 @@ void FTopologicalFace::ApplyCriteria(const TArray<TSharedPtr<FCriterion>>& Crite
 			double& DeltaVMin = DeltaVMinArray[IndexV];
 			double& DeltaVMax = DeltaVMaxArray[IndexV];
 
-			const double UNextMinusU = Grid.GetCoordinate(EIso::IsoU, IndexU + 1) - Grid.GetCoordinate(EIso::IsoU, IndexU);
-			const double VNextMinusV = Grid.GetCoordinate(EIso::IsoV, IndexV + 1) - Grid.GetCoordinate(EIso::IsoV, IndexV);
+			const double UNextMinusU = CrossingCoordinates[EIso::IsoU][IndexU + 1] - CrossingCoordinates[EIso::IsoU][IndexU];
+			const double VNextMinusV = CrossingCoordinates[EIso::IsoV][IndexV + 1] - CrossingCoordinates[EIso::IsoV][IndexV];
 
 			for (const TSharedPtr<FCriterion>& Criterion : Criteria)
 			{
@@ -577,17 +577,17 @@ void FTopologicalFace::ApplyCriteria(const TArray<TSharedPtr<FCriterion>>& Crite
 	}
 
 	// Delta of the extremities are smooth to avoid big disparity 
-	if (DeltaUMaxArray.Num() > 2)
+	TFunction<void(TArray<double>&)> SmoothExtremities = [](TArray<double>& DeltaMaxArray)
 	{
-		DeltaUMaxArray[0] = (DeltaUMaxArray[0] + DeltaUMaxArray[1] * 2) * AThird;
-		DeltaUMaxArray.Last() = (DeltaUMaxArray.Last() + DeltaUMaxArray[DeltaUMaxArray.Num() - 2] * 2) * AThird;
-	}
+		if (DeltaMaxArray.Num() > 2)
+		{
+			DeltaMaxArray[0] = (DeltaMaxArray[0] + DeltaMaxArray[1] * 2) * AThird;
+			DeltaMaxArray.Last() = (DeltaMaxArray.Last() + DeltaMaxArray[DeltaMaxArray.Num() - 2] * 2) * AThird;
+		}
+	};
 
-	if (DeltaVMaxArray.Num() > 2)
-	{
-		DeltaVMaxArray[0] = (DeltaVMaxArray[0] + DeltaVMaxArray[1] * 2) * AThird;
-		DeltaVMaxArray.Last() = (DeltaVMaxArray.Last() + DeltaVMaxArray[DeltaVMaxArray.Num() - 2] * 2) * AThird;
-	}
+	SmoothExtremities(DeltaUMaxArray);
+	SmoothExtremities(DeltaVMaxArray);
 
 	SetEstimatedMinimalElementLength(ElementLengthMin);
 	SetApplyCriteriaMarker();
diff --git a/Engine/Source/Runtime/Datasmith/CADKernel/Private/CADKernel/Topo/Topomaker.cpp b/Engine/Source/Runtime/Datasmith/CADKernel/Private/CADKernel/Topo/Topomaker.cpp
--- a/Engine/Source/Runtime/Datasmith/CADKernel/Private/CADKernel/Topo/Topomaker.cpp
+++ b/Engine/Source/Runtime/Datasmith/CADKernel/Private/CADKernel/Topo/Topomaker.cpp
@@ -217,7 +217,7 @@ void StitchParallelEdges(TArray<TSharedPtr<FTopologicalVertex>>& VerticesToProce
 
 	for (int32 VertexI = 0; VertexI < VerticesToProcess.Num(); ++VertexI)
 	{
-		TSharedPtr<FTopologicalVertex>& Vertex = VerticesToProcess[VertexI];
+		TSharedPtr<FTopologicalVertex> Vertex = VerticesToProcess[VertexI];
 
 		if (!Vertex.IsValid() || Vertex->IsDeleted() || !Vertex->IsBorderVertex())
 		{
@@ -338,9 +338,8 @@ void MergeCoincidentEdges(FTopologicalEdge* Edge, double MinEdgeLength)
 
 	const bool bFirstEdgeBorder = Edge->IsBorder();
 
-	for (int32 Index = 0; Index < ConnectedEdgeCount; ++Index)
+	for (FTopologicalEdge* SecondEdge : ConnectedEdges)
 	{
-		FTopologicalEdge* SecondEdge = ConnectedEdges[Index];
 		if (SecondEdge == Edge || !SecondEdge->IsActiveEntity() || SecondEdge->IsDegenerated())
 		{
 			continue;
@@ -534,7 +533,7 @@ void FTopomaker::EmptyShells()
 {
 	for (FShell* Shell : Shells)
 	{
-		Shell->Empty();
+		Shell->RemoveFaces();
 	}
 }
 
@@ -659,7 +658,6 @@ void FTopomaker::MergeUnconnectedSuccessiveEdges()
 			TArray<FOrientedEdge>& Edges = Loop->GetEdges();
 			int32 EdgeCount = Edges.Num();
 
-
 			// Find the starting edge i.e. the next edge of the first edge that its ending vertex is connecting to 3 or more edges 
 			// The algorithm start to the last edge of the loop, if it verifies the criteria then the first edge is the edges[0]
 			int32 EndIndex = EdgeCount;
@@ -1138,7 +1136,6 @@ void FTopomaker::SplitIntoConnectedShells()
 		if (FaceSubset.MainShell != nullptr)
 		{
 			FShell* Shell = (FShell*)FaceSubset.MainShell;
-			Shell->Empty();
 			Shell->Add(FaceSubset.Faces);
 		}
 		else
diff --git a/Engine/Source/Runtime/Datasmith/CADKernel/Private/CADKernel/Utils/ArrayUtils.h b/Engine/Source/Runtime/Datasmith/CADKernel/Private/CADKernel/Utils/ArrayUtils.h
--- a/Engine/Source/Runtime/Datasmith/CADKernel/Private/CADKernel/Utils/ArrayUtils.h
+++ b/Engine/Source/Runtime/Datasmith/CADKernel/Private/CADKernel/Utils/ArrayUtils.h
@@ -167,25 +167,37 @@ inline void SubArrayWithoutBoundary(const TArray<double>& InitialArray, const FL
 			return;
 		}
 
-		MaxIndex++;
 		if (InitialArray[MaxIndex] < Boundary.GetMin())
 		{
 			return;
 		}
 
-		if (Boundary.GetMin() + 5 * Tolerance > InitialArray[MinIndex])
+		// MinIndex >= 0, if MinIndex == 0, Boundary.GetMin() can be < InitialArray[0]
+		if (InitialArray[MinIndex] < Boundary.GetMin())
 		{
 			MinIndex++;
 		}
 
-		if (Boundary.GetMax() - 5 * Tolerance > InitialArray[MaxIndex])
+		// FDichotomyFinder return the low index, so MaxIndex < InitialArray.Num() - 1
+		if (InitialArray[MaxIndex + 1] < Boundary.GetMax())
 		{
 			MaxIndex++;
 		}
 
-		if (MaxIndex - MinIndex)
+		if (Boundary.GetMin() + 5 * Tolerance > InitialArray[MinIndex])
+		{
+			MinIndex++;
+		}
+
+		if (Boundary.GetMax() - 5 * Tolerance < InitialArray[MaxIndex])
+		{
+			MaxIndex--;
+		}
+
+		const int32 ElementCount = MaxIndex - MinIndex + 1;
+		if (ElementCount > 0)
 		{
-			OutArray.Append(InitialArray.GetData() + MinIndex, MaxIndex - MinIndex);
+			OutArray.Append(InitialArray.GetData() + MinIndex, ElementCount);
 		}
 	}
 	else if (InitialSize == 1)
diff --git a/Engine/Source/Runtime/Datasmith/CADKernel/Public/CADKernel/Topo/Body.h b/Engine/Source/Runtime/Datasmith/CADKernel/Public/CADKernel/Topo/Body.h
--- a/Engine/Source/Runtime/Datasmith/CADKernel/Public/CADKernel/Topo/Body.h
+++ b/Engine/Source/Runtime/Datasmith/CADKernel/Public/CADKernel/Topo/Body.h
@@ -40,6 +40,11 @@ private:
 
 public:
 
+	virtual ~FBody() override
+	{
+		FBody::Empty();
+	}
+
 	virtual void Serialize(FCADKernelArchive& Ar) override
 	{
 		FTopologicalShapeEntity::Serialize(Ar);
@@ -80,6 +85,10 @@ public:
 
 	virtual void Empty() override
 	{
+		for (TSharedPtr<FShell>& Shell : Shells)
+		{
+			Shell->ResetHost();
+		}
 		Shells.Empty();
 	}
 
diff --git a/Engine/Source/Runtime/Datasmith/CADKernel/Public/CADKernel/Topo/Model.h b/Engine/Source/Runtime/Datasmith/CADKernel/Public/CADKernel/Topo/Model.h
--- a/Engine/Source/Runtime/Datasmith/CADKernel/Public/CADKernel/Topo/Model.h
+++ b/Engine/Source/Runtime/Datasmith/CADKernel/Public/CADKernel/Topo/Model.h
@@ -32,6 +32,11 @@ protected:
 	}
 
 public:
+	
+	virtual ~FModel() override
+	{
+		FModel::Empty();
+	}
 
 	virtual void Serialize(FCADKernelArchive& Ar) override
 	{
diff --git a/Engine/Source/Runtime/Datasmith/CADKernel/Public/CADKernel/Topo/Shell.h b/Engine/Source/Runtime/Datasmith/CADKernel/Public/CADKernel/Topo/Shell.h
--- a/Engine/Source/Runtime/Datasmith/CADKernel/Public/CADKernel/Topo/Shell.h
+++ b/Engine/Source/Runtime/Datasmith/CADKernel/Public/CADKernel/Topo/Shell.h
@@ -67,6 +67,11 @@ private:
 
 public:
 
+	virtual ~FShell()
+	{
+		FShell::Empty();
+	}
+
 	virtual void Serialize(FCADKernelArchive& Ar) override
 	{
 		FTopologicalShapeEntity::Serialize(Ar);
@@ -89,6 +94,7 @@ public:
 		ResetMarkersRecursivelyOnEntities((TArray<TOrientedEntity<FEntity>>&) TopologicalFaces);
 	}
 
+	void RemoveFaces();
 	virtual void Empty() override;
 
 	void Add(TSharedRef<FTopologicalFace> InTopologicalFace, EOrientation InOrientation);
diff --git a/Engine/Source/Runtime/Datasmith/CADKernel/Public/CADKernel/Topo/TopologicalLoop.h b/Engine/Source/Runtime/Datasmith/CADKernel/Public/CADKernel/Topo/TopologicalLoop.h
--- a/Engine/Source/Runtime/Datasmith/CADKernel/Public/CADKernel/Topo/TopologicalLoop.h
+++ b/Engine/Source/Runtime/Datasmith/CADKernel/Public/CADKernel/Topo/TopologicalLoop.h
@@ -77,7 +77,10 @@ private:
 
 public:
 
-	~FTopologicalLoop() = default;
+	virtual ~FTopologicalLoop() override
+	{
+		FTopologicalLoop::Empty();
+	}
 
 	virtual void Empty() override
 	{
