ThinFace preparatory work:
- refactoring CriteriaGrid and Grid with a base class
+ small fixes
#swarm https://p4-swarm.epicgames.net/reviews/24834555

diff --git a/Engine/Plugins/Enterprise/DatasmithCADImporter/Source/CADInterfaces/Private/TechSoftFileParser.cpp b/Engine/Plugins/Enterprise/DatasmithCADImporter/Source/CADInterfaces/Private/TechSoftFileParser.cpp
--- a/Engine/Plugins/Enterprise/DatasmithCADImporter/Source/CADInterfaces/Private/TechSoftFileParser.cpp
+++ b/Engine/Plugins/Enterprise/DatasmithCADImporter/Source/CADInterfaces/Private/TechSoftFileParser.cpp
@@ -14,9 +14,9 @@
 #ifndef CADKERNEL_DEV
 #include "Serialization/JsonSerializer.h"
 #include "Serialization/JsonWriter.h"
+#include "Tasks/Task.h"
 #endif
 
-#include "Tasks/Task.h"
 #include "Templates/UnrealTemplate.h"
 
 namespace CADLibrary
@@ -330,11 +330,13 @@ FTechSoftFileParser::FTechSoftFileParser(FCADFileData& InCADData, const FString&
 ECADParsingResult FTechSoftFileParser::Process()
 {
 	bProcessIsRunning = true;
+#ifndef CADKERNEL_DEV
 	TArray<UE::Tasks::FTask> Checkers;
 	if(FImportParameters::bValidationProcess)
 	{
 		Checkers.Emplace(UE::Tasks::Launch(TEXT("MemoryChecker"), [&]() { CheckMemory(); }));
 	}
+#endif
 
 	//FPlatformMemoryStats StartMem = FPlatformMemory::GetStats();
 	uint64 StartTime = FPlatformTime::Cycles64();
@@ -463,7 +465,10 @@ ECADParsingResult FTechSoftFileParser::Process()
 	ProcessReport.LoadProcessTime = FPlatformTime::ToMilliseconds64(FPlatformTime::Cycles64() - StartTime);
 
 	bProcessIsRunning = false;
+
+#ifndef CADKERNEL_DEV
 	UE::Tasks::Wait(Checkers);
+#endif
 
 	return Result;
 }
@@ -2003,17 +2008,17 @@ void FTechSoftFileParser::ReadMaterialsAndColors()
 
 void FTechSoftFileParser::CheckMemory()
 {
-	CADFileData.GetRecord().StartMemoryUsed = FPlatformMemory::GetStats().UsedPhysical;
-	uint64& MaxMemoryUsed = CADFileData.GetRecord().MaxMemoryUsed;
-	while (bProcessIsRunning)
-	{
-		FPlatformProcess::Sleep(0.1);
-		const uint64 MemoryUsed = FPlatformMemory::GetStats().UsedPhysical;
-		if (MaxMemoryUsed < MemoryUsed)
-		{
-			MaxMemoryUsed = MemoryUsed;
-		}
-	}
+	//CADFileData.GetRecord().StartMemoryUsed = FPlatformMemory::GetStats().UsedPhysical;
+	//uint64& MaxMemoryUsed = CADFileData.GetRecord().MaxMemoryUsed;
+	//while (bProcessIsRunning)
+	//{
+	//	FPlatformProcess::Sleep(0.1);
+	//	const uint64 MemoryUsed = FPlatformMemory::GetStats().UsedPhysical;
+	//	if (MaxMemoryUsed < MemoryUsed)
+	//	{
+	//		MaxMemoryUsed = MemoryUsed;
+	//	}
+	//}
 }
 
 
diff --git a/Engine/Plugins/Enterprise/DatasmithCADImporter/Source/CADInterfaces/Private/TechSoftFileParser.h b/Engine/Plugins/Enterprise/DatasmithCADImporter/Source/CADInterfaces/Private/TechSoftFileParser.h
--- a/Engine/Plugins/Enterprise/DatasmithCADImporter/Source/CADInterfaces/Private/TechSoftFileParser.h
+++ b/Engine/Plugins/Enterprise/DatasmithCADImporter/Source/CADInterfaces/Private/TechSoftFileParser.h
@@ -7,6 +7,8 @@
 #include "CADFileParser.h"
 #include "TechSoftInterface.h"
 
+#include <atomic>
+
 typedef void A3DAsmProductOccurrence;
 typedef void A3DRiRepresentationItem;
 
diff --git a/Engine/Plugins/Enterprise/DatasmithCADImporter/Source/CADInterfaces/Private/TechSoftInterface.cpp b/Engine/Plugins/Enterprise/DatasmithCADImporter/Source/CADInterfaces/Private/TechSoftInterface.cpp
--- a/Engine/Plugins/Enterprise/DatasmithCADImporter/Source/CADInterfaces/Private/TechSoftInterface.cpp
+++ b/Engine/Plugins/Enterprise/DatasmithCADImporter/Source/CADInterfaces/Private/TechSoftInterface.cpp
@@ -106,7 +106,9 @@ FUniqueTechSoftModelFile LoadModelFileFromFile(const A3DImport& Import, A3DStatu
 #if !PLATFORM_EXCEPTIONS_DISABLED
 	catch (...)
 	{
+#ifndef CADKERNEL_DEV
 		UE_LOG(LogCADInterfaces, Warning, TEXT("Failed to load %s. An exception is thrown."), *Import.GetFilePath());
+#endif
 		return FUniqueTechSoftModelFile();
 	}
 #endif
diff --git a/Engine/Source/Runtime/Datasmith/CADKernel/Private/CADKernel/Math/SlopeUtils.h b/Engine/Source/Runtime/Datasmith/CADKernel/Private/CADKernel/Math/SlopeUtils.h
deleted file mode 100644
diff --git a/Engine/Source/Runtime/Datasmith/CADKernel/Private/CADKernel/Mesh/Criteria/CriteriaGrid.cpp b/Engine/Source/Runtime/Datasmith/CADKernel/Private/CADKernel/Mesh/Criteria/CriteriaGrid.cpp
--- a/Engine/Source/Runtime/Datasmith/CADKernel/Private/CADKernel/Mesh/Criteria/CriteriaGrid.cpp
+++ b/Engine/Source/Runtime/Datasmith/CADKernel/Private/CADKernel/Mesh/Criteria/CriteriaGrid.cpp
@@ -1,27 +1,21 @@
 // Copyright Epic Games, Inc. All Rights Reserved.
 #include "CADKernel/Mesh/Criteria/CriteriaGrid.h"
 
-#include "CADKernel/Mesh/Criteria/Criterion.h"
-
 #include "CADKernel/Topo/TopologicalFace.h"
-#include "CADKernel/Geo/Surfaces/Surface.h"
-#include "CADKernel/UI/Display.h"
-
-#ifdef CADKERNEL_DEV
-#include "CADKernel/Mesh/Meshers/IsoTriangulator/DefineForDebug.h"
-#endif
 
 namespace UE::CADKernel
 {
 
-void FCriteriaGrid::Init(const FTopologicalFace& Face)
+void FCriteriaGrid::Init()
 {
-	TFunction<void(const TArray<double>&, TArray<double>&)> ComputeMiddlePointsCoordinates = [](const TArray<double>& Tab, TArray<double>& Tab2)
+	TFunction<void(EIso)> ComputeMiddlePointsCoordinates = [&](EIso Iso)
 	{
-		const int32 TabSize = Tab.Num();
-		ensure(TabSize);
+		const TArray<double>& Tab = Face.GetCrossingPointCoordinates(Iso);
+		TArray<double>& Tab2 = CoordinateGrid[Iso];
 
-		Tab2.SetNum(TabSize * 2 - 1);
+		CuttingCount[Iso] = Tab.Num() * 2 - 1;
+		ensure(CuttingCount[Iso]);
+		Tab2.SetNum(CuttingCount[Iso]);
 
 		Tab2[0] = Tab[0];
 		for (int32 Index = 1; Index < Tab.Num(); Index++)
@@ -31,111 +25,22 @@ void FCriteriaGrid::Init(const FTopologicalFace& Face)
 		}
 	};
 
-	FCoordinateGrid CoordinateGrid2;
-	ComputeMiddlePointsCoordinates(Face.GetCrossingPointCoordinates(EIso::IsoU), CoordinateGrid2[EIso::IsoU]);
-	ComputeMiddlePointsCoordinates(Face.GetCrossingPointCoordinates(EIso::IsoV), CoordinateGrid2[EIso::IsoV]);
+	ComputeMiddlePointsCoordinates(EIso::IsoU);
+	ComputeMiddlePointsCoordinates(EIso::IsoV);
 
-	Face.GetCarrierSurface()->EvaluatePointGrid(CoordinateGrid2, Grid, true);
-	TrueUcoorindateCount = CoordinateGrid2[EIso::IsoU].Num();
+	const bool bWithNormals = false;
+	EvaluatePointGrid(CoordinateGrid, bWithNormals);
 }
 
-FCriteriaGrid::FCriteriaGrid(const FTopologicalFace& InFace)
-	: CoordinateGrid(InFace.GetCrossingPointCoordinates())
+FCriteriaGrid::FCriteriaGrid(FTopologicalFace& InFace)
+	: FGridBase(InFace)
+	, CoordinateGrid(InFace.GetCrossingPointCoordinates())
 {
-	Init(InFace);
+	Init();
 #ifdef DISPLAY_CRITERIA_GRID
 	Display();
 #endif
 }
 
-#ifdef DISPLAY_CRITERIA_GRID
-void FCriteriaGrid::Display() const
-{
-	F3DDebugSession _(TEXT("Grid"));
-
-	{
-		F3DDebugSession A(TEXT("CriteriaGrid Point 3d"));
-		for (int32 VIndex = 0; VIndex < GetCoordinateCount(EIso::IsoV); ++VIndex)
-		{
-			for (int32 UIndex = 0; UIndex < GetCoordinateCount(EIso::IsoU); ++UIndex)
-			{
-				UE::CADKernel::DisplayPoint(GetPoint(UIndex, VIndex), UIndex);
-			}
-		}
-	}
-
-	{
-		F3DDebugSession A(TEXT("CriteriaGrid IntermediateU"));
-		for (int32 VIndex = 0; VIndex < GetCoordinateCount(EIso::IsoV); ++VIndex)
-		{
-			for (int32 UIndex = 0; UIndex < GetCoordinateCount(EIso::IsoU) - 1; ++UIndex)
-			{
-				UE::CADKernel::DisplayPoint(GetIntermediateU(UIndex, VIndex), EVisuProperty::PurplePoint);
-			}
-		}
-	}
-
-	{
-		F3DDebugSession A(TEXT("CriteriaGrid IntermediateV"));
-		for (int32 VIndex = 0; VIndex < GetCoordinateCount(EIso::IsoV) - 1; ++VIndex)
-		{
-			for (int32 UIndex = 0; UIndex < GetCoordinateCount(EIso::IsoU); ++UIndex)
-			{
-				UE::CADKernel::DisplayPoint(GetIntermediateV(UIndex, VIndex), EVisuProperty::PurplePoint);
-			}
-		}
-	}
-
-	{
-		F3DDebugSession A(TEXT("CriteriaGrid IntermediateUV"));
-		for (int32 VIndex = 0; VIndex < GetCoordinateCount(EIso::IsoV) - 1; ++VIndex)
-		{
-			for (int32 UIndex = 0; UIndex < GetCoordinateCount(EIso::IsoU) - 1; ++UIndex)
-			{
-				UE::CADKernel::DisplayPoint(GetIntermediateUV(UIndex, VIndex), EVisuProperty::PurplePoint);
-			}
-		}
-	}
-
-	{
-		F3DDebugSession A(TEXT("CriteriaGrid Point 2D"));
-		for (int32 VIndex = 0; VIndex < GetCoordinateCount(EIso::IsoV); ++VIndex)
-		{
-			for (int32 UIndex = 0; UIndex < GetCoordinateCount(EIso::IsoU); ++UIndex)
-			{
-				UE::CADKernel::DisplayPoint(FPoint2D(GetCoordinate(EIso::IsoU, UIndex), GetCoordinate(EIso::IsoV, VIndex)));
-			}
-		}
-	}
-
-	{
-		F3DDebugSession A(TEXT("CriteriaGrid Point 2D Intermediate"));
-		for (int32 VIndex = 0; VIndex < GetCoordinateCount(EIso::IsoV); ++VIndex)
-		{
-			for (int32 UIndex = 1; UIndex < GetCoordinateCount(EIso::IsoU); ++UIndex)
-			{
-				UE::CADKernel::DisplayPoint(FPoint2D((GetCoordinate(EIso::IsoU, UIndex) + GetCoordinate(EIso::IsoU, UIndex - 1)) * 0.5, GetCoordinate(EIso::IsoV, VIndex)), EVisuProperty::PurplePoint);
-			}
-		}
-		for (int32 VIndex = 1; VIndex < GetCoordinateCount(EIso::IsoV); ++VIndex)
-		{
-			for (int32 UIndex = 0; UIndex < GetCoordinateCount(EIso::IsoU); ++UIndex)
-			{
-				UE::CADKernel::DisplayPoint(FPoint2D(GetCoordinate(EIso::IsoU, UIndex), (GetCoordinate(EIso::IsoV, VIndex) + GetCoordinate(EIso::IsoV, VIndex - 1)) * 0.5), EVisuProperty::PurplePoint);
-			}
-		}
-		for (int32 VIndex = 1; VIndex < GetCoordinateCount(EIso::IsoV); ++VIndex)
-		{
-			for (int32 UIndex = 1; UIndex < GetCoordinateCount(EIso::IsoU); ++UIndex)
-			{
-				UE::CADKernel::DisplayPoint(FPoint2D((GetCoordinate(EIso::IsoU, UIndex) + GetCoordinate(EIso::IsoU, UIndex - 1)) * 0.5, (GetCoordinate(EIso::IsoV, VIndex) + GetCoordinate(EIso::IsoV, VIndex - 1)) * 0.5), EVisuProperty::PurplePoint);
-			}
-		}
-	}
-
-	Wait();
-}
-#endif
-
 } // namespace UE::CADKernel
 
diff --git a/Engine/Source/Runtime/Datasmith/CADKernel/Private/CADKernel/Mesh/Criteria/CriteriaGrid.h b/Engine/Source/Runtime/Datasmith/CADKernel/Private/CADKernel/Mesh/Criteria/CriteriaGrid.h
--- a/Engine/Source/Runtime/Datasmith/CADKernel/Private/CADKernel/Mesh/Criteria/CriteriaGrid.h
+++ b/Engine/Source/Runtime/Datasmith/CADKernel/Private/CADKernel/Mesh/Criteria/CriteriaGrid.h
@@ -2,20 +2,30 @@
 #pragma once
 
 #include "CADKernel/Core/Types.h"
-#include "CADKernel/Geo/GeoPoint.h"
+#include "CADKernel/Mesh/Structure/GridBase.h"
+
 #include "CADKernel/Geo/Sampling/SurfacicSampling.h"
 
 namespace UE::CADKernel
 {
 class FTopologicalFace;
 
-class FCriteriaGrid
+class FCriteriaGrid : public FGridBase
 {
+
+private:
+
+	/*
+	 * Cutting coordinates of the face respecting the meshing criteria
+	 */
+	FCoordinateGrid CoordinateGrid;
+
 protected:
 
-	const FCoordinateGrid& CoordinateGrid;
-	int32 TrueUcoorindateCount;
-	FSurfacicSampling Grid;
+	virtual const FCoordinateGrid& GetCoordinateGrid() const override
+	{
+		return CoordinateGrid;
+	}
 
 	int32 GetIndex(int32 UIndex, int32 VIndex, bool bIsInternalU, bool bIsInternalV) const
 	{
@@ -29,76 +39,51 @@ protected:
 		{
 			++TrueVIndex;
 		}
-		return TrueVIndex * TrueUcoorindateCount + TrueUIndex;
+		return TrueVIndex * CuttingCount[EIso::IsoU] + TrueUIndex;
 	}
 
-	const FPoint& GetPoint(int32 UIndex, int32 VIndex, bool bIsInternalU, bool bIsInternalV, FVector3f* OutNormal = nullptr) const
+	const FPoint& GetPoint(int32 UIndex, int32 VIndex, bool bIsInternalU, bool bIsInternalV) const
 	{
 		int32 Index = GetIndex(UIndex, VIndex, bIsInternalU, bIsInternalV);
-		ensureCADKernel(Grid.Points3D.IsValidIndex(Index));
-
-		if (OutNormal)
-		{
-			*OutNormal = Grid.Normals[Index];
-		}
-		return Grid.Points3D[Index];
+		ensureCADKernel(Points3D.IsValidIndex(Index));
+		return Points3D[Index];
 	}
 
-	void Init(const FTopologicalFace& Face);
-
-public:
-
-	FCriteriaGrid(const FTopologicalFace& InFace);
+	void Init();
 
 	double GetCoordinate(EIso Iso, int32 ind) const
 	{
 		return CoordinateGrid[Iso][ind];
 	}
 
-	constexpr const TArray<double>& GetCoordinates(EIso Iso) const
-	{
-		return CoordinateGrid[Iso];
-	}
+public:
 
-	int32 GetCoordinateCount(EIso Iso) const
-	{
-		return CoordinateGrid.IsoCount(Iso);
-	}
+	FCriteriaGrid(FTopologicalFace& InFace);
 
-	const TArray<FPoint>& GetGridPoints() const
-	{
-		return Grid.Points3D;
-	}
 
-	const TArray<FVector3f>& GetGridNormals() const
+	const FPoint& GetPoint(int32 iU, int32 iV) const
 	{
-		return Grid.Normals;
+		return GetPoint(iU, iV, false, false);
 	}
 
-	const FPoint& GetPoint(int32 iU, int32 iV, FVector3f* normal = nullptr) const
+	const FPoint& GetIntermediateU(int32 iU, int32 iV) const
 	{
-		return GetPoint(iU, iV, false, false, normal);
+		return GetPoint(iU, iV, true, false);
 	}
 
-	const FPoint& GetIntermediateU(int32 iU, int32 iV, FVector3f* normal = nullptr) const
+	const FPoint& GetIntermediateV(int32 iU, int32 iV) const
 	{
-		return GetPoint(iU, iV, true, false, normal);
+		return GetPoint(iU, iV, false, true);
 	}
 
-	const FPoint& GetIntermediateV(int32 iU, int32 iV, FVector3f* normal = nullptr) const
+	const FPoint& GetIntermediateUV(int32 iU, int32 iV) const
 	{
-		return GetPoint(iU, iV, false, true, normal);
+		return GetPoint(iU, iV, true, true);
 	}
 
-	const FPoint& GetIntermediateUV(int32 iU, int32 iV, FVector3f* normal = nullptr) const
-	{
-		return GetPoint(iU, iV, true, true, normal);
-	}
-
-#ifdef DISPLAY_CRITERIA_GRID
+#ifdef CADKERNEL_DEV
 	void Display() const;
 #endif
-
 };
 }
 
diff --git a/Engine/Source/Runtime/Datasmith/CADKernel/Private/CADKernel/Mesh/MeshEnum.h b/Engine/Source/Runtime/Datasmith/CADKernel/Private/CADKernel/Mesh/MeshEnum.h
deleted file mode 100644
diff --git a/Engine/Source/Runtime/Datasmith/CADKernel/Private/CADKernel/Mesh/Meshers/ParametricMesher.cpp b/Engine/Source/Runtime/Datasmith/CADKernel/Private/CADKernel/Mesh/Meshers/ParametricMesher.cpp
--- a/Engine/Source/Runtime/Datasmith/CADKernel/Private/CADKernel/Mesh/Meshers/ParametricMesher.cpp
+++ b/Engine/Source/Runtime/Datasmith/CADKernel/Private/CADKernel/Mesh/Meshers/ParametricMesher.cpp
@@ -101,17 +101,25 @@ void FParametricMesher::PreMeshingTasks()
 
 	const TArray<TSharedPtr<FCriterion>>& Criteria = GetMeshModel().GetCriteria();
 
+#ifdef CADKERNEL_DEV
+	for(FTopologicalFace* Face : Faces)
+	{
+#else
 	TArray<FTopologicalFace*>& LocalFaces = Faces;
 	ParallelFor(Faces.Num(), [&LocalFaces, &Criteria](int32 Index)
 		{
 			FTopologicalFace* Face = LocalFaces[Index];
+#endif
 			if (Face == nullptr || Face->IsDeleted() || Face->IsMeshed())
 			{
 				return;
 			}
 			ApplyFaceCriteria(*Face, Criteria);
 			Face->ComputeSurfaceSideProperties();
-		});
+		}
+#ifndef CADKERNEL_DEV
+	);
+#endif
 
 #ifdef CADKERNEL_DEV
 	MesherReport.Chronos.ApplyCriteriaDuration = FChrono::Elapse(ApplyCriteriaStartTime);
@@ -301,7 +309,9 @@ void FParametricMesher::GenerateCloud(FGrid& Grid)
 	if (bFindThinZone)
 	{
 		FTimePoint StartTime = FChrono::Now();
+#ifdef THIN_ZONES
 		Grid.SearchThinZones();
+#endif
 
 		if (Grid.GetFace().HasThinZone())
 		{
@@ -323,8 +333,11 @@ void FParametricMesher::GenerateCloud(FGrid& Grid)
 	}
 
 #ifdef DEBUG_THIN_ZONES
-	Grid.DisplayInnerDomainPoints(TEXT("FGrid::PointCloud 2D"), Grid.GetInner2DPoints(EGridSpace::Default2D));
-	//Wait(Grid.bDisplay);
+	if(Grid.bDisplay)
+	{
+		Grid.DisplayInnerDomainPoints(TEXT("FGrid::PointCloud 2D"), Grid.GetInner2DPoints(EGridSpace::Default2D));
+		Wait();
+	}
 #endif
 
 	FTimePoint StartTime = FChrono::Now();
@@ -1280,7 +1293,7 @@ void FParametricMesher::MeshSurfaceByFront(TArray<FCostToFace>& QuadTrimmedSurfa
 // =========================================================================================================================================================================================================
 //
 //
-//                                                                            NOT YET REVIEWED
+//                                                                            NOT YET IMPLEMENTED
 //
 //
 // =========================================================================================================================================================================================================
@@ -1289,6 +1302,7 @@ void FParametricMesher::MeshSurfaceByFront(TArray<FCostToFace>& QuadTrimmedSurfa
 
 void FParametricMesher::MeshThinZoneEdges(FGrid& Grid)
 {
+#ifdef UNUSED
 #ifdef DEBUG_MESHTHINSURF
 	Open3DDebugSession(FString::Printf(TEXT("thin Surfaces cutting on surf %d"), Grid.GetFace()->GetId()));
 #endif
@@ -1400,7 +1414,7 @@ void FParametricMesher::MeshThinZoneEdges(FGrid& Grid)
 #ifdef CADKERNEL_DEV
 	MesherReport.Chronos.GlobalMeshThinZones += FChrono::Elapse(MeshStartTime);
 #endif
-
+#endif
 }
 
 static void AddActiveEdgeThinZone(TSharedPtr<FTopologicalEdge> Edge, TSharedPtr<FTopologicalEdge> ActiveEdge, FLinearBoundary& SideEdgeCoordinate)
@@ -1418,6 +1432,7 @@ static void AddActiveEdgeThinZone(TSharedPtr<FTopologicalEdge> Edge, TSharedPtr<
 
 void FParametricMesher::GetThinZoneBoundary(const FThinZoneSide& Side)
 {
+#ifdef UNUSED
 	TSharedPtr<FTopologicalEdge> Edge = nullptr;
 	TSharedPtr<FTopologicalEdge> ActiveEdge = nullptr;
 	FLinearBoundary SideEdgeCoordinate;
@@ -1446,10 +1461,12 @@ void FParametricMesher::GetThinZoneBoundary(const FThinZoneSide& Side)
 		}
 	};
 	AddActiveEdgeThinZone(Edge, ActiveEdge, SideEdgeCoordinate);
+#endif
 }
 
 void FParametricMesher::MeshThinZoneSide(const FThinZoneSide& Side)
 {
+#ifdef UNUSED
 	typedef TFunction<bool(double, double)> CompareMethode;
 
 	TSharedPtr<FTopologicalEdge> Edge;
@@ -1642,7 +1659,7 @@ void FParametricMesher::MeshThinZoneSide(const FThinZoneSide& Side)
 	Close3DDebugSession();
 	Close3DDebugSession();
 #endif
-
+#endif
 }
 
 #ifdef DEBUG_INTERSECTEDGEISOS
diff --git a/Engine/Source/Runtime/Datasmith/CADKernel/Private/CADKernel/Mesh/Structure/EdgeSegment.h b/Engine/Source/Runtime/Datasmith/CADKernel/Private/CADKernel/Mesh/Structure/EdgeSegment.h
deleted file mode 100644
diff --git a/Engine/Source/Runtime/Datasmith/CADKernel/Private/CADKernel/Mesh/Structure/Grid.cpp b/Engine/Source/Runtime/Datasmith/CADKernel/Private/CADKernel/Mesh/Structure/Grid.cpp
--- a/Engine/Source/Runtime/Datasmith/CADKernel/Private/CADKernel/Mesh/Structure/Grid.cpp
+++ b/Engine/Source/Runtime/Datasmith/CADKernel/Private/CADKernel/Mesh/Structure/Grid.cpp
@@ -11,31 +11,23 @@
 #include "CADKernel/Utils/Util.h"
 #include "CADKernel/Utils/ArrayUtils.h"
 
-//#define DEBUG_GRID
 //#define DEBUG_GETPREFERREDUVCOORDINATESFROMNEIGHBOURS
 namespace UE::CADKernel
 {
 
 FGrid::FGrid(FTopologicalFace& InFace, FModelMesh& InMeshModel)
-	: Face(InFace)
+	: FGridBase(InFace)
+	, CoordinateGrid(InFace.GetCuttingPointCoordinates())
 	, FaceTolerance(InFace.GetIsoTolerances())
-	, Tolerance3D(InFace.GetCarrierSurface()->Get3DTolerance())
 	, MinimumElementSize(Tolerance3D * 2.)
 	, MeshModel(InMeshModel)
-	, ThinZoneFinder(*this)
-	, CuttingCoordinates(InFace.GetCuttingPointCoordinates())
 {
 #ifdef DEBUG_ONLY_SURFACE_TO_DEBUG
-	bDisplay = (Face.GetId() == FaceToDebug);
-	Open3DDebugSession(bDisplay, FString::Printf(TEXT("Grid %d"), Face.GetId()));
+	bDisplay = (InFace.GetId() == FaceToDebug);
+	Open3DDebugSession(bDisplay, FString::Printf(TEXT("Grid %d"), InFace.GetId()));
 #endif
 }
 
-void FGrid::PrintTimeElapse() const
-{
-	Chronos.PrintTimeElapse();
-}
-
 void FGrid::ProcessPointCloud()
 {
 	EGridSpace DisplaySpace = EGridSpace::UniformScaled;
@@ -105,13 +97,13 @@ void FGrid::DefineCuttingParameters()
 {
 	FTimePoint StartTime = FChrono::Now();
 
-	FCuttingGrid PreferredCuttingParametersFormLoops;
-	GetPreferredUVCuttingParametersFromLoops(PreferredCuttingParametersFormLoops);
+	FCuttingGrid PreferredCuttingParametersFromLoops;
+	GetPreferredUVCuttingParametersFromLoops(PreferredCuttingParametersFromLoops);
 
-	DefineCuttingParameters(EIso::IsoU, PreferredCuttingParametersFormLoops);
-	DefineCuttingParameters(EIso::IsoV, PreferredCuttingParametersFormLoops);
+	DefineCuttingParameters(EIso::IsoU, PreferredCuttingParametersFromLoops);
+	DefineCuttingParameters(EIso::IsoV, PreferredCuttingParametersFromLoops);
 
-	CuttingSize = CuttingCoordinates.Count();
+	CuttingSize = CoordinateGrid.Count();
 
 	Chronos.DefineCuttingParametersDuration = FChrono::Elapse(StartTime);
 }
@@ -195,7 +187,7 @@ void FGrid::DefineCuttingParameters(EIso Iso, FCuttingGrid& Neighbors)
 		//Wait();
 	}
 #endif
-	CuttingCount[Iso] = CuttingCoordinates.IsoCount(Iso);
+	CuttingCount[Iso] = CoordinateGrid.IsoCount(Iso);
 
 	Chronos.DefineCuttingParametersDuration = FChrono::Elapse(StartTime);
 }
@@ -328,35 +320,16 @@ bool FGrid::GeneratePointCloud()
 {
 	FTimePoint StartTime = FChrono::Now();
 
-	ComputeMaxDeltaUV();
-	if (MaxDeltaUV[EIso::IsoU] < FaceTolerance[EIso::IsoU] || MaxDeltaUV[EIso::IsoV] < FaceTolerance[EIso::IsoV])
+	if (CheckIf2DGridIsDegenerate())
 	{
-		SetAsDegenerated();
 		return false;
 	}
 
 	NodeMarkers.Init(ENodeMarker::None, CuttingSize);
 
-	CountOfInnerNodes = CuttingSize;
-	for (int32 Index = 0; Index < EGridSpace::EndGridSpace; ++Index)
-	{
-		Points2D[Index].SetNum(CuttingSize);
-	}
-	Points3D.SetNum(CuttingSize);
-	Normals.SetNum(CuttingSize);
-
-	int32 Index = 0;
-	for (int32 IPointV = 0; IPointV < CuttingCount[EIso::IsoV]; ++IPointV)
-	{
-		for (int32 IPointU = 0; IPointU < CuttingCount[EIso::IsoU]; ++IPointU, ++Index)
-		{
-			Points2D[EGridSpace::Default2D][Index].Set(Face.GetCuttingCoordinatesAlongIso(EIso::IsoU)[IPointU], Face.GetCuttingCoordinatesAlongIso(EIso::IsoV)[IPointV]);
-		}
-	}
-
-	Face.EvaluateGrid(*this);
+	EvaluatePointGrid(CoordinateGrid, true);
 
-	ComputeMaxElementSize();
+	CountOfInnerNodes = CuttingSize;
 
 	if (!ScaleGrid())
 	{
@@ -367,33 +340,6 @@ bool FGrid::GeneratePointCloud()
 	return true;
 }
 
-void FGrid::ComputeMaxElementSize()
-{
-	MaxElementSize[EIso::IsoV] = 0;
-	for (int32 IndexU = 0; IndexU < CuttingCount[EIso::IsoU]; ++IndexU)
-	{
-		int32 Index = IndexU;
-		for (int32 IndexV = 1; IndexV < CuttingCount[EIso::IsoV]; ++IndexV, Index += CuttingCount[EIso::IsoU])
-		{
-			MaxElementSize[EIso::IsoV] = FMath::Max(Points3D[Index].SquareDistance(Points3D[Index + CuttingCount[EIso::IsoU]]), MaxElementSize[EIso::IsoV]);
-		}
-	}
-
-	MaxElementSize[EIso::IsoU] = 0;
-	for (int32 IndexV = 0, Index = 0; IndexV < CuttingCount[EIso::IsoV]; ++IndexV)
-	{
-		for (int32 IndexU = 1; IndexU < CuttingCount[EIso::IsoU]; ++IndexU, ++Index)
-		{
-			MaxElementSize[EIso::IsoU] = FMath::Max(Points3D[Index].SquareDistance(Points3D[Index + 1]), MaxElementSize[EIso::IsoU]);
-		}
-		++Index;
-	}
-
-	MaxElementSize[EIso::IsoU] = sqrt(MaxElementSize[EIso::IsoU]);
-	MaxElementSize[EIso::IsoV] = sqrt(MaxElementSize[EIso::IsoV]);
-	MinOfMaxElementSize = FMath::Min(MaxElementSize[EIso::IsoU], MaxElementSize[EIso::IsoV]);
-}
-
 void FGrid::FindPointsCloseToLoop()
 {
 	FTimePoint StartTime = FChrono::Now();
@@ -418,10 +364,10 @@ void FGrid::FindPointsCloseToLoop()
 		if (bDisplay)
 		{
 			F3DDebugSession _(*FString::Printf(TEXT("Cell %d"), CellIndex++));
-			DisplayPoint(Points2D[EGridSpace::UniformScaled][GlobalIndex] * DisplayScale);
-			DisplayPoint(Points2D[EGridSpace::UniformScaled][GlobalIndex - 1] * DisplayScale);
-			DisplayPoint(Points2D[EGridSpace::UniformScaled][GlobalIndex - 1 - CuttingCount[EIso::IsoU]] * DisplayScale);
-			DisplayPoint(Points2D[EGridSpace::UniformScaled][GlobalIndex - CuttingCount[EIso::IsoU]] * DisplayScale);
+			DisplayPoint2DWithScale(Points2D[EGridSpace::UniformScaled][GlobalIndex]);
+			DisplayPoint2DWithScale(Points2D[EGridSpace::UniformScaled][GlobalIndex - 1]);
+			DisplayPoint2DWithScale(Points2D[EGridSpace::UniformScaled][GlobalIndex - 1 - CuttingCount[EIso::IsoU]]);
+			DisplayPoint2DWithScale(Points2D[EGridSpace::UniformScaled][GlobalIndex - CuttingCount[EIso::IsoU]]);
 			Wait(bWaitCell);
 		}
 	};
@@ -582,7 +528,7 @@ void FGrid::FindPointsCloseToLoop()
 			if (bDisplay)
 			{
 				F3DDebugSession _(TEXT("SEG"));
-				DisplaySegment(*PointB * DisplayScale, *PointA * DisplayScale);
+				DisplaySegmentWithScale(*PointB, *PointA);
 			}
 #endif
 
@@ -1122,17 +1068,10 @@ void FGrid::GetMeshOfLoop(const FTopologicalLoop& Loop)
 	}
 }
 
-//#define DEBUG_GET_BOUNDARY_MESH
 bool FGrid::GetMeshOfLoops()
 {
-	int32 ThinZoneNum = 0;
-	if (Face.HasThinZone())
-	{
-		ThinZoneNum = ThinZoneFinder.GetThinZones().Num();
-	}
-
 	int32 LoopCount = Face.GetLoops().Num();
-	FaceLoops2D[EGridSpace::Default2D].Reserve(LoopCount + ThinZoneNum);
+	FaceLoops2D[EGridSpace::Default2D].Reserve(LoopCount);
 
 	FaceLoops3D.Reserve(LoopCount);
 	NormalsOfFaceLoops.Reserve(LoopCount);
@@ -1160,45 +1099,11 @@ bool FGrid::GetMeshOfLoops()
 		}
 	}
 
-	if (CheckIfDegenerated())
+	if (CheckIfExternalLoopIsDegenerate())
 	{
 		return false;
 	}
 
-	if (ThinZoneNum)
-	{
-		for (const FThinZone2D& ThinZone : ThinZoneFinder.GetThinZones())
-		{
-			int32 PointNum = ThinZone.GetFirstSide().GetSegments().Num();
-			PointNum += ThinZone.GetSecondSide().GetSegments().Num();
-			TArray<FPoint2D>& LoopPoints = FaceLoops2D[EGridSpace::Default2D].Emplace_GetRef();
-			LoopPoints.Reserve(PointNum + 4);
-
-			// First point Side1
-			{
-				const FEdgeSegment* Segment = ThinZone.GetFirstSide().GetSegments()[0];
-				LoopPoints.Emplace(Segment->GetEdge()->Approximate2DPoint(Segment->GetCoordinate(ELimit::Start)));
-			}
-
-			for (const FEdgeSegment* Segment : ThinZone.GetFirstSide().GetSegments())
-			{
-				LoopPoints.Emplace(Segment->GetEdge()->Approximate2DPoint(Segment->GetCoordinate(ELimit::End)));
-			}
-
-			// First point Side2
-			{
-				const FEdgeSegment* Segment = ThinZone.GetSecondSide().GetSegments().Last();
-				LoopPoints.Emplace(Segment->GetEdge()->Approximate2DPoint(Segment->GetCoordinate(ELimit::Start)));
-			}
-
-			const TArray<FEdgeSegment*>& Segments = ThinZone.GetSecondSide().GetSegments();
-			for (int32 Index = Segments.Num() - 1; Index >= 0; --Index)
-			{
-				LoopPoints.Emplace(Segments[Index]->GetEdge()->Approximate2DPoint(Segments[Index]->GetCoordinate(ELimit::End)));
-			}
-		}
-	}
-
 	// Fit boundaries to Surface bounds.
 	const FSurfacicBoundary& Bounds = Face.GetBoundary();
 	for (TArray<FPoint2D>& Loop : FaceLoops2D[EGridSpace::Default2D])
@@ -1211,218 +1116,6 @@ bool FGrid::GetMeshOfLoops()
 	return true;
 }
 
-bool FGrid::ScaleGrid()
-{
-	FTimePoint StartTime = FChrono::Now();
-
-	TFunction<double(const TArray<double>&)> GetMean = [](const TArray<double>& Lengths)
-	{
-		double MeanLength = 0;
-		for (double Length : Lengths)
-		{
-			MeanLength += Length;
-		}
-		MeanLength /= Lengths.Num();
-		return MeanLength;
-	};
-
-	TFunction<double(const TArray<double>&, const double)> StandardDeviation = [](const TArray<double>& Lengths, const double MeanLength)
-	{
-		double StandardDeviation = 0;
-		for (double Length : Lengths)
-		{
-			StandardDeviation += FMath::Square(Length - MeanLength);
-		}
-		StandardDeviation /= Lengths.Num();
-		StandardDeviation = sqrt(StandardDeviation);
-		return StandardDeviation;
-	};
-
-	TFunction<void(const TArray<double>&, const double, TArray<double>&)> ScaleCoordinates = [](const TArray<double>& InCoordinates, const double ScaleFactor, TArray<double>& OutCoordinatesScaled)
-	{
-		OutCoordinatesScaled.Reserve(InCoordinates.Num());
-
-		for (double Coordinate : InCoordinates)
-		{
-			OutCoordinatesScaled.Add(Coordinate * ScaleFactor);
-		}
-	};
-
-	TFunction<int32(const TArray<double>&, const double)> GetMiddleIndex = [](const TArray<double>& Coordinates, double Middle)
-	{
-		int32 StartIndexUp = 1;
-		for (; StartIndexUp < Coordinates.Num(); ++StartIndexUp)
-		{
-			if (Coordinates[StartIndexUp] > Middle)
-			{
-				break;
-			}
-		}
-		return StartIndexUp;
-	};
-
-	TArray<double> LengthsV;
-	LengthsV.SetNum(CuttingCount[EIso::IsoU]);
-	for (int32 IndexU = 0; IndexU < CuttingCount[EIso::IsoU]; ++IndexU)
-	{
-		int32 Index = IndexU;
-		double Length = 0;
-		for (int32 IndexV = 1; IndexV < CuttingCount[EIso::IsoV]; ++IndexV)
-		{
-			Length += Points3D[Index].Distance(Points3D[Index + CuttingCount[EIso::IsoU]]);
-			Index += CuttingCount[EIso::IsoU];
-		}
-		LengthsV[IndexU] = Length;
-	}
-
-	TArray<double> LengthsU;
-	LengthsU.SetNum(CuttingCount[EIso::IsoV]);
-	for (int32 IndexV = 0, Index = 0; IndexV < CuttingCount[EIso::IsoV]; ++IndexV)
-	{
-		double Length = 0;
-		for (int32 IndexU = 1; IndexU < CuttingCount[EIso::IsoU]; ++IndexU)
-		{
-			Length += Points3D[Index].Distance(Points3D[Index + 1]);
-			Index++;
-		}
-		Index++;
-		LengthsU[IndexV] = Length;
-	}
-
-	double MeanLengthV = GetMean(LengthsV);
-	double MeanLengthU = GetMean(LengthsU);
-	if(MeanLengthV < Tolerance3D || MeanLengthU < Tolerance3D)
-	{
-		SetAsDegenerated();
-		return false;
-	}
-
-	double FactorV = MeanLengthV / (CuttingCoordinates[EIso::IsoV].Last() - CuttingCoordinates[EIso::IsoV][0]);
-	double FactorU = MeanLengthU / (CuttingCoordinates[EIso::IsoU].Last() - CuttingCoordinates[EIso::IsoU][0]);
-
-	//TArray<double> ScaledCuttingCoordinates;
-	ScaleCoordinates(CuttingCoordinates[EIso::IsoU], FactorU, UniformCuttingCoordinates[EIso::IsoU]);
-	ScaleCoordinates(CuttingCoordinates[EIso::IsoV], FactorV, UniformCuttingCoordinates[EIso::IsoV]);
-
-	{
-		int32 NumUV = 0;
-		for (int32 IPointV = 0; IPointV < CuttingCount[EIso::IsoV]; ++IPointV)
-		{
-			for (int32 IPointU = 0; IPointU < CuttingCount[EIso::IsoU]; ++IPointU, ++NumUV)
-			{
-				Points2D[EGridSpace::UniformScaled][NumUV].Set(UniformCuttingCoordinates[EIso::IsoU][IPointU], UniformCuttingCoordinates[EIso::IsoV][IPointV]);
-			}
-		}
-	}
-
-	double StandardDeviationU = StandardDeviation(LengthsU, MeanLengthU);
-	double StandardDeviationV = StandardDeviation(LengthsV, MeanLengthV);
-
-	if (StandardDeviationV > StandardDeviationU)
-	{
-		double MiddleV = (CuttingCoordinates[EIso::IsoV].Last() + CuttingCoordinates[EIso::IsoV][0]) * 0.5;
-
-		FCoordinateGrid Grid;
-		Grid[EIso::IsoU] = CuttingCoordinates[EIso::IsoU];
-		Grid[EIso::IsoV].Add(MiddleV);
-
-		FSurfacicSampling MiddlePoints;
-		Face.EvaluatePointGrid(Grid, MiddlePoints);
-
-		int32 StartIndexUp = GetMiddleIndex(CuttingCoordinates[EIso::IsoV], MiddleV);
-		int32 StartIndexDown = StartIndexUp - 1;
-
-		int32 NumUV = 0;
-		for (int32 IPointU = 0; IPointU < CuttingCount[EIso::IsoU]; ++IPointU)
-		{
-			double Length = 0;
-			FPoint LastPoint = MiddlePoints.Points3D[IPointU];
-			for (int32 IPointV = StartIndexUp; IPointV < CuttingCount[EIso::IsoV]; ++IPointV)
-			{
-				NumUV = IPointV * CuttingCount[EIso::IsoU] + IPointU;
-				Length += LastPoint.Distance(Points3D[NumUV]);
-				Points2D[EGridSpace::Scaled][NumUV].Set(Points2D[EGridSpace::UniformScaled][NumUV].U, Length);
-				LastPoint = Points3D[NumUV];
-			}
-
-			Length = 0;
-			LastPoint = MiddlePoints.Points3D[IPointU];
-			for (int32 IPointV = StartIndexDown; IPointV >= 0; --IPointV)
-			{
-				NumUV = IPointV * CuttingCount[EIso::IsoU] + IPointU;
-				Length -= LastPoint.Distance(Points3D[NumUV]);
-				Points2D[EGridSpace::Scaled][NumUV].Set(Points2D[EGridSpace::UniformScaled][NumUV].U, Length);
-				LastPoint = Points3D[NumUV];
-			}
-		}
-	}
-	else
-	{
-		double MiddleU = (CuttingCoordinates[EIso::IsoU].Last() + CuttingCoordinates[EIso::IsoU][0]) * 0.5;
-
-		FCoordinateGrid Grid;
-		Grid[EIso::IsoU].Add(MiddleU);
-		Grid[EIso::IsoV] = CuttingCoordinates[EIso::IsoV];
-
-		FSurfacicSampling MiddlePoints;
-		Face.EvaluatePointGrid(Grid, MiddlePoints);
-
-		int32 StartIndexUp = GetMiddleIndex(CuttingCoordinates[EIso::IsoU], MiddleU);
-		int32 StartIndexDown = StartIndexUp - 1;
-
-		int32 NumUV = 0;
-		for (int32 IPointV = 0; IPointV < CuttingCount[EIso::IsoV]; ++IPointV)
-		{
-			double Length = 0;
-			FPoint LastPoint = MiddlePoints.Points3D[IPointV];
-			for (int32 IPointU = StartIndexUp; IPointU < CuttingCount[EIso::IsoU]; ++IPointU)
-			{
-				NumUV = IPointV * CuttingCount[EIso::IsoU] + IPointU;
-				Length += LastPoint.Distance(Points3D[NumUV]);
-				Points2D[EGridSpace::Scaled][NumUV].Set(Length, Points2D[EGridSpace::UniformScaled][NumUV].V);
-				LastPoint = Points3D[NumUV];
-			}
-
-			Length = 0;
-			LastPoint = MiddlePoints.Points3D[IPointV];
-			for (int32 IPointU = StartIndexDown; IPointU >= 0; --IPointU)
-			{
-				NumUV = IPointV * CuttingCount[EIso::IsoU] + IPointU;
-				Length -= LastPoint.Distance(Points3D[NumUV]);
-				Points2D[EGridSpace::Scaled][NumUV].Set(Length, Points2D[EGridSpace::UniformScaled][NumUV].V);
-				LastPoint = Points3D[NumUV];
-			}
-		}
-	}
-	Chronos.ScaleGridDuration = FChrono::Elapse(StartTime);
-	return true;
-}
-
-void FGrid::TransformPoints(EGridSpace DestinationSpace, const TArray<FPoint2D>& InPointsToScale, TArray<FPoint2D>& OutTransformedPoints) const
-{
-	OutTransformedPoints.SetNum(InPointsToScale.Num());
-
-	int32 IndexU = 0;
-	int32 IndexV = 0;
-	for (int32 Index = 0; Index < InPointsToScale.Num(); ++Index)
-	{
-		const FPoint2D& Point = InPointsToScale[Index];
-
-		ArrayUtils::FindCoordinateIndex(CuttingCoordinates[EIso::IsoU], Point.U, IndexU);
-		ArrayUtils::FindCoordinateIndex(CuttingCoordinates[EIso::IsoV], Point.V, IndexV);
-
-		ComputeNewCoordinate(Points2D[DestinationSpace], IndexU, IndexV, Point, OutTransformedPoints[Index]);
-	}
-}
-
-void FGrid::SearchThinZones()
-{
-	ensure(false);
-	double Size = GetMinElementSize();
-	ThinZoneFinder.Set(Size / 3.);
-	ThinZoneFinder.SearchThinZones();
-}
-
 void FGrid::ScaleLoops()
 {
 	FaceLoops2D[EGridSpace::Scaled].SetNum(FaceLoops2D[EGridSpace::Default2D].Num());
@@ -1443,8 +1136,8 @@ void FGrid::ScaleLoops()
 		{
 			const FPoint2D& Point = Loop[Index];
 
-			ArrayUtils::FindCoordinateIndex(CuttingCoordinates[EIso::IsoU], Point.U, IndexU);
-			ArrayUtils::FindCoordinateIndex(CuttingCoordinates[EIso::IsoV], Point.V, IndexV);
+			ArrayUtils::FindCoordinateIndex(CoordinateGrid[EIso::IsoU], Point.U, IndexU);
+			ArrayUtils::FindCoordinateIndex(CoordinateGrid[EIso::IsoV], Point.V, IndexV);
 
 			ComputeNewCoordinate(Points2D[EGridSpace::Scaled], IndexU, IndexV, Point, ScaledLoop[Index]);
 			ComputeNewCoordinate(Points2D[EGridSpace::UniformScaled], IndexU, IndexV, Point, UniformScaledLoop[Index]);
@@ -1453,21 +1146,28 @@ void FGrid::ScaleLoops()
 	}
 }
 
-void FGrid::ComputeMaxDeltaUV()
+bool FGrid::CheckIf2DGridIsDegenerate() const
 {
-	MaxDeltaUV[EIso::IsoU] = 0;
-	for (int32 Index = 1; Index < CuttingCoordinates[EIso::IsoU].Num(); ++Index)
+	TFunction<bool(EIso)> IsDegenerate = [&](EIso Iso) -> bool
 	{
-		double Delta = CuttingCoordinates[EIso::IsoU][Index] - CuttingCoordinates[EIso::IsoU][Index - 1];
-		MaxDeltaUV[EIso::IsoU] = FMath::Max(MaxDeltaUV[EIso::IsoU], Delta);
-	}
+		double MaxDelta = 0;
+		for (int32 Index = 1; Index < CoordinateGrid[Iso].Num(); ++Index)
+		{
+			double Delta = CoordinateGrid[Iso][Index] - CoordinateGrid[Iso][Index - 1];
+			if (Delta > MaxDelta)
+			{
+				MaxDelta = Delta;
+			}
+		}
+		return MaxDelta < FaceTolerance[Iso];
+	};
 
-	MaxDeltaUV[EIso::IsoV] = 0;
-	for (int32 Index = 1; Index < CuttingCoordinates[EIso::IsoV].Num(); ++Index)
+	if (IsDegenerate(EIso::IsoU) || IsDegenerate(EIso::IsoV))
 	{
-		double Delta = CuttingCoordinates[EIso::IsoV][Index] - CuttingCoordinates[EIso::IsoV][Index - 1];
-		MaxDeltaUV[EIso::IsoV] = FMath::Max(MaxDeltaUV[EIso::IsoV], Delta);
+		SetAsDegenerated();
+		return true;
 	}
+	return false;
 }
 
 void FGrid::FindInnerFacePoints()
@@ -1480,19 +1180,24 @@ void FGrid::FindInnerFacePoints()
 
 	FTimePoint StartTime = FChrono::Now();
 
-	TArray<char> NbIntersectUForward; // we need to know if intersect is pair 0, 2, 4... intersection of impair 1, 3, 5... false is pair, true is impaire 
-	TArray<char> NbIntersectUBackward;
-	TArray<char> NbIntersectVForward;
-	TArray<char> NbIntersectVBackward;
+	TFunction<void(TArray<char>&, int32)> AddIntersection = [](TArray<char>& Intersect, int32 Index)
+	{
+		Intersect[Index] = Intersect[Index] > 0 ? 0 : 1;
+	};
+
+	TArray<char> VForwardIntersectCount; // we need to know if intersect is pair 0, 2, 4... intersection of impair 1, 3, 5... false is pair, true is impair 
+	TArray<char> VBackwardIntersectCount;
+	TArray<char> UForwardIntersectCount;
+	TArray<char> UBackwardIntersectCount;
 	TArray<char> IntersectLoop;
 
 	IntersectLoop.Init(0, CuttingSize);
 	NodeMarkers.Init(ENodeMarker::IsInside, CuttingSize);
 
-	NbIntersectUForward.Init(0, CuttingSize);
-	NbIntersectUBackward.Init(0, CuttingSize);
-	NbIntersectVForward.Init(0, CuttingSize);
-	NbIntersectVBackward.Init(0, CuttingSize);
+	VForwardIntersectCount.Init(0, CuttingSize);
+	VBackwardIntersectCount.Init(0, CuttingSize);
+	UForwardIntersectCount.Init(0, CuttingSize);
+	UBackwardIntersectCount.Init(0, CuttingSize);
 
 	// Loop node too close to one of CoordinateU or CoordinateV are moved a little to avoid floating error of comparison 
 	// This step is necessary instead of all points could be considered outside...
@@ -1636,22 +1341,22 @@ void FGrid::FindInnerFacePoints()
 
 					if (UniformCuttingCoordinates[EIso::IsoU][IndexU] < UMin)
 					{
-						NbIntersectVForward[Index] = NbIntersectVForward[Index] > 0 ? 0 : 1;
+						AddIntersection(UForwardIntersectCount, Index);
 					}
 					else if (UniformCuttingCoordinates[EIso::IsoU][IndexU] > Umax)
 					{
-						NbIntersectVBackward[Index] = NbIntersectVBackward[Index] > 0 ? 0 : 1;
+						AddIntersection(UBackwardIntersectCount, Index);
 					}
 					else
 					{
 						double APvectAB = UniformCuttingCoordinates[EIso::IsoV][IndexV] * ABu - UniformCuttingCoordinates[EIso::IsoU][IndexU] * ABv + AuABVMinusAvABu;
 						if (APvectAB > DOUBLE_SMALL_NUMBER)
 						{
-							NbIntersectVForward[Index] = NbIntersectVForward[Index] > 0 ? 0 : 1;
+							AddIntersection(UForwardIntersectCount, Index);
 						}
 						else if (APvectAB < DOUBLE_SMALL_NUMBER)
 						{
-							NbIntersectVBackward[Index] = NbIntersectVBackward[Index] > 0 ? 0 : 1;
+							AddIntersection(UBackwardIntersectCount, Index);
 						}
 						else
 						{
@@ -1721,22 +1426,22 @@ void FGrid::FindInnerFacePoints()
 
 					if (UniformCuttingCoordinates[EIso::IsoV][IndexV] < VMin)
 					{
-						NbIntersectUForward[Index] = NbIntersectUForward[Index] > 0 ? 0 : 1;
+						AddIntersection(VForwardIntersectCount, Index);
 					}
 					else if (UniformCuttingCoordinates[EIso::IsoV][IndexV] > Vmax)
 					{
-						NbIntersectUBackward[Index] = NbIntersectUBackward[Index] > 0 ? 0 : 1;
+						AddIntersection(VBackwardIntersectCount, Index);
 					}
 					else
 					{
 						double APvectAB = UniformCuttingCoordinates[EIso::IsoV][IndexV] * ABu - UniformCuttingCoordinates[EIso::IsoU][IndexU] * ABv + AuABVMinusAvABu;
 						if (APvectAB > DOUBLE_SMALL_NUMBER)
 						{
-							NbIntersectUBackward[Index] = NbIntersectUBackward[Index] > 0 ? 0 : 1;
+							AddIntersection(VBackwardIntersectCount, Index);
 						}
 						else if (APvectAB < DOUBLE_SMALL_NUMBER)
 						{
-							NbIntersectUForward[Index] = NbIntersectUForward[Index] > 0 ? 0 : 1;
+							AddIntersection(VForwardIntersectCount, Index);
 						}
 						else
 						{
@@ -1759,19 +1464,19 @@ void FGrid::FindInnerFacePoints()
 		}
 
 		int32 IsInside = 0;
-		if (NbIntersectVForward[Index] > 0)
+		if (UForwardIntersectCount[Index] > 0)
 		{
 			IsInside++;
 		}
-		if (NbIntersectVBackward[Index] > 0)
+		if (UBackwardIntersectCount[Index] > 0)
 		{
 			IsInside++;
 		}
-		if (NbIntersectUForward[Index] > 0)
+		if (VForwardIntersectCount[Index] > 0)
 		{
 			IsInside++;
 		}
-		if (NbIntersectUBackward[Index] > 0)
+		if (VBackwardIntersectCount[Index] > 0)
 		{
 			IsInside++;
 		}
@@ -1785,7 +1490,7 @@ void FGrid::FindInnerFacePoints()
 	Chronos.FindInnerDomainPointsDuration += FChrono::Elapse(StartTime);
 }
 
-bool FGrid::CheckIfDegenerated()
+bool FGrid::CheckIfExternalLoopIsDegenerate() const
 {
 	if (FaceLoops2D[EGridSpace::Default2D].Num() == 0)
 	{
@@ -1804,6 +1509,6 @@ bool FGrid::CheckIfDegenerated()
 	return false;
 }
 
-}
+} // NS CADKernel
 
 
diff --git a/Engine/Source/Runtime/Datasmith/CADKernel/Private/CADKernel/Mesh/Structure/Grid.h b/Engine/Source/Runtime/Datasmith/CADKernel/Private/CADKernel/Mesh/Structure/Grid.h
--- a/Engine/Source/Runtime/Datasmith/CADKernel/Private/CADKernel/Mesh/Structure/Grid.h
+++ b/Engine/Source/Runtime/Datasmith/CADKernel/Private/CADKernel/Mesh/Structure/Grid.h
@@ -6,7 +6,7 @@
 #include "CADKernel/Geo/GeoPoint.h"
 #include "CADKernel/Math/Point.h"
 #include "CADKernel/Mesh/MeshEnum.h"
-#include "CADKernel/Mesh/Structure/ThinZone2D.h"
+#include "CADKernel/Mesh/Structure/GridBase.h"
 #include "CADKernel/UI/Display.h"
 
 #ifdef CADKERNEL_DEV
@@ -29,47 +29,6 @@ class FPoint2D;
 class FThinZone2DFinder;
 class FTopologicalFace;
 
-struct FGridChronos
-{
-	FDuration DefineCuttingParametersDuration;
-	FDuration GeneratePointCloudDuration;
-	FDuration ProcessPointCloudDuration;
-	FDuration FindInnerDomainPointsDuration;
-	FDuration Build2DLoopDuration;
-	FDuration RemovePointsClosedToLoopDuration;
-	FDuration FindPointsCloseToLoopDuration;
-	FDuration ScaleGridDuration;
-
-	FGridChronos()
-		: DefineCuttingParametersDuration(FChrono::Init())
-		, GeneratePointCloudDuration(FChrono::Init())
-		, ProcessPointCloudDuration(FChrono::Init())
-		, FindInnerDomainPointsDuration(FChrono::Init())
-		, Build2DLoopDuration(FChrono::Init())
-		, RemovePointsClosedToLoopDuration(FChrono::Init())
-		, FindPointsCloseToLoopDuration(FChrono::Init())
-		, ScaleGridDuration(FChrono::Init())
-	{}
-
-	void PrintTimeElapse() const
-	{
-		FDuration GridDuration = FChrono::Elapse(FChrono::Now());
-		GridDuration += DefineCuttingParametersDuration;
-		GridDuration += GeneratePointCloudDuration;
-		GridDuration += FindInnerDomainPointsDuration;
-		GridDuration += Build2DLoopDuration;
-		GridDuration += RemovePointsClosedToLoopDuration;
-		GridDuration += FindPointsCloseToLoopDuration;
-		FChrono::PrintClockElapse(Log, TEXT(""), TEXT("Grid"), GridDuration);
-		FChrono::PrintClockElapse(Log, TEXT("  "), TEXT("DefineCuttingParameters"), DefineCuttingParametersDuration);
-		FChrono::PrintClockElapse(Log, TEXT("  "), TEXT("GeneratePointCloud"), GeneratePointCloudDuration);
-		FChrono::PrintClockElapse(Log, TEXT("  "), TEXT("GenerateDomainPoints"), ProcessPointCloudDuration);
-		FChrono::PrintClockElapse(Log, TEXT("    "), TEXT("FindInnerDomainPointsDuration"), FindInnerDomainPointsDuration);
-		FChrono::PrintClockElapse(Log, TEXT("    "), TEXT("FindPointsCloseToLoop"), FindPointsCloseToLoopDuration);
-		FChrono::PrintClockElapse(Log, TEXT("    "), TEXT("RemovePointsClosedToLoop"), RemovePointsClosedToLoopDuration);
-		FChrono::PrintClockElapse(Log, TEXT("    "), TEXT("Build2DLoopDuration"), Build2DLoopDuration);
-	}
-};
 
 enum class ENodeMarker : uint8
 {
@@ -84,18 +43,26 @@ enum class ENodeMarker : uint8
 ENUM_CLASS_FLAGS(ENodeMarker);
 
 
-class FGrid : public FHaveStates
+class FGrid : public FGridBase
 {
 protected:
 
-	FTopologicalFace& Face;
+	/*
+	 * Cutting coordinates of the face respecting the meshing criteria
+	 */
+	const FCoordinateGrid& CoordinateGrid;
+
+	virtual const FCoordinateGrid& GetCoordinateGrid() const override
+	{
+		return CoordinateGrid;
+	}
+
 	const FSurfacicTolerance FaceTolerance;
-	const double Tolerance3D;
 	const double MinimumElementSize;
 
+
 	FModelMesh& MeshModel;
 
-	FThinZone2DFinder ThinZoneFinder;
 
 	/**
 	 * 2D Coordinate of Loop's nodes in each space
@@ -114,33 +81,11 @@ protected:
 
 	TArray<TArray<int32>> NodeIdsOfFaceLoops;
 
-	/**
-	 * grid point cloud size
-	 */
-	int32 CuttingCount[2] = { 0 , 0 };
-	int32 CuttingSize = 0;
-
 	/**
 	 * count of node inside the face i.e. node inside the external loop and outer the inner loops
 	 */
 	int32 CountOfInnerNodes = 0;
 
-	/*
-	 * Cutting coordinates of the face respecting the meshing criteria
-	 */
-	const FCoordinateGrid& CuttingCoordinates;
-	FCoordinateGrid UniformCuttingCoordinates;
-
-	/*
-	 * Maximum difference of coordinate along the specified axis of two successive cutting points
-	 */
-	FPoint2D MaxDeltaUV = { 0., 0. };
-
-	/**
-	 * Maximum 3d distance along each axis of two successive cutting points
-	 */
-	FPoint2D MaxElementSize = { 0., 0. };
-
 	double MinOfMaxElementSize = 0;
 
 	/**
@@ -149,21 +94,6 @@ protected:
 	 */
 	TArray<ENodeMarker> NodeMarkers;
 
-	/**
-	 * 2D Coordinate of grid nodes in each space
-	 */
-	TArray<FPoint2D> Points2D[EGridSpace::EndGridSpace];
-
-	/**
-	 * 3D Coordinate of inner nodes
-	 */
-	TArray<FPoint> Points3D;
-
-	/**
-	 * Surface Normal at each inner nodes
-	 */
-	TArray<FVector3f> Normals;
-
 public:
 	FGridChronos Chronos;
 
@@ -183,11 +113,6 @@ public:
 	// Meshing tools =========================================================================================================================================================================================================
 	// ======================================================================================================================================================================================================================
 
-	/**
-	 * Return true if the grid is not consistent to build a mesh as composed of only two border nodes.
-	 */
-	bool CheckIfDegenerated();
-
 	/**
 	 * Defines the cutting coordinate of the grid according to mesh criteria and the existing mesh of bordering edges (loop's edges)
 	 * @see GetPreferredUVCoordinatesFromNeighbours
@@ -209,22 +134,8 @@ public:
 	 */
 	void ProcessPointCloud();
 
-	/**
-	 * Convert an array of points of "DefaultParametric" space into a scaled parametric space
-	 * @see FThinZone2DFinder::BuildBoundarySegments()
-	 */
-	void TransformPoints(EGridSpace DestinationSpace, const TArray<FPoint2D>& InPointsToScale, TArray<FPoint2D>& OutScaledPoints) const;
-
 protected:
 
-	/**
-	 * Builds the scaled parametric spaces
-	 * @see Points2D
-	 * @see GeneratePointCloud (ended GeneratePointCloud process)
-	 * @return false if the scaled grid is degenerated
-	 */
-	bool ScaleGrid();
-
 	/**
 	 * Projects loop's points in the scaled parametric spaces
 	 * @see FaceLoops2D
@@ -232,29 +143,6 @@ protected:
 	 */
 	void ScaleLoops();
 
-	/**
-	 * Convert Coordinate of "DefaultParametric" space into a scaled parametric space
-	 * @see ScaleLoops
-	 */
-	void ComputeNewCoordinate(const TArray<FPoint2D>& NewGrid, int32 IndexU, int32 IndexV, const FPoint2D& InPoint, FPoint2D& OutNewScaledPoint) const
-	{
-		const FPoint2D& PointU0V0 = NewGrid[(IndexV + 0) * CuttingCount[EIso::IsoU] + (IndexU + 0)];
-		const FPoint2D& PointU1V0 = NewGrid[(IndexV + 0) * CuttingCount[EIso::IsoU] + (IndexU + 1)];
-		const FPoint2D& PointU0V1 = NewGrid[(IndexV + 1) * CuttingCount[EIso::IsoU] + (IndexU + 0)];
-		const FPoint2D& PointU1V1 = NewGrid[(IndexV + 1) * CuttingCount[EIso::IsoU] + (IndexU + 1)];
-
-		const double U1MinusU0 = CuttingCoordinates[EIso::IsoU][IndexU + 1] - CuttingCoordinates[EIso::IsoU][IndexU];
-		const double U0MinusU = CuttingCoordinates[EIso::IsoU][IndexU] - InPoint.U;
-		const double V1MinusV0 = CuttingCoordinates[EIso::IsoV][IndexV + 1] - CuttingCoordinates[EIso::IsoV][IndexV];
-		const double V0MinusV = CuttingCoordinates[EIso::IsoV][IndexV] - InPoint.V;
-
-		OutNewScaledPoint =
-			PointU0V0 +
-			(PointU0V0 - PointU1V0) * (U0MinusU / U1MinusU0) +
-			(PointU0V0 - PointU0V1) * (V0MinusV / V1MinusV0) +
-			(PointU0V0 - PointU0V1 + PointU1V1 - PointU1V0) * ((U0MinusU * V0MinusV) / (U1MinusU0 * V1MinusV0));		
-	};
-
 	/**
 	 * @see ProcessPointClound (called in ProcessPointClound)
 	 */
@@ -289,26 +177,26 @@ protected:
 	 */
 	void GetPreferredUVCuttingParametersFromLoops(FCuttingGrid& CuttingFromLoops);
 
-	void ComputeMaxElementSize();
-	void ComputeMaxDeltaUV();
-
 	/**
 	 * @ return false if the mesh of the loop is degenerated
 	 */
 	bool GetMeshOfLoops();
 	void GetMeshOfLoop(const FTopologicalLoop& Loop);
 
-public:
-
-	// ======================================================================================================================================================================================================================
-	// ThinZone properties for meshing scheduling optimisation ==============================================================================================================================================================
+	// Meshing tools =========================================================================================================================================================================================================
 	// ======================================================================================================================================================================================================================
 
-	void SearchThinZones();
-	const TArray<FThinZone2D>& GetThinZones() const
-	{
-		return ThinZoneFinder.GetThinZones();
-	}
+	/**
+	 * @return true if the grid is not consistent to build a mesh as composed of only two border nodes.
+	 */
+	bool CheckIfExternalLoopIsDegenerate() const;
+
+	/**
+	 * @return true if along an Iso, all (U(i+1) - U(i)) < FaceTolerance
+	 */
+	bool CheckIf2DGridIsDegenerate() const;
+
+public:
 
 	// ======================================================================================================================================================================================================================
 	// GET Methodes =========================================================================================================================================================================================================
@@ -420,45 +308,11 @@ public:
 		NodeMarkers[Index] &= ~ENodeMarker::IsInside;
 	}
 
-	/**
-	 * @return the FPoint2D (parametric coordinates) of the point at the Index of the grid in the defined grid space
-	 * @see EGridSpace
-	 * @see Points2D
-	 */
-	const FPoint2D& GetInner2DPoint(EGridSpace Space, int32 Index) const
-	{
-		return Points2D[(int32)Space][Index];
-	}
-
 	void SetInner2DPoint(EGridSpace Space, int32 Index, const FPoint2D& NewCoordinate)
 	{
 		Points2D[(int32)Space][Index] = NewCoordinate;
 	}
 
-	/**
-	 * @return the FPoint2D (parametric coordinates) of the point at the Index of the grid in the defined grid space @see EGridSpace
-	 */
-	const FPoint2D& GetInner2DPoint(EGridSpace Space, int32 IndexU, int32 IndexV) const
-	{
-		return Points2D[(int32)Space][GobalIndex(IndexU, IndexV)];
-	}
-
-	/**
-	 * @return the FPoint (3D coordinates) of the point at the Index of the grid
-	 */
-	const FPoint& GetInner3DPoint(int32 Index) const
-	{
-		return Points3D[Index];
-	}
-
-	/**
-	 * @return the FPoint (3D coordinates) of the point at the Index of the grid
-	 */
-	const FPoint& GetInner3DPoint(int32 IndexU, int32 IndexV) const
-	{
-		return Points3D[GobalIndex(IndexU, IndexV)];
-	}
-
 	/**
 	 * @return the normal of the surface at the point at the Index of the grid
 	 */
@@ -475,48 +329,14 @@ public:
 		return Normals[Index];
 	}
 
-	constexpr const TArray<double>& GetCuttingCoordinatesAlongIso(EIso Iso) const
+	const TArray<double>& GetCuttingCoordinatesAlongIso(EIso Iso) const
 	{
-		return CuttingCoordinates[Iso];
+		return CoordinateGrid[Iso];
 	}
 
 	const FCoordinateGrid& GetCuttingCoordinates() const
 	{
-		return CuttingCoordinates;
-	}
-
-	constexpr const TArray<double>& GetUniformCuttingCoordinatesAlongIso(EIso Iso) const
-	{
-		return UniformCuttingCoordinates[Iso];
-	}
-
-	const FCoordinateGrid& GetUniformCuttingCoordinates() const
-	{
-		return UniformCuttingCoordinates;
-	}
-
-	/**
-	 * @return the array of 3d points of the grid
-	 */
-	TArray<FPoint>& GetInner3DPoints()
-	{
-		return Points3D;
-	}
-
-	/**
-	 * @return the array of 3d points of the grid
-	 */
-	const TArray<FPoint>& GetInner3DPoints() const
-	{
-		return Points3D;
-	}
-
-	/**
-	 * @return the array of 2d points of the grid in the defined space
-	 */
-	const TArray<FPoint2D>& GetInner2DPoints(EGridSpace Space) const
-	{
-		return Points2D[(int32)Space];
+		return CoordinateGrid;
 	}
 
 	/**
@@ -527,16 +347,6 @@ public:
 		return Normals;
 	}
 
-	const FTopologicalFace& GetFace() const
-	{
-		return Face;
-	}
-
-	FTopologicalFace& GetFace()
-	{
-		return Face;
-	}
-
 	const TArray<TArray<int32>>& GetNodeIdsOfFaceLoops() const
 	{
 		return NodeIdsOfFaceLoops;
@@ -594,14 +404,6 @@ public:
 		return NormalsOfFaceLoops;
 	}
 
-	/**
-	 * @return the Index of the position in the arrays of a point [IndexU, IndexV] of the grid
-	 */
-	int32 GobalIndex(int32 IndexU, int32 IndexV) const
-	{
-		return IndexV * CuttingCount[EIso::IsoU] + IndexU;
-	}
-
 	/**
 	 * @return the Index of the position in the arrays of a point [IndexU, IndexV] of the grid
 	 */
@@ -611,14 +413,6 @@ public:
 		OutIndexV = GLobalIndex / CuttingCount[EIso::IsoU];
 	}
 
-	/**
-	 * @return the maximum difference of coordinate along the specified axis of two two successive points
-	 */
-	double GetMaxDeltaU(EIso Iso) const
-	{
-		return MaxDeltaUV[Iso];
-	}
-
 	/**
 	 * @return the minimal size of the mesh elements according to 3d geometric tolerance
 	 */
@@ -649,15 +443,15 @@ public:
 	void DisplayIsoSegment(EGridSpace Space, const FIsoSegment& Segment, FIdent Ident = 0, EVisuProperty Property = EVisuProperty::BlueCurve, bool bDisplayOrientation = false) const;
 	void DisplayIsoSegment(EGridSpace Space, const FIsoNode& NodeA, const FIsoNode& NodeB, FIdent Ident = 0, EVisuProperty Property = EVisuProperty::BlueCurve) const;
 
-	void DisplayIsoSegments(EGridSpace Space, const TArray<FIsoSegment*>& Segments, bool bDisplayNode = false, bool bDisplayOrientation = false, EVisuProperty Property = EVisuProperty::BlueCurve) const;
-	void DisplayIsoSegments(const FString& Message, EGridSpace Space, const TArray<FIsoSegment*>& Segments, bool bDisplayNode = false, bool bDisplayOrientation = false, EVisuProperty Property = EVisuProperty::BlueCurve) const
+	void DisplayIsoSegments(EGridSpace Space, const TArray<FIsoSegment*>& Segments, bool bDisplayNode = false, bool bDisplayOrientation = false, bool bSplitBySegment = false, EVisuProperty Property = EVisuProperty::BlueCurve) const;
+	void DisplayIsoSegments(const FString& Message, EGridSpace Space, const TArray<FIsoSegment*>& Segments, bool bDisplayNode = false, bool bDisplayOrientation = false, bool bSplitBySegment = false, EVisuProperty Property = EVisuProperty::BlueCurve) const
 	{
 		if (!bDisplay)
 		{
 			return;
 		}
 		F3DDebugSession _(Message);
-		DisplayIsoSegments(Space, Segments, bDisplayNode, bDisplayOrientation, Property);
+		DisplayIsoSegments(Space, Segments, bDisplayNode, bDisplayOrientation, bSplitBySegment, Property);
 	}
 
 	void DisplayTriangle(EGridSpace Space, const FIsoNode& NodeA, const FIsoNode& NodeB, const FIsoNode& NodeC) const;
@@ -706,19 +500,19 @@ public:
 		{
 			if (IsNodeInsideAndCloseToLoop(Index))
 			{
-				DisplayPoint(Points[Index] * DisplayScale, EVisuProperty::OrangePoint, Index);
+				DisplayPoint2DWithScale(Points[Index], EVisuProperty::OrangePoint, Index);
 			}
 			else if (IsNodeInsideAndMeshable(Index))
 			{
-				DisplayPoint(Points[Index] * DisplayScale, EVisuProperty::BluePoint, Index);
+				DisplayPoint2DWithScale(Points[Index], EVisuProperty::BluePoint, Index);
 			}
 			else if (IsNodeTooCloseToLoop(Index))
 			{
-				DisplayPoint(Points[Index] * DisplayScale, EVisuProperty::RedPoint, Index);
+				DisplayPoint2DWithScale(Points[Index], EVisuProperty::RedPoint, Index);
 			}
 			else
 			{
-				DisplayPoint(Points[Index] * DisplayScale, IsNodeCloseToLoop(Index) ? EVisuProperty::YellowPoint : EVisuProperty::GreenPoint , Index);
+				DisplayPoint2DWithScale(Points[Index], IsNodeCloseToLoop(Index) ? EVisuProperty::YellowPoint : EVisuProperty::GreenPoint, Index);
 			}
 		}
 	}
@@ -770,8 +564,6 @@ public:
 	}
 
 #endif
-
-	void PrintTimeElapse() const;
 };
 
 }
diff --git a/Engine/Source/Runtime/Datasmith/CADKernel/Private/CADKernel/Mesh/Structure/GridBase.cpp b/Engine/Source/Runtime/Datasmith/CADKernel/Private/CADKernel/Mesh/Structure/GridBase.cpp
new file mode 100644
--- /dev/null
+++ b/Engine/Source/Runtime/Datasmith/CADKernel/Private/CADKernel/Mesh/Structure/GridBase.cpp
@@ -0,0 +1,254 @@
+// Copyright Epic Games, Inc. All Rights Reserved.
+
+#include "CADKernel/Mesh/Structure/GridBase.h"
+
+#include "CADKernel/Geo/Sampling/SurfacicSampling.h"
+#include "CADKernel/Topo/TopologicalFace.h"
+#include "CADKernel/Utils/ArrayUtils.h"
+
+//#define DEBUG_GRID
+//#define DEBUG_GETPREFERREDUVCOORDINATESFROMNEIGHBOURS
+namespace UE::CADKernel
+{
+
+FGridBase::FGridBase(FTopologicalFace& InFace)
+	: Face(InFace)
+	, Tolerance3D(InFace.GetCarrierSurface()->Get3DTolerance())
+{
+#ifdef DEBUG_ONLY_SURFACE_TO_DEBUG
+	bDisplay = (InFace.GetId() == FaceToDebug);
+	Open3DDebugSession(bDisplay, FString::Printf(TEXT("Grid %d"), InFace.GetId()));
+#endif
+}
+
+bool FGridBase::ScaleGrid()
+{
+	const FCoordinateGrid& CoordinateGrid = GetCoordinateGrid();
+
+	FTimePoint StartTime = FChrono::Now();
+
+	TFunction<double(const TArray<double>&)> GetMean = [](const TArray<double>& Lengths)
+	{
+		double MeanLength = 0;
+		for (double Length : Lengths)
+		{
+			MeanLength += Length;
+		}
+		MeanLength /= Lengths.Num();
+		return MeanLength;
+	};
+
+	TFunction<double(const TArray<double>&, const double)> StandardDeviation = [](const TArray<double>& Lengths, const double MeanLength)
+	{
+		double StandardDeviation = 0;
+		for (double Length : Lengths)
+		{
+			StandardDeviation += FMath::Square(Length - MeanLength);
+		}
+		StandardDeviation /= Lengths.Num();
+		StandardDeviation = sqrt(StandardDeviation);
+		return StandardDeviation;
+	};
+
+	TFunction<void(const TArray<double>&, const double, TArray<double>&)> ScaleCoordinates = [](const TArray<double>& InCoordinates, const double ScaleFactor, TArray<double>& OutCoordinatesScaled)
+	{
+		OutCoordinatesScaled.Reserve(InCoordinates.Num());
+
+		for (double Coordinate : InCoordinates)
+		{
+			OutCoordinatesScaled.Add(Coordinate * ScaleFactor);
+		}
+	};
+
+	TFunction<int32(const TArray<double>&, const double)> GetMiddleIndex = [](const TArray<double>& Coordinates, double Middle)
+	{
+		int32 StartIndexUp = 1;
+		for (; StartIndexUp < Coordinates.Num(); ++StartIndexUp)
+		{
+			if (Coordinates[StartIndexUp] > Middle)
+			{
+				break;
+			}
+		}
+		return StartIndexUp;
+	};
+
+	for (int32 Index = EGridSpace::Default2D + 1; Index < EGridSpace::EndGridSpace; ++Index)
+	{
+		Points2D[Index].SetNum(CuttingSize);
+	}
+
+	TArray<double> LengthsV;
+	LengthsV.SetNum(CuttingCount[EIso::IsoU]);
+	for (int32 IndexU = 0; IndexU < CuttingCount[EIso::IsoU]; ++IndexU)
+	{
+		int32 Index = IndexU;
+		double Length = 0;
+		for (int32 IndexV = 1; IndexV < CuttingCount[EIso::IsoV]; ++IndexV)
+		{
+			Length += Points3D[Index].Distance(Points3D[Index + CuttingCount[EIso::IsoU]]);
+			Index += CuttingCount[EIso::IsoU];
+		}
+		LengthsV[IndexU] = Length;
+	}
+
+	TArray<double> LengthsU;
+	LengthsU.SetNum(CuttingCount[EIso::IsoV]);
+	for (int32 IndexV = 0, Index = 0; IndexV < CuttingCount[EIso::IsoV]; ++IndexV)
+	{
+		double Length = 0;
+		for (int32 IndexU = 1; IndexU < CuttingCount[EIso::IsoU]; ++IndexU)
+		{
+			Length += Points3D[Index].Distance(Points3D[Index + 1]);
+			Index++;
+		}
+		Index++;
+		LengthsU[IndexV] = Length;
+	}
+
+	double MeanLengthV = GetMean(LengthsV);
+	double MeanLengthU = GetMean(LengthsU);
+	if (MeanLengthV < Tolerance3D || MeanLengthU < Tolerance3D)
+	{
+		SetAsDegenerated();
+		return false;
+	}
+
+	double FactorV = MeanLengthV / (CoordinateGrid[EIso::IsoV].Last() - CoordinateGrid[EIso::IsoV][0]);
+	double FactorU = MeanLengthU / (CoordinateGrid[EIso::IsoU].Last() - CoordinateGrid[EIso::IsoU][0]);
+
+	ScaleCoordinates(CoordinateGrid[EIso::IsoU], FactorU, UniformCuttingCoordinates[EIso::IsoU]);
+	ScaleCoordinates(CoordinateGrid[EIso::IsoV], FactorV, UniformCuttingCoordinates[EIso::IsoV]);
+
+	{
+		int32 NumUV = 0;
+		for (int32 IPointV = 0; IPointV < CuttingCount[EIso::IsoV]; ++IPointV)
+		{
+			for (int32 IPointU = 0; IPointU < CuttingCount[EIso::IsoU]; ++IPointU, ++NumUV)
+			{
+				Points2D[EGridSpace::UniformScaled][NumUV].Set(UniformCuttingCoordinates[EIso::IsoU][IPointU], UniformCuttingCoordinates[EIso::IsoV][IPointV]);
+			}
+		}
+	}
+
+	double StandardDeviationU = StandardDeviation(LengthsU, MeanLengthU);
+	double StandardDeviationV = StandardDeviation(LengthsV, MeanLengthV);
+
+	if (StandardDeviationV > StandardDeviationU)
+	{
+		double MiddleV = (CoordinateGrid[EIso::IsoV].Last() + CoordinateGrid[EIso::IsoV][0]) * 0.5;
+
+		FCoordinateGrid Grid;
+		Grid[EIso::IsoU] = CoordinateGrid[EIso::IsoU];
+		Grid[EIso::IsoV].Add(MiddleV);
+
+		FSurfacicSampling MiddlePoints;
+		Face.EvaluatePointGrid(Grid, MiddlePoints);
+
+		int32 StartIndexUp = GetMiddleIndex(CoordinateGrid[EIso::IsoV], MiddleV);
+		int32 StartIndexDown = StartIndexUp - 1;
+
+		int32 NumUV = 0;
+		for (int32 IPointU = 0; IPointU < CuttingCount[EIso::IsoU]; ++IPointU)
+		{
+			double Length = 0;
+			FPoint LastPoint = MiddlePoints.Points3D[IPointU];
+			for (int32 IPointV = StartIndexUp; IPointV < CuttingCount[EIso::IsoV]; ++IPointV)
+			{
+				NumUV = IPointV * CuttingCount[EIso::IsoU] + IPointU;
+				Length += LastPoint.Distance(Points3D[NumUV]);
+				Points2D[EGridSpace::Scaled][NumUV].Set(Points2D[EGridSpace::UniformScaled][NumUV].U, Length);
+				LastPoint = Points3D[NumUV];
+			}
+
+			Length = 0;
+			LastPoint = MiddlePoints.Points3D[IPointU];
+			for (int32 IPointV = StartIndexDown; IPointV >= 0; --IPointV)
+			{
+				NumUV = IPointV * CuttingCount[EIso::IsoU] + IPointU;
+				Length -= LastPoint.Distance(Points3D[NumUV]);
+				Points2D[EGridSpace::Scaled][NumUV].Set(Points2D[EGridSpace::UniformScaled][NumUV].U, Length);
+				LastPoint = Points3D[NumUV];
+			}
+		}
+	}
+	else
+	{
+		double MiddleU = (CoordinateGrid[EIso::IsoU].Last() + CoordinateGrid[EIso::IsoU][0]) * 0.5;
+
+		FCoordinateGrid Grid;
+		Grid[EIso::IsoU].Add(MiddleU);
+		Grid[EIso::IsoV] = CoordinateGrid[EIso::IsoV];
+
+		FSurfacicSampling MiddlePoints;
+		Face.EvaluatePointGrid(Grid, MiddlePoints);
+
+		int32 StartIndexUp = GetMiddleIndex(CoordinateGrid[EIso::IsoU], MiddleU);
+		int32 StartIndexDown = StartIndexUp - 1;
+
+		int32 NumUV = 0;
+		for (int32 IPointV = 0; IPointV < CuttingCount[EIso::IsoV]; ++IPointV)
+		{
+			double Length = 0;
+			FPoint LastPoint = MiddlePoints.Points3D[IPointV];
+			for (int32 IPointU = StartIndexUp; IPointU < CuttingCount[EIso::IsoU]; ++IPointU)
+			{
+				NumUV = IPointV * CuttingCount[EIso::IsoU] + IPointU;
+				Length += LastPoint.Distance(Points3D[NumUV]);
+				Points2D[EGridSpace::Scaled][NumUV].Set(Length, Points2D[EGridSpace::UniformScaled][NumUV].V);
+				LastPoint = Points3D[NumUV];
+			}
+
+			Length = 0;
+			LastPoint = MiddlePoints.Points3D[IPointV];
+			for (int32 IPointU = StartIndexDown; IPointU >= 0; --IPointU)
+			{
+				NumUV = IPointV * CuttingCount[EIso::IsoU] + IPointU;
+				Length -= LastPoint.Distance(Points3D[NumUV]);
+				Points2D[EGridSpace::Scaled][NumUV].Set(Length, Points2D[EGridSpace::UniformScaled][NumUV].V);
+				LastPoint = Points3D[NumUV];
+			}
+		}
+	}
+	Chronos.ScaleGridDuration = FChrono::Elapse(StartTime);
+	return true;
+}
+
+void FGridBase::TransformPoints(EGridSpace DestinationSpace, const TArray<FPoint2D>& InPointsToScale, TArray<FPoint2D>& OutTransformedPoints) const
+{
+	const FCoordinateGrid& CoordinateGrid = GetCoordinateGrid();
+
+	OutTransformedPoints.SetNum(InPointsToScale.Num());
+
+	int32 IndexU = 0;
+	int32 IndexV = 0;
+	for (int32 Index = 0; Index < InPointsToScale.Num(); ++Index)
+	{
+		const FPoint2D& Point = InPointsToScale[Index];
+
+		ArrayUtils::FindCoordinateIndex(CoordinateGrid[EIso::IsoU], Point.U, IndexU);
+		ArrayUtils::FindCoordinateIndex(CoordinateGrid[EIso::IsoV], Point.V, IndexV);
+
+		ComputeNewCoordinate(Points2D[DestinationSpace], IndexU, IndexV, Point, OutTransformedPoints[Index]);
+	}
+}
+
+void FGridBase::EvaluatePointGrid(const FCoordinateGrid& CoordGrid, bool bWithNormals)
+{
+	CuttingSize = CoordGrid.Count();
+
+	FSurfacicSampling Grid;
+	Face.EvaluatePointGrid(CoordGrid, Grid, bWithNormals);
+
+	Points2D[EGridSpace::Default2D] = MoveTemp(Grid.Points2D);
+	Points3D = MoveTemp(Grid.Points3D);
+
+	if (bWithNormals)
+	{
+		Normals = MoveTemp(Grid.Normals);
+	}
+}
+
+}
+
+
diff --git a/Engine/Source/Runtime/Datasmith/CADKernel/Private/CADKernel/Mesh/Structure/GridBase.h b/Engine/Source/Runtime/Datasmith/CADKernel/Private/CADKernel/Mesh/Structure/GridBase.h
new file mode 100644
--- /dev/null
+++ b/Engine/Source/Runtime/Datasmith/CADKernel/Private/CADKernel/Mesh/Structure/GridBase.h
@@ -0,0 +1,263 @@
+// Copyright Epic Games, Inc. All Rights Reserved.
+#pragma once
+
+#include "CADKernel/Core/Chrono.h"
+#include "CADKernel/Core/HaveStates.h"
+#include "CADKernel/Geo/GeoEnum.h"
+#include "CADKernel/Geo/GeoPoint.h"
+#include "CADKernel/Math/Point.h"
+#include "CADKernel/Mesh/MeshEnum.h"
+#include "CADKernel/UI/Display.h"
+
+#ifdef CADKERNEL_DEV
+#include "CADKernel/Mesh/Meshers/IsoTriangulator/DefineForDebug.h"
+#endif
+
+
+namespace UE::CADKernel
+{
+
+struct FCuttingPoint;
+struct FCuttingGrid;
+
+class FIsoNode;
+class FIsoSegment;
+class FLoopNode;
+class FModelMesh;
+class FPoint;
+class FPoint2D;
+class FThinZone2DFinder;
+class FTopologicalFace;
+
+struct FGridChronos
+{
+	FDuration DefineCuttingParametersDuration;
+	FDuration GeneratePointCloudDuration;
+	FDuration ProcessPointCloudDuration;
+	FDuration FindInnerDomainPointsDuration;
+	FDuration Build2DLoopDuration;
+	FDuration RemovePointsClosedToLoopDuration;
+	FDuration FindPointsCloseToLoopDuration;
+	FDuration ScaleGridDuration;
+
+	FGridChronos()
+		: DefineCuttingParametersDuration(FChrono::Init())
+		, GeneratePointCloudDuration(FChrono::Init())
+		, ProcessPointCloudDuration(FChrono::Init())
+		, FindInnerDomainPointsDuration(FChrono::Init())
+		, Build2DLoopDuration(FChrono::Init())
+		, RemovePointsClosedToLoopDuration(FChrono::Init())
+		, FindPointsCloseToLoopDuration(FChrono::Init())
+		, ScaleGridDuration(FChrono::Init())
+	{}
+
+	void PrintTimeElapse() const
+	{
+		FDuration GridDuration = FChrono::Elapse(FChrono::Now());
+		GridDuration += DefineCuttingParametersDuration;
+		GridDuration += GeneratePointCloudDuration;
+		GridDuration += FindInnerDomainPointsDuration;
+		GridDuration += Build2DLoopDuration;
+		GridDuration += RemovePointsClosedToLoopDuration;
+		GridDuration += FindPointsCloseToLoopDuration;
+		FChrono::PrintClockElapse(Log, TEXT(""), TEXT("Grid"), GridDuration);
+		FChrono::PrintClockElapse(Log, TEXT("  "), TEXT("DefineCuttingParameters"), DefineCuttingParametersDuration);
+		FChrono::PrintClockElapse(Log, TEXT("  "), TEXT("GeneratePointCloud"), GeneratePointCloudDuration);
+		FChrono::PrintClockElapse(Log, TEXT("  "), TEXT("GenerateDomainPoints"), ProcessPointCloudDuration);
+		FChrono::PrintClockElapse(Log, TEXT("    "), TEXT("FindInnerDomainPointsDuration"), FindInnerDomainPointsDuration);
+		FChrono::PrintClockElapse(Log, TEXT("    "), TEXT("FindPointsCloseToLoop"), FindPointsCloseToLoopDuration);
+		FChrono::PrintClockElapse(Log, TEXT("    "), TEXT("RemovePointsClosedToLoop"), RemovePointsClosedToLoopDuration);
+		FChrono::PrintClockElapse(Log, TEXT("    "), TEXT("Build2DLoopDuration"), Build2DLoopDuration);
+	}
+};
+
+class FGridBase : public FHaveStates
+{
+
+protected:
+
+	FTopologicalFace& Face;
+	const double Tolerance3D;
+
+	/**
+	 * grid point cloud size
+	 */
+	int32 CuttingCount[2] = { 0 , 0 };
+	int32 CuttingSize = 0;
+
+	FCoordinateGrid UniformCuttingCoordinates;
+
+	/**
+	 * 2D Coordinate of grid nodes in each space
+	 */
+	TArray<FPoint2D> Points2D[EGridSpace::EndGridSpace];
+
+	/**
+	 * 3D Coordinate of inner nodes
+	 */
+	TArray<FPoint> Points3D;
+
+	/**
+	 * Surface Normal at each inner nodes
+	 */
+	TArray<FVector3f> Normals;
+
+public:
+	FGridChronos Chronos;
+
+public:
+	FGridBase(FTopologicalFace& InFace);
+
+	virtual ~FGridBase()
+	{
+#ifdef DEBUG_ONLY_SURFACE_TO_DEBUG
+		Close3DDebugSession(bDisplay);
+#endif
+	}
+
+	/**
+	 * Convert an array of points of "DefaultParametric" space into a scaled parametric space
+	 * @see FThinZone2DFinder::BuildBoundarySegments()
+	 */
+	void TransformPoints(EGridSpace DestinationSpace, const TArray<FPoint2D>& InPointsToScale, TArray<FPoint2D>& OutScaledPoints) const;
+
+	/**
+	 * Builds the scaled parametric spaces
+	 * @see Points2D
+	 * @see GeneratePointCloud (ended GeneratePointCloud process)
+	 * @return false if the scaled grid is degenerated
+	 */
+	bool ScaleGrid();
+
+	double GetTolerance() const
+	{
+		return Tolerance3D;
+	}
+
+protected:
+
+	virtual const FCoordinateGrid& GetCoordinateGrid() const = 0;
+
+	void EvaluatePointGrid(const FCoordinateGrid& CoordinateGrid, bool bWithNormals);
+
+	/**
+	 * Convert Coordinate of "DefaultParametric" space into a scaled parametric space
+	 * @see ScaleLoops
+	 */
+	void ComputeNewCoordinate(const TArray<FPoint2D>& NewGrid, int32 IndexU, int32 IndexV, const FPoint2D& InPoint, FPoint2D& OutNewScaledPoint) const
+	{
+		const FCoordinateGrid& CoordinateGrid = GetCoordinateGrid();
+
+		const FPoint2D& PointU0V0 = NewGrid[(IndexV + 0) * CuttingCount[EIso::IsoU] + (IndexU + 0)];
+		const FPoint2D& PointU1V0 = NewGrid[(IndexV + 0) * CuttingCount[EIso::IsoU] + (IndexU + 1)];
+		const FPoint2D& PointU0V1 = NewGrid[(IndexV + 1) * CuttingCount[EIso::IsoU] + (IndexU + 0)];
+		const FPoint2D& PointU1V1 = NewGrid[(IndexV + 1) * CuttingCount[EIso::IsoU] + (IndexU + 1)];
+
+		const double U1MinusU0 = CoordinateGrid[EIso::IsoU][IndexU + 1] - CoordinateGrid[EIso::IsoU][IndexU];
+		const double U0MinusU = CoordinateGrid[EIso::IsoU][IndexU] - InPoint.U;
+		const double V1MinusV0 = CoordinateGrid[EIso::IsoV][IndexV + 1] - CoordinateGrid[EIso::IsoV][IndexV];
+		const double V0MinusV = CoordinateGrid[EIso::IsoV][IndexV] - InPoint.V;
+
+		OutNewScaledPoint =
+			PointU0V0 +
+			(PointU0V0 - PointU1V0) * (U0MinusU / U1MinusU0) +
+			(PointU0V0 - PointU0V1) * (V0MinusV / V1MinusV0) +
+			(PointU0V0 - PointU0V1 + PointU1V1 - PointU1V0) * ((U0MinusU * V0MinusV) / (U1MinusU0 * V1MinusV0));	
+	};
+
+public:
+
+	/**
+	 * @return the FPoint2D (parametric coordinates) of the point at the Index of the grid in the defined grid space
+	 * @see EGridSpace
+	 * @see Points2D
+	 */
+	const FPoint2D& GetInner2DPoint(EGridSpace Space, int32 Index) const
+	{
+		return Points2D[(int32)Space][Index];
+	}
+
+	/**
+	 * @return the FPoint2D (parametric coordinates) of the point at the Index of the grid in the defined grid space @see EGridSpace
+	 */
+	const FPoint2D& GetInner2DPoint(EGridSpace Space, int32 IndexU, int32 IndexV) const
+	{
+		return Points2D[(int32)Space][GobalIndex(IndexU, IndexV)];
+	}
+
+	/**
+	 * @return the FPoint (3D coordinates) of the point at the Index of the grid
+	 */
+	const FPoint& GetInner3DPoint(int32 Index) const
+	{
+		return Points3D[Index];
+	}
+
+	constexpr const TArray<double>& GetUniformCuttingCoordinatesAlongIso(EIso Iso) const
+	{
+		return UniformCuttingCoordinates[Iso];
+	}
+
+	const FCoordinateGrid& GetUniformCuttingCoordinates() const
+	{
+		return UniformCuttingCoordinates;
+	}
+
+	/**
+	 * @return the array of 3d points of the grid
+	 */
+	TArray<FPoint>& GetInner3DPoints()
+	{
+		return Points3D;
+	}
+
+	/**
+	 * @return the array of 3d points of the grid
+	 */
+	const TArray<FPoint>& GetInner3DPoints() const
+	{
+		return Points3D;
+	}
+
+	/**
+	 * @return the array of 2d points of the grid in the defined space
+	 */
+	const TArray<FPoint2D>& GetInner2DPoints(EGridSpace Space) const
+	{
+		return Points2D[(int32)Space];
+	}
+
+	const FTopologicalFace& GetFace() const
+	{
+		return Face;
+	}
+
+	FTopologicalFace& GetFace()
+	{
+		return Face;
+	}
+
+	/**
+	 * @return the Index of the position in the arrays of a point [IndexU, IndexV] of the grid
+	 */
+	int32 GobalIndex(int32 IndexU, int32 IndexV) const
+	{
+		return IndexV * CuttingCount[EIso::IsoU] + IndexU;
+	}
+
+#ifdef CADKERNEL_DEV
+
+	// ======================================================================================================================================================================================================================
+	// Display Methodes   ================================================================================================================================================================================================
+	// ======================================================================================================================================================================================================================
+	bool bDisplay = true;
+
+	void DisplayIsoNode(EGridSpace Space, const int32 PointIndex, FIdent Ident = 0, EVisuProperty Property = EVisuProperty::BluePoint) const;
+	virtual void DisplayGridPoints(EGridSpace DisplaySpace) const;
+	void DisplayInnerPoints(TCHAR* Message, EGridSpace DisplaySpace) const;
+
+	void PrintTimeElapse() const;
+#endif
+};
+
+}
diff --git a/Engine/Source/Runtime/Datasmith/CADKernel/Private/CADKernel/Mesh/Structure/GridDisplay.cpp b/Engine/Source/Runtime/Datasmith/CADKernel/Private/CADKernel/Mesh/Structure/GridDisplay.cpp
deleted file mode 100644
diff --git a/Engine/Source/Runtime/Datasmith/CADKernel/Private/CADKernel/Mesh/Structure/ThinZone2D.cpp b/Engine/Source/Runtime/Datasmith/CADKernel/Private/CADKernel/Mesh/Structure/ThinZone2D.cpp
--- a/Engine/Source/Runtime/Datasmith/CADKernel/Private/CADKernel/Mesh/Structure/ThinZone2D.cpp
+++ b/Engine/Source/Runtime/Datasmith/CADKernel/Private/CADKernel/Mesh/Structure/ThinZone2D.cpp
@@ -1,811 +1,3 @@
 // Copyright Epic Games, Inc. All Rights Reserved.
 
 #include "CADKernel/Mesh/Structure/ThinZone2D.h"
-
-#include "CADKernel/Math/Point.h"
-#include "CADKernel/Mesh/Structure/Grid.h"
-#include "CADKernel/Mesh/Structure/EdgeSegment.h"
-#include "CADKernel/Topo/TopologicalEdge.h"
-#include "CADKernel/Topo/TopologicalFace.h"
-#include "CADKernel/Topo/TopologicalLoop.h"
-#include "CADKernel/UI/Display.h"
-
-namespace UE::CADKernel
-{
-
-FIdent FEdgeSegment::LastId = 0;
-
-void FThinZone2DFinder::FindClosedSegments()
-{
-	double MaxSegmentLength = 0.0;
-	for (const FEdgeSegment* Segment : LoopSegments)
-	{
-		double Length = Segment->GetLength();
-		if (Length > MaxSegmentLength) 
-		{
-			MaxSegmentLength = Length;
-		}
-	}
-	MaxSegmentLength *= 1.01;
-
-	double MaxSpace = FMath::Max(MaxSegmentLength, Tolerance * 1.01);
-	double MaxSquareSpace = 4 * FMath::Square(MaxSpace);
-
-	// Copy of loop segments to generate a sorted array of segments
-	SortedLoopSegments = LoopSegments;
-	Algo::Sort(SortedLoopSegments, [](FEdgeSegment* SegmentA, FEdgeSegment* SegmentB)
-		{
-			return (SegmentA->GetAxeMin() <SegmentB->GetAxeMin());
-		});
-
-
-#ifdef DEBUG_THIN_ZONES
-	Open3DDebugSession(TEXT("Closed Segment");
-#endif
-
-	// For each segment, the nearest segment is search
-	// If segment is from an inner loop, the nearest could not be from the same loop
-	int32 SegmentToBeginIndex = 0;
-	for (FEdgeSegment* Segment : SortedLoopSegments)
-	{
-		FTopologicalLoop* SegmentLoop = nullptr;
-		if (Segment->IsInner())
-		{
-			SegmentLoop = Segment->GetEdge()->GetLoop();
-		}
-		else
-		{
-			SegmentLoop = nullptr;
-		}
-
-		FPoint2D SegmentMiddle = Segment->GetCenter();
-
-		FPoint2D ClosedPoint, ClosedPoint2;
-		FEdgeSegment* ClosedSegment = nullptr;
-
-		double SegmentMin = Segment->GetAxeMin();
-		double MinSquareThickness = HUGE_VALUE;
-
-		for (int32 CandidateIndex = SegmentToBeginIndex; CandidateIndex <SortedLoopSegments.Num(); ++CandidateIndex)
-		{
-			FEdgeSegment* Candidate = SortedLoopSegments[CandidateIndex];
-			if (Candidate == Segment) 
-			{
-				continue;
-			}
-
-			// Inner boundary are not check for thin inner boundary
-			if (SegmentLoop == Candidate->GetEdge()->GetLoop()) 
-			{
-				continue;
-			}
-
-			double CandidateSegmentMin = Candidate->GetAxeMin();
-
-			// If min point of candiate segment + maxSpace (= max length of the segments + Tolerance) is smaller than Min point of current segment then the distance between both segments cannot be smaller the Tolerance
-			if ((CandidateSegmentMin + MaxSpace) < SegmentMin)
-			{
-				SegmentToBeginIndex = CandidateIndex;
-				continue;
-			}
-
-			// If Min point of current segment + maxSpace is smaller than Min point of candidate segment then the distance between both segments cannot be smaller the Tolerance.
-			// As segments are sorted, next segments are not close to the current segment
-			if (CandidateSegmentMin > (SegmentMin + MaxSpace))
-			{
-				break;
-			}
-
-			const FPoint2D& FirstPointCandidate = Candidate->GetExtemity(ELimit::Start);
-
-			// if the distance of FirstPointCandidate with the middle of current segment is biggest than MaxSpace, then the projection of the point cannot be smaller than the tolerance, 
-			double SquareDistance = SegmentMiddle.SquareDistance(FirstPointCandidate);
-			if (SquareDistance > MaxSquareSpace)
-			{
-				continue;
-			}
-
-			const FPoint2D& SecondPointCandidate = Candidate->GetExtemity(ELimit::End);
-			double Coordinate;
-			FPoint2D Projection = ProjectPointOnSegment(SegmentMiddle, FirstPointCandidate, SecondPointCandidate, Coordinate, true);
-
-			SquareDistance = SegmentMiddle.SquareDistance(Projection);
-			if (SquareDistance > SquareTolerance)
-			{
-				continue;
-			}
-
-			if (MinSquareThickness > SquareDistance)
-			{
-				// check the angle between segments. As they are opposite, the cosAngle as to be close to -1
-				double Slope = Segment->ComputeUnorientedSlopeOf(Candidate);
-				if (Slope < 3.33) // Angle <5Pi/6 (150 deg)
-				{
-					continue;
-				}
-
-				// check the angle between segment and Middle-Projection. As they are opposite, the cosAngle as to be close to 0
-				Slope = Segment->ComputeUnorientedSlopeOf(SegmentMiddle, Projection);
-				if (Slope < 1.33) // Angle <Pi/3 (60 deg)
-				{
-					continue;
-				}
-
-				MinSquareThickness = SquareDistance;
-				ClosedPoint = Projection;
-				ClosedSegment = Candidate;
-			}
-		}
-
-		if (ClosedSegment)
-		{
-			Segment->SetClosedSegment(ClosedSegment, nullptr, MinSquareThickness, true);
-#ifdef DEBUG_THIN_ZONES
-			DisplaySegment(Segment->GetFirstPoint(), Segment->GetSecondPoint());
-			DisplaySegment(ClosedSegment->GetFirstPoint(), ClosedSegment->GetSecondPoint());
-			Wait();
-#endif
-		}
-	}
-#ifdef DEBUG_THIN_ZONES
-	Close3DDebugSession();
-#endif
-}
-
-void FThinZone2DFinder::DisplayClosedSegments()
-{
-#ifdef DEBUG_THIN_ZONES
-	F3DDebugSession _(TEXT("Closed Segment"));
-	for (const FEdgeSegment* Segment : LoopSegments)
-	{
-		if (Segment->GetClosedSegment())
-		{
-			DisplaySegment(Segment->GetFirstPoint(), Segment->GetSecondPoint(), Segment->GetId(), PropertyElement);
-			FEdgeSegment* Parallel = Segment->GetClosedSegment();
-			DisplaySegment(Parallel->GetFirstPoint(), Parallel->GetSecondPoint(), Segment->GetId(), PropertyElement);
-
-			double Coordinate;
-			FPoint Projection = ProjectPointOnSegment(Segment->GetCenter(), Parallel->GetFirstPoint(), Parallel->GetSecondPoint(), &Coordinate, true);
-
-			DisplaySegment(Projection, Segment->GetCenter(), Segment->GetId(), EVisuProperty::Iso);
-		}
-	}
-#endif
-}
-
-void FThinZone2DFinder::CheckClosedSegments()
-{
-	for (FEdgeSegment* Segment : LoopSegments)
-	{
-		if (Segment->GetClosedSegment())
-		{
-			double NextSegmentThickness = FMath::Min(Segment->GetPrevious()->GetClosedSquareDistance(), Segment->GetNext()->GetClosedSquareDistance());
-			double SegmentThickness = Segment->GetClosedSquareDistance();
-			GetMinMax(NextSegmentThickness, SegmentThickness);
-			SegmentThickness = sqrt(SegmentThickness);
-			NextSegmentThickness = sqrt(NextSegmentThickness);
-			double Delta = SegmentThickness - NextSegmentThickness;
-			if (Delta > (Segment->GetLength()))
-			{
-				Segment->ResetClosedData();
-			}
-		}
-	}
-}
-
-void FThinZone2DFinder::LinkClosedSegments()
-{
-	SortedLoopSegments.Empty();
-	SortedLoopSegments.Reserve(LoopSegments.Num());
-
-	FEdgeSegment* StartSegment = nullptr;
-	FEdgeSegment* Segment = nullptr;
-	FIdent Index = 0;
-	FIdent StartZoneIndex = 0;
-	FTopologicalLoop* LastLoop = nullptr;
-	bool bThinZone = false;
-
-	for (FEdgeSegment* EdgeSegment : LoopSegments)
-	{
-		Segment = EdgeSegment;
-		if (LastLoop != EdgeSegment->GetEdge()->GetLoop())
-		{
-			LastLoop = EdgeSegment->GetEdge()->GetLoop();
-			Index++;
-			StartSegment = EdgeSegment;
-			StartZoneIndex = Index;
-			bThinZone = false;
-		}
-
-		if (Segment->GetClosedSegment() != nullptr)
-		{
-			Segment->SetChainIndex(Index);
-			SortedLoopSegments.Add(Segment);
-			bThinZone = true;
-
-			if (Segment->GetNext() == StartSegment)
-			{
-				FIdent currentZoneIndex = Segment->GetChainIndex();
-				if (StartZoneIndex != 0 && currentZoneIndex != StartZoneIndex)
-				{
-					do {
-						Segment->SetChainIndex(StartZoneIndex);
-						Segment = Segment->GetPrevious();
-					} while (Segment->GetChainIndex() == currentZoneIndex);
-				}
-			}
-		}
-		else
-		{
-			if (bThinZone)
-			{
-				Index++;
-				bThinZone = false;
-			}
-		}
-	}
-
-	// the number of ThinZone should be less than 2x the number of Chain, 
-	ThinZones.Reserve(Index * 2);
-#ifdef DEBUG_THIN_ZONES
-	{
-		F3DDebugSession _(TEXT("Linking"));
-		FIdent LastIndex = 0;
-		int32 Color = 1;
-		for (FEdgeSegment* EdgeSegment : SortedLoopSegments)
-		{
-			if (LastIndex != EdgeSegment->GetChainIndex())
-			{
-				if (Color > 1) Close3DDebugSession();
-				Color++;
-				LastIndex = EdgeSegment->GetChainIndex();
-				Open3DDebugSession(FString::Printf(TEXT("Chain %d\n"), LastIndex));
-				if (Color == 7) Color++;
-				if (Color == 11) Color = 2;
-			}
-			DisplaySegment(EdgeSegment->GetExtemity(ELimit::Start), EdgeSegment->GetExtemity(ELimit::End), EdgeSegment->GetId(), (EVisuProperty)Color);
-		}
-		if (SortedLoopSegments.Num() > 0)
-		{
-			Close3DDebugSession();
-		}
-	}
-#endif
-
-}
-
-
-static void GetThinZoneSideConnectionsLength(TArray<FEdgeSegment*>& FirstSide, TArray<FEdgeSegment*>& SecondSide, double MaxLength, double* LengthBetweenExtremity, TArray<TSharedPtr<FTopologicalEdge>>* PeakEdges)
-{
-	LengthBetweenExtremity[0] = 0.0;
-	TSharedPtr<FTopologicalEdge> Edge = nullptr;
-	FEdgeSegment* Segment = FirstSide[0]->GetPrevious();
-	while (Segment != SecondSide[0])
-	{
-		if (Edge != Segment->GetEdge())
-		{
-			Edge = Segment->GetEdge();
-			PeakEdges[0].Add(Edge);
-		}
-		LengthBetweenExtremity[0] += Segment->GetLength();
-		Segment = Segment->GetPrevious();
-
-		if (LengthBetweenExtremity[0] > MaxLength)
-		{
-			LengthBetweenExtremity[0] = HUGE_VALUE;
-			PeakEdges[0].Empty(0);
-			break;
-		}
-	}
-
-	LengthBetweenExtremity[1] = 0.0;
-	Edge = nullptr;
-	Segment = FirstSide.Last()->GetNext();
-	while (Segment != SecondSide.Last())
-	{
-		if (Edge != Segment->GetEdge())
-		{
-			Edge = Segment->GetEdge();
-			PeakEdges[1].Add(Edge);
-		}
-
-		LengthBetweenExtremity[1] += Segment->GetLength();
-		Segment = Segment->GetNext();
-
-		if (LengthBetweenExtremity[1] > MaxLength)
-		{
-			LengthBetweenExtremity[1] = HUGE_VALUE;
-			PeakEdges[1].Empty(0);
-			break;
-		}
-	}
-}
-
-void FThinZone2DFinder::BuildThinZone()
-{
-	for (FEdgeSegment* EdgeSegment : SortedLoopSegments)
-	{
-		EdgeSegment->ResetMarker1();
-	}
-
-	TArray<FEdgeSegment*> FirstSide;
-	TArray<FEdgeSegment*> SecondSide;
-	FirstSide.Reserve(SortedLoopSegments.Num());
-	SecondSide.Reserve(SortedLoopSegments.Num());
-
-	FEdgeSegment* FirstSideSegment = nullptr;
-	FEdgeSegment* SecondSideSegment = nullptr;
-
-	for (FEdgeSegment* EdgeSegment : SortedLoopSegments)
-	{
-		bool bFirstSideIsClosed = false;
-		bool bSecondSideIsClosed = false;
-
-		FirstSide.Empty();
-		SecondSide.Empty();
-
-		if (EdgeSegment->HasMarker1()) 
-		{
-			continue;
-		}
-
-		FirstSideSegment = EdgeSegment;
-
-		// Find the first segment
-		FIdent FirstChainIndex = FirstSideSegment->GetChainIndex();
-		FEdgeSegment* Segment1 = FirstSideSegment;
-		while (Segment1->GetChainIndex() == FirstChainIndex && !Segment1->HasMarker1())
-		{
-			Segment1 = Segment1->GetPrevious();
-			if (Segment1 == FirstSideSegment)
-			{
-				Segment1 = Segment1->GetPrevious();
-				bFirstSideIsClosed = true;
-				break;
-			}
-		}
-		FirstSideSegment = Segment1->GetNext();
-
-		// Find the parallel chain
-		SecondSideSegment = FirstSideSegment->GetClosedSegment();
-		if (SecondSideSegment == nullptr)
-		{
-			SecondSideSegment = FirstSideSegment->GetNext()->GetClosedSegment();
-			if (SecondSideSegment == nullptr)
-			{
-				FirstSideSegment->HasMarker1();
-				continue;
-			}
-		}
-		FIdent SecondChainIndex = SecondSideSegment->GetChainIndex();
-		if (SecondChainIndex == 0) continue;
-
-		int32 SegmentNum = 0;
-		FEdgeSegment* Segment2 = SecondSideSegment;
-		if (!(bFirstSideIsClosed && (FirstChainIndex == SecondChainIndex)))
-		{
-			while (Segment2->GetChainIndex() == SecondChainIndex && !Segment2->HasMarker1() &&
-				((Segment2->GetClosedSegment() == FirstSideSegment) || (Segment2->GetClosedSegment() == FirstSideSegment->GetNext())
-					|| (FirstSideSegment->GetClosedSegment() == Segment2) || (FirstSideSegment->GetClosedSegment() == Segment2->GetPrevious())))
-			{
-				if (Segment2->HasMarker1())
-				{
-					break;
-				}
-
-				Segment2 = Segment2->GetNext();
-				if (SegmentNum > 0 && Segment2->GetPrevious() == SecondSideSegment)
-				{
-					bSecondSideIsClosed = true;
-					break;
-				}
-				SegmentNum++;
-			}
-			SecondSideSegment = Segment2->GetPrevious();
-		}
-
-		Segment1 = FirstSideSegment;
-		Segment2 = SecondSideSegment;
-
-		// both chains are browsed as long as they stay in front of each other
-		double MaxSquareThickness = 0.0;
-		do {
-			if (!FirstSideSegment->HasMarker1())
-			{
-				FirstSide.Add(FirstSideSegment);
-				FirstSideSegment->SetMarker1();
-				if (MaxSquareThickness <FirstSideSegment->GetClosedSquareDistance()) 
-				{
-					MaxSquareThickness = FirstSideSegment->GetClosedSquareDistance();
-				}
-			}
-
-			if (!SecondSideSegment->HasMarker1())
-			{
-				SecondSide.Add(SecondSideSegment);
-				SecondSideSegment->SetMarker1();
-				if (MaxSquareThickness <SecondSideSegment->GetClosedSquareDistance()) 
-				{
-					MaxSquareThickness = SecondSideSegment->GetClosedSquareDistance();
-				}
-			}
-
-			if (Segment1 == FirstSideSegment) 
-			{
-				Segment1 = FirstSideSegment->GetNext();
-			}
-			if (Segment2 == SecondSideSegment) 
-			{
-				Segment2 = SecondSideSegment->GetPrevious();
-			}
-
-			if ((!Segment1->HasMarker1()) && (Segment1->GetChainIndex() == FirstChainIndex) && ((Segment1->GetClosedSegment() == SecondSideSegment->GetNext()) || (Segment1->GetClosedSegment() == SecondSideSegment)))
-			{
-				FirstSideSegment = Segment1;
-			}
-			else if ((!Segment2->HasMarker1()) && (Segment2->GetChainIndex() == SecondChainIndex) && ((Segment2->GetClosedSegment() == FirstSideSegment->GetPrevious()) || (Segment2->GetClosedSegment() == FirstSideSegment)))
-			{
-				SecondSideSegment = Segment2;
-			}
-			else if ((!Segment1->HasMarker1()) && (Segment1->GetChainIndex() == FirstChainIndex) && (Segment1->GetClosedSegment() == SecondSideSegment->GetPrevious()))
-			{
-				FirstSideSegment = Segment1;
-			}
-			else if ((!Segment2->HasMarker1()) && (Segment2->GetChainIndex() == SecondChainIndex) && (Segment2->GetClosedSegment() == FirstSideSegment->GetNext()))
-			{
-				SecondSideSegment = Segment2;
-			}
-		} 
-		while ((FirstSideSegment == Segment1) || (SecondSideSegment == Segment2));
-
-		if ((FirstSide.Num() == 0) || (SecondSide.Num() == 0))
-		{
-			continue;
-		}
-
-		// Length of each side
-		double LengthSide1 = 0;
-		double MediumThickness1 = 0;
-		for (FEdgeSegment* Segment : FirstSide)
-		{
-			LengthSide1 += Segment->GetLength();
-			MediumThickness1 += Segment->GetClosedSquareDistance();
-		}
-		MediumThickness1 /= FirstSide.Num();
-		MediumThickness1 = sqrt(MediumThickness1);
-
-		double LengthSide2 = 0;
-		double MediumThickness2 = 0;
-		for (FEdgeSegment* Segment : SecondSide)
-		{
-			LengthSide2 += Segment->GetLength();
-			MediumThickness2 += Segment->GetClosedSquareDistance();
-		}
-		MediumThickness2 /= FirstSide.Num();
-		MediumThickness2 = sqrt(MediumThickness2);
-
-
-		if ((LengthSide1 <MediumThickness2 * 3) || (LengthSide2 <MediumThickness2 * 3))
-		{
-			continue;
-		}
-
-		const double MaxThickness = sqrt(MaxSquareThickness);
-		const double MaxLengthBetweendSideToBePeak = 3 * MaxThickness;
-
-		// make the thin zone
-		FThinZone2D& Zone = ThinZones.Emplace_GetRef(FirstSide, bFirstSideIsClosed, SecondSide, bSecondSideIsClosed, MaxThickness);
-
-		TFunction<void(TArray<TSharedPtr<FTopologicalEdge>>&)> DeleteCuttingPointOfPeak = [](const TArray<TSharedPtr<FTopologicalEdge>>& PeakEdges)
-		{
-			for (TSharedPtr<FTopologicalEdge> Edge : PeakEdges)
-			{
-				if (Edge->IsMeshed())
-				{
-					continue;
-				}
-				Edge->GetLinkActiveEdge()->SetThinPeak();
-			}
-		};
-
-		if (!Zone.GetFirstSide().IsInner() && !Zone.GetSecondSide().IsInner())
-		{
-			double LengthBetweenExtremity[2] = { HUGE_VALUE, HUGE_VALUE };
-			TArray<TSharedPtr<FTopologicalEdge>> PeakEdges[2];
-			GetThinZoneSideConnectionsLength(FirstSide, SecondSide, 10 * MaxThickness, LengthBetweenExtremity, PeakEdges);
-
-			if ((Zone.GetMaxLength() * 2 > (FirstLoopLength - MaxLengthBetweendSideToBePeak)) || (LengthBetweenExtremity[0] <MaxLengthBetweendSideToBePeak && LengthBetweenExtremity[1] <MaxLengthBetweendSideToBePeak))
-			{
-				Zone.SetCategory(EThinZone2DType::Global);
-				DeleteCuttingPointOfPeak(PeakEdges[0]);
-				DeleteCuttingPointOfPeak(PeakEdges[1]);
-			}
-			else if (LengthBetweenExtremity[0] <MaxLengthBetweendSideToBePeak)
-			{
-				Zone.SetCategory(EThinZone2DType::PeakStart);
-				if (Zone.GetFirstSide().GetLength() <MaxThickness * 10)
-				{
-					ThinZones.Pop();
-					continue;
-				}
-				DeleteCuttingPointOfPeak(PeakEdges[0]);
-			}
-			else if (LengthBetweenExtremity[1] <MaxLengthBetweendSideToBePeak)
-			{
-				Zone.SetCategory(EThinZone2DType::PeakEnd);
-				if (Zone.GetFirstSide().GetLength() <MaxThickness * 10)
-				{
-					ThinZones.Pop();
-					continue;
-				}
-				DeleteCuttingPointOfPeak(PeakEdges[1]);
-			}
-			else {
-				Zone.SetCategory(EThinZone2DType::Butterfly);
-			}
-		}
-		else
-		{
-			Zone.SetCategory(EThinZone2DType::BetweenLoops);
-		}
-
-		Zone.SetEdgesAsThinZone();
-	}
-}
-
-void FThinZone2DFinder::DisplayLoopSegments()
-{
-#ifdef DEBUG_THIN_ZONES
-	TSharedPtr<FTopologicalEdge> currentEdge = nullptr;
-	currentEdge = LoopSegments[0]->GetEdge();
-	F3DDebugSession _(TEXT("BoundarySegment"));
-
-	Open3DDebugSession(FString::Printf(TEXT("Edge %d\n"), currentEdge->GetId()));
-	for (FEdgeSegment* EdgeSegment : LoopSegments)
-	{
-		if (currentEdge != EdgeSegment->GetEdge())
-		{
-			currentEdge = EdgeSegment->GetEdge();
-			Close3DDebugSession();
-			Open3DDebugSession(FString::Printf(TEXT("Edge %d\n"), currentEdge->GetId()));
-		}
-		DisplaySegment(EdgeSegment->GetExtemity(ELimit::Start), EdgeSegment->GetExtemity(ELimit::End), EdgeSegment->GetId(), EVisuProperty::GreenCurve);
-	}
-	Close3DDebugSession();
-	//Wait();
-#endif
-}
-
-void FThinZone2DFinder::SearchThinZones()
-{
-#ifdef DEBUG_THIN_ZONES
-	{
-		F3DDebugSession _(TEXT("Thin Surface"));
-		Grid.DisplayPoints("Scaled Grid", Grid.GetInner2DPoints(EGridSpace::Scaled));
-	}
-#endif
-
-	//FMessage::Printf(DBG, TEXT("Searching thin zones on Surface %d\n", Surface->GetId());
-	FTimePoint StartTime = FChrono::Now();
-	BuildLoopSegments();
-	Chronos.BuildLoopSegmentsTime = FChrono::Elapse(StartTime);
-
-	DisplayLoopSegments();
-
-	StartTime = FChrono::Now();
-	FindClosedSegments();
-	Chronos.FindClosedSegmentTime = FChrono::Elapse(StartTime);
-	DisplayClosedSegments();
-
-	StartTime = FChrono::Now();
-	CheckClosedSegments();
-	Chronos.CheckClosedSegmentTime = FChrono::Elapse(StartTime);
-	DisplayClosedSegments();
-
-	StartTime = FChrono::Now();
-	LinkClosedSegments();
-	Chronos.LinkClosedSegmentTime = FChrono::Elapse(StartTime);
-
-	StartTime = FChrono::Now();
-	BuildThinZone();
-	Chronos.BuildThinZoneTime = FChrono::Elapse(StartTime);
-
-	if (ThinZones.Num() > 0)
-	{
-		Grid.GetFace().SetHasThinZone();
-	}
-
-#ifdef DISPLAY_THIN_ZONES
-	if (ThinZones.Num() > 0)
-	{
-		F3DDebugSession _(TEXT("Zones fines"));
-
-		int32 index = 0;
-		for (const FThinZone2D& Zone : ThinZones)
-		{
-			index++;
-			EVisuProperty prop = PropertyElement;
-			switch (Zone.GetCategory())
-			{
-			default:
-			case FThinZone2D::UNDEFINED:
-				prop = EVisuProperty::Iso;
-				Open3DDebugSession(TEXT("Zone UNDEFINED"));
-				break;
-			case FThinZone2D::GLOBAL:
-				prop = PROPERTY_BORDER_GEO;
-				Open3DDebugSession(TEXT("Zone GLOBAL"));
-				break;
-			case FThinZone2D::PEAK_START:
-				prop = EVisuProperty::BorderEdge;
-				Open3DDebugSession(TEXT("Zone PEAK start"));
-				break;
-			case FThinZone2D::PEAK_END:
-				prop = EVisuProperty::BorderEdge;
-				Open3DDebugSession(TEXT("Zone PEAK end"));
-				break;
-			case FThinZone2D::BUTTERFLY:
-				prop = PropertyElement;
-				Open3DDebugSession(TEXT("Zone BUTTERFLY"));
-				break;
-			case FThinZone2D::BETWEEN_CONTOUR:
-				prop = EVisuProperty::NonManifoldEdge;
-				Open3DDebugSession(TEXT("Zone BETWEEN_CONTOUR"));
-				break;
-			default :
-				Open3DDebugSession(TEXT("Zone Unknown"));
-			}
-
-			{
-				F3DDebugSession _(TEXT("Side1"));
-				int32 ind = 0;
-				for (FEdgeSegment* EdgeSegment : Zone.GetFirstSide().GetSegments())
-				{
-					DisplaySegment(EdgeSegment->GetExtemity(ELimit::Start)(), EdgeSegment->GetExtemity(ELimit::End)(), ++ind, prop);
-				}
-			}
-			{
-				F3DDebugSession _(TEXT("Side2"));
-				int32 ind = 0;
-				for (FEdgeSegment* EdgeSegment : Zone.GetSecondSide().GetSegments())
-				{
-					DisplaySegment(EdgeSegment->GetExtemity(ELimit::Start)(), EdgeSegment->GetExtemity(ELimit::End)(), ++ind, prop);
-				}
-			}
-
-			Close3DDebugSession();
-		}
-	}
-#endif
-
-}
-void FThinZone2DFinder::BuildLoopSegments()
-{
-	double Length = 0;
-	double SegmentLength = Tolerance / 5.;
-
-	const TArray<TSharedPtr<FTopologicalLoop>>& Loops = Grid.GetFace().GetLoops();
-
-	FirstLoopLength = -1.;
-	for (const TSharedPtr<FTopologicalLoop>& Loop : Loops)
-	{
-		for (const FOrientedEdge& Edge : Loop->GetEdges())
-		{
-			Length += Edge.Entity->Length();
-		}
-
-		if (FirstLoopLength < 0)
-		{
-			FirstLoopLength = Length;
-		}
-	}
-
-	double SegmentNum = 1.2 * Length / SegmentLength;
-	LoopSegments.Empty((int32)SegmentNum);
-
-	TSharedPtr<FTopologicalLoop> OuterLoop = Loops[0];
-	for (TSharedPtr<FTopologicalLoop> Loop : Loops)
-	{
-		const TArray<FOrientedEdge>& Edges = Loop->GetEdges();
-
-		FEdgeSegment* FirstSegment = nullptr;
-		FEdgeSegment* PrecedingSegment = nullptr;
-
-		for (const FOrientedEdge& Edge : Edges)
-		{
-			TArray<double> Coordinates;
-			Edge.Entity->Sample(SegmentLength, Coordinates);
-
-			TArray<FPoint2D> Points;
-			Edge.Entity->Approximate2DPoints(Coordinates, Points);
-
-			TArray<FPoint2D> ScaledPoints;
-			Grid.TransformPoints(EGridSpace::Scaled, Points, ScaledPoints);
-
-			int32 PointNum = Points.Num();
-			for (int32 ISegment = 0; ISegment < PointNum - 1; ISegment++)
-			{
-				int32 Index1 = (Edge.Direction == EOrientation::Front) ? ISegment : PointNum - ISegment - 1l;
-				int32 Index2 = (Edge.Direction == EOrientation::Front) ? ISegment + 1l : PointNum - ISegment - 2l;
-				FEdgeSegment& CurrentSeg = SegmentFatory.New();
-				CurrentSeg.SetBoundarySegment(OuterLoop != Loop, Edge.Entity, Coordinates[Index1], Coordinates[Index2], ScaledPoints[Index1], ScaledPoints[Index2]);
-
-				LoopSegments.Add(&CurrentSeg);
-				if (!FirstSegment)
-				{
-					FirstSegment = PrecedingSegment = &CurrentSeg;
-				}
-				else
-				{
-					PrecedingSegment->SetNext(&CurrentSeg);
-				}
-				PrecedingSegment = &CurrentSeg;
-			}
-		}
-		PrecedingSegment->SetNext(FirstSegment);
-	}
-	SortedLoopSegments = LoopSegments;
-}
-
-FThinZoneSide::FThinZoneSide(FThinZoneSide* InFrontSide, const TArray<FEdgeSegment*>& InSegments, bool bInIsFirstSide)
-	: FrontSide(*InFrontSide)
-{
-	Segments = InSegments;
-	Length = 0;
-	for (const FEdgeSegment* Segment : Segments)
-	{
-		Length += Segment->GetLength();
-	}
-}
-
-void FThinZoneSide::SetEdgesAsThinZone()
-{
-	for (FEdgeSegment* Segment : Segments)
-	{
-		Segment->GetEdge()->SetThinZone();
-	}
-}
-
-bool FThinZoneSide::IsPartiallyMeshed() const
-{
-	TSharedPtr<FTopologicalEdge> Edge = nullptr;
-	for (FEdgeSegment* EdgeSegment : Segments)
-	{
-		if (Edge != EdgeSegment->GetEdge())
-		{
-			Edge = EdgeSegment->GetEdge();
-			if (Edge->GetLinkActiveEdge()->IsMeshed())
-			{
-				return true;
-			}
-		}
-	}
-	return false;
-}
-
-double FThinZoneSide::GetMeshedLength() const
-{
-	double LocalLength = 0;
-	TSharedPtr<FTopologicalEdge> Edge = nullptr;
-	for (FEdgeSegment* EdgeSegment : Segments)
-	{
-		if (Edge != EdgeSegment->GetEdge())
-		{
-			Edge = EdgeSegment->GetEdge();
-			if (Edge->GetLinkActiveEdge()->IsMeshed())
-			{
-				LocalLength += Edge->Length();
-			}
-		}
-	}
-	return LocalLength;
-}
-
-void FThinZone2D::SetEdgesAsThinZone()
-{
-	FirstSide.SetEdgesAsThinZone();
-	SecondSide.SetEdgesAsThinZone();
-}
-
-}
-
diff --git a/Engine/Source/Runtime/Datasmith/CADKernel/Private/CADKernel/Mesh/Structure/ThinZone2D.h b/Engine/Source/Runtime/Datasmith/CADKernel/Private/CADKernel/Mesh/Structure/ThinZone2D.h
--- a/Engine/Source/Runtime/Datasmith/CADKernel/Private/CADKernel/Mesh/Structure/ThinZone2D.h
+++ b/Engine/Source/Runtime/Datasmith/CADKernel/Private/CADKernel/Mesh/Structure/ThinZone2D.h
@@ -1,303 +1,3 @@
 // Copyright Epic Games, Inc. All Rights Reserved.
 #pragma once
 
-#include "CADKernel/Core/Chrono.h"
-#include "CADKernel/Core/HaveStates.h"
-#include "CADKernel/Core/Factory.h"
-#include "CADKernel/Core/Types.h"
-#include "CADKernel/Math/Point.h"
-#include "CADKernel/Mesh/Structure/EdgeSegment.h"
-
-//#define DEBUG_THIN_ZONES
-//#define DISPLAY_THIN_ZONES // to display found thin zone
-
-//#define DEBUG_BORDER_THIN_ZONES
-//#define DEBUG_REMOVE_THIN_ZONE
-//#define DEBUG_REMOVABLE
-//#define DEBUG_BEST_VERTEX
-
-namespace UE::CADKernel
-{
-	class FTopologicalLoop;
-	class FTopologicalEdge;
-	class FGrid;
-	class FThinZone2D;
-	class FThinZone2DFinder;
-
-	struct FThinZoneChronos
-	{
-		FDuration BuildLoopSegmentsTime;
-		FDuration DisplayBoundarySegmentTime;
-		FDuration FindClosedSegmentTime;
-		FDuration CheckClosedSegmentTime;
-		FDuration LinkClosedSegmentTime;
-		FDuration BuildThinZoneTime;
-	};
-
-	enum class EThinZone2DType : uint8
-	{
-		Undefined = 0,
-		Global,			// a Surface globally thin
-		PeakStart,		// an extremity of a Surface that is fine
-		PeakEnd,		// an extremity of a Surface that is fine
-		Butterfly,		// the outter loop that is like a bow tie (Butterfly)
-		BetweenLoops    // a bow tie between two different loops
-	};
-
-	class FThinZoneSide : public FHaveStates
-	{
-	private:
-		TArray<FEdgeSegment*> Segments;
-		FThinZoneSide& FrontSide;
-		double Length;
-
-	public:
-
-		FThinZoneSide(FThinZoneSide* InFrontSide, const TArray<FEdgeSegment*>& InSegments, bool bInIsFirstSide);
-		virtual ~FThinZoneSide() = default;
-
-		bool IsClosed() const
-		{
-			if (GetFirst() == nullptr)
-			{
-				return false;
-			}
-			return GetFirst()->GetPrevious() == GetLast();
-		}
-
-		FEdgeSegment* GetFirst() const
-		{
-			return Segments.Num() > 0 ? Segments[0] : nullptr;
-		}
-
-		FEdgeSegment* GetLast() const
-		{
-			return Segments.Num() > 0 ? Segments.Last() : nullptr;
-		}
-
-		void SetEdgesAsThinZone();
-
-		const TArray<FEdgeSegment*>& GetSegments() const
-		{
-			return Segments;
-		}
-
-		bool IsPartiallyMeshed() const;
-		double GetMeshedLength() const;
-
-		double GetLength() const
-		{
-			return Length;
-		}
-
-		bool IsInner() const
-		{
-			if (Segments.Num() == 0)
-			{
-				return true;
-			}
-			return Segments[0]->IsInner();
-		}
-
-		bool IsFirstSide() const
-		{
-			return ((States & EHaveStates::IsFirstSide) == EHaveStates::IsFirstSide);
-		}
-
-		void SetFirstSide() const
-		{
-			States |= EHaveStates::IsFirstSide;
-		}
-
-		void ResetFirstSide() const
-		{
-			States &= ~EHaveStates::IsFirstSide;
-		}
-	};
-
-	class FThinZone2D : public FHaveStates
-	{
-	private:
-
-		FThinZoneSide FirstSide;
-		FThinZoneSide SecondSide;
-
-		EThinZone2DType Category;
-
-		TArray<TSharedPtr<FTopologicalEdge>> PeakEdges[2];
-
-		double Thickness;
-
-	public:
-
-		FThinZone2D(const TArray<FEdgeSegment*>& InFirstSideSegments, bool bInIsClosed1, const TArray<FEdgeSegment*>& InSecondSideSegments, bool bInIsClosed2, double ZoneThickness)
-			: FirstSide(&SecondSide, InFirstSideSegments, true)
-			, SecondSide(&FirstSide, InSecondSideSegments, false)
-			, Category(EThinZone2DType::Undefined)
-			, Thickness(ZoneThickness)
-		{
-			if (bInIsClosed1)
-			{
-				SetFirstSideClosed();
-			}
-
-			if (bInIsClosed2)
-			{
-				SetFirstSideClosed();
-			}
-		}
-
-		virtual ~FThinZone2D() = default;
-
-		const double GetThickness()
-		{
-			return Thickness;
-		};
-
-		const TArray<TSharedPtr<FTopologicalEdge>>* GetPeakEdges() const
-		{
-			return PeakEdges;
-		}
-
-		TArray<TSharedPtr<FTopologicalEdge>>* GetPeakEdges()
-		{
-			return PeakEdges;
-		}
-
-		const FThinZoneSide& GetFirstSide() const
-		{
-			return FirstSide;
-		}
-
-		const FThinZoneSide& GetSecondSide() const
-		{
-			return SecondSide;
-		}
-
-		EThinZone2DType GetCategory() const
-		{
-			return Category;
-		}
-
-		void SetCategory(EThinZone2DType InType)
-		{
-			Category = InType;
-		}
-
-
-		void SetEdgesAsThinZone();
-
-		double GetMaxLength() const
-		{
-			return FMath::Max(FirstSide.GetLength(), SecondSide.GetLength());
-		}
-
-		bool IsRemoved() const
-		{
-			return ((States & EHaveStates::IsRemoved) == EHaveStates::IsRemoved);
-		}
-
-		void SetRemoved() const
-		{
-			States |= EHaveStates::IsRemoved;
-		}
-
-		void ResetRemoved() const
-		{
-			States &= ~EHaveStates::IsRemoved;
-		}
-
-		bool IsFirstSideClosed() const
-		{
-			return ((States & EHaveStates::FirstSideClosed) == EHaveStates::FirstSideClosed);
-		}
-
-		void SetFirstSideClosed() const
-		{
-			States |= EHaveStates::FirstSideClosed;
-		}
-
-		void ResetFirstSideClosed() const
-		{
-			States &= ~EHaveStates::FirstSideClosed;
-		}
-
-		bool IsSecondSideClosed() const
-		{
-			return ((States & EHaveStates::SecondSideClosed) == EHaveStates::SecondSideClosed);
-		}
-
-		void SetSecondSideClosed() const
-		{
-			States |= EHaveStates::SecondSideClosed;
-		}
-
-		void ResetSecondSideClosed() const
-		{
-			States &= ~EHaveStates::SecondSideClosed;
-		}
-
-		bool HasClosedSide() const 
-		{
-			constexpr EHaveStates ATLeastOneClosed = EHaveStates::FirstSideClosed | EHaveStates::SecondSideClosed;
-			return EnumHasAnyFlags(States, ATLeastOneClosed);
-		}
-
-	};
-
-	class FThinZone2DFinder
-	{
-		friend class FThinZone2D;
-
-	public:
-		FThinZoneChronos Chronos;
-
-	protected:
-		FGrid& Grid;
-
-		double Tolerance;
-		double SquareTolerance;
-		double MinFeatureSize;
-
-		TArray<FThinZone2D> ThinZones;
-
-		TArray<FEdgeSegment*> LoopSegments;
-		TArray<FEdgeSegment*> SortedLoopSegments;
-
-		double FirstLoopLength = 0.0;
-
-		TFactory<FEdgeSegment> SegmentFatory;
-
-	public:
-
-		FThinZone2DFinder(FGrid& InGrid)
-			: Grid(InGrid)
-			, SegmentFatory()
-		{
-		}
-
-		void Set(double InTolerance)
-		{
-			Tolerance = InTolerance;
-			SquareTolerance = FMath::Square(Tolerance);
-		}
-
-		void SearchThinZones();
-
-		const TArray<FThinZone2D>& GetThinZones() const
-		{
-			return ThinZones;
-		}
-
-	private:
-
-		void BuildLoopSegments();
-		void FindClosedSegments();
-		void CheckClosedSegments();
-		void LinkClosedSegments();
-		void BuildThinZone();
-
-		void DisplayClosedSegments();
-		void DisplayLoopSegments();
-	};
-}
diff --git a/Engine/Source/Runtime/Datasmith/CADKernel/Private/CADKernel/UI/Display.cpp b/Engine/Source/Runtime/Datasmith/CADKernel/Private/CADKernel/UI/Display.cpp
--- a/Engine/Source/Runtime/Datasmith/CADKernel/Private/CADKernel/UI/Display.cpp
+++ b/Engine/Source/Runtime/Datasmith/CADKernel/Private/CADKernel/UI/Display.cpp
@@ -379,7 +379,7 @@ void Display(const FPlane& Plane, FIdent Ident)
 	TArray<FPoint> Points;
 
 	const FVector Point = Plane.GetOrigin();
-	
+
 	FVector VTemp;
 
 	VTemp = Point + UAxis + VAxis;
@@ -655,17 +655,19 @@ void DisplayControlPolygon(const FCurve& Curve)
 	};
 
 	F3DDebugSegment GraphicSegment(Curve.GetId());
-	if (Curve.GetCurveType() == ECurve::Bezier)
+	switch (Curve.GetCurveType())
+	{
+	case  ECurve::Bezier:
 	{
 		const FBezierCurve& Bezier = (const FBezierCurve&)Curve;
 		DisplayHull(Bezier.GetPoles());
 		return;
 	}
-	if (Curve.GetCurveType() == ECurve::Nurbs)
+	case ECurve::Nurbs:
 	{
 		const FNURBSCurve& Nurbs = (const FNURBSCurve&)Curve;
 		const TArray<FPoint>& Poles = Nurbs.GetPoles();
-		DisplayHull(Nurbs.GetPoles());
+		DisplayHull(Poles);
 		return;
 	}
 #endif
@@ -733,6 +735,14 @@ void Display2D(const FTopologicalFace& Face)
 #endif
 }
 
+void Display2DWithScale(const FTopologicalFace& Face)
+{
+#ifdef CADKERNEL_DEV
+	F3DDebugSegment GraphicSegment(Face.GetId());
+	Draw2D(Face);
+#endif
+}
+
 void Draw(const FTopologicalFace& Face)
 {
 #ifdef CADKERNEL_DEV
@@ -745,7 +755,7 @@ void Draw(const FTopologicalFace& Face)
 				switch (Edge.Entity->GetTwinEntityCount())
 				{
 				case 1:
-					Property =  Edge.Entity->IsDegenerated() ? EVisuProperty::OrangeCurve : EVisuProperty::BorderEdge;
+					Property = Edge.Entity->IsDegenerated() ? EVisuProperty::OrangeCurve : EVisuProperty::BorderEdge;
 					break;
 				case 2:
 					Property = EVisuProperty::BlueCurve;
@@ -822,7 +832,6 @@ void Draw2D(const FTopologicalFace& Face)
 #endif
 }
 
-
 void DrawIsoCurves(const FTopologicalFace& Face)
 {
 #ifdef CADKERNEL_DEV
@@ -909,6 +918,21 @@ void Display2D(const FTopologicalEdge& Edge, EVisuProperty Property)
 #endif
 }
 
+void Display2DWithScale(const FTopologicalEdge& Edge, EVisuProperty Property)
+{
+#ifdef CADKERNEL_DEV
+	F3DDebugSegment GraphicSegment(Edge.GetId());
+	TArray<FPoint2D> Polyline;
+	Edge.GetCurve()->GetDiscretizationPoints(Edge.GetBoundary(), EOrientation::Front, Polyline);
+	DisplayPolylineWithScale(Polyline, Property);
+
+	for (const FPoint2D& Point : Polyline)
+	{
+		DisplayPoint2DWithScale(Point, EVisuProperty::RedPoint);
+	}
+#endif
+}
+
 void Draw(const FTopologicalEdge& Edge, EVisuProperty Property)
 {
 #ifdef CADKERNEL_DEV
@@ -1227,4 +1251,14 @@ void Display2D(const FTopologicalLoop& Loop)
 #endif
 }
 
+void Display2DWithScale(const FTopologicalLoop& Loop)
+{
+#ifdef CADKERNEL_DEV
+	for (const FOrientedEdge& Edge : Loop.GetEdges())
+	{
+		Display2DWithScale(*Edge.Entity);
+	}
+#endif
+}
+
 } // namespace UE::CADKernel
diff --git a/Engine/Source/Runtime/Datasmith/CADKernel/Public/CADKernel/Core/Types.h b/Engine/Source/Runtime/Datasmith/CADKernel/Public/CADKernel/Core/Types.h
--- a/Engine/Source/Runtime/Datasmith/CADKernel/Public/CADKernel/Core/Types.h
+++ b/Engine/Source/Runtime/Datasmith/CADKernel/Public/CADKernel/Core/Types.h
@@ -20,6 +20,12 @@ typedef uint32 FIdent;
 
 namespace UE::CADKernel
 {
+
+namespace Ident
+{
+static const FIdent Undefined = (FIdent)-1;
+}
+
 enum class EValue : uint8
 {
 	Entity,
diff --git a/Engine/Source/Runtime/Datasmith/CADKernel/Public/CADKernel/Math/SlopeUtils.h b/Engine/Source/Runtime/Datasmith/CADKernel/Public/CADKernel/Math/SlopeUtils.h
new file mode 100644
--- /dev/null
+++ b/Engine/Source/Runtime/Datasmith/CADKernel/Public/CADKernel/Math/SlopeUtils.h
@@ -0,0 +1,384 @@
+// Copyright Epic Games, Inc. All Rights Reserved.
+
+#pragma once
+
+#include "CADKernel/Core/Types.h"
+#include "CADKernel/Math/Point.h"
+#include "Algo/AllOf.h"
+
+namespace UE::CADKernel
+{
+
+namespace Slope
+{
+
+constexpr double NullSlope = 0.;
+constexpr double FullSlope = 8.;
+
+/**
+ * RightSlope i.e. Right angle i.e Pi / 2
+ */
+constexpr double RightSlope = 2.;
+
+/**
+ * ThreeRightSlope i.e. 3Pi / 2
+ */
+constexpr double ThreeRightSlope = 6.;
+
+/**
+ * MinusRightSlope i.e. -Pi / 2
+ */
+constexpr double MinusRightSlope = -2.;
+
+/**
+ * PiSlope i.e. Pi angle
+ */
+constexpr double PiSlope = 4.;
+}
+
+typedef TFunction<double(const FPoint2D&, const FPoint2D&, double)> SlopeMethod;
+
+
+/**
+ * Transform a positive slope into an oriented slope [-4, 4] i.e. an equivalent angle between [-Pi, Pi]
+ * @return a slope between [-4, 4]
+ */
+inline double TransformIntoOrientedSlope(double Slope)
+{
+	if (Slope > 4.) Slope -= 8;
+	if (Slope < -4.) Slope += 8;
+	return Slope;
+}
+
+/**
+ * Transform a slope into a positive slope [0, *] i.e. an equivalent angle between [0, 2.Pi]
+ * @return a slope between [0, 8]
+ */
+inline double TransformIntoPositiveSlope(double Slope)
+{
+	if (Slope < 0) Slope += 8;
+	return Slope;
+}
+
+/**
+ * Fast angle approximation by a "slope" :
+ * This method is used to compute an approximation of the angle between the input segment defined by two points and [0, u) axis.
+ * The return value is a real in the interval [0, 8] for an angle in the interval [0, 2Pi]
+ * Warning, it's only an approximation... The conversion is not linear but the error is small near the integer value of slope (0, 1, 2, 3, ...8)
+ *
+ * This approximation is very good when only comparison of angles is needed.
+ * This method is not adapted to compute angle
+ *
+ * To compute an angle value between two segments, the call of acos (and asin for an oriented angle) is necessary while with this approximation, only a division is useful.
+ *
+ * [0 - 2Pi] is divide into 8 angular sector i.e. [0, Pi/4] = [0,1], [Pi/4, Pi/2] = [1,2], ...
+ *
+ * @return a slope between [0, 8] i.e. an equivalent angle between [0, 2Pi]
+ *
+ * Angle (Degree) to Slop
+ *		  0   = 0
+ *		  7.  = 0.125
+ *		 14.  = 0.25
+ *		 30   = 0.5
+ *		 36.8 = 0.75
+ *		 45   = 1
+ *		 53.2 = 1.25
+ *		 60   = 1.5
+ *		 76.  = 1.75
+ *		 90   = 2
+ *		120   = 2.5
+ *		135   = 3
+ *		180   = 4
+ *		360   = 8
+ */
+inline double ComputeSlope(const FPoint2D& StartPoint, const FPoint2D& EndPoint)
+{
+	double DeltaU = EndPoint.U - StartPoint.U;
+	double DeltaV = EndPoint.V - StartPoint.V;
+	double Delta;
+
+	if (FMath::Abs(DeltaU) < DOUBLE_SMALL_NUMBER && FMath::Abs(DeltaV) < DOUBLE_SMALL_NUMBER)
+	{
+		return 0;
+	}
+
+	if (DeltaU > DOUBLE_SMALL_NUMBER)
+	{
+		if (DeltaV > DOUBLE_SMALL_NUMBER)
+		{
+			if (DeltaU > DeltaV)
+			{
+				Delta = DeltaV / DeltaU;
+			}
+			else
+			{
+				Delta = 2 - DeltaU / DeltaV;
+			}
+		}
+		else
+		{
+			if (DeltaU > -DeltaV)
+			{
+				Delta = 8 + DeltaV / DeltaU;
+			}
+			else
+			{
+				Delta = 6 - DeltaU / DeltaV; // deltaU/deltaV <0
+			}
+		}
+	}
+	else
+	{
+		if (DeltaV > DOUBLE_SMALL_NUMBER)
+		{
+			if (-DeltaU > DeltaV)
+			{
+				Delta = 4 + DeltaV / DeltaU;
+			}
+			else
+			{
+				Delta = 2 - DeltaU / DeltaV;
+			}
+		}
+		else
+		{
+			if (-DeltaU > -DeltaV)
+			{
+				Delta = 4 + DeltaV / DeltaU;
+			}
+			else
+			{
+				Delta = 6 - DeltaU / DeltaV;
+			}
+		}
+	}
+
+	return Delta;
+}
+
+/**
+ * Compute the oriented slope of a segment according to a reference slope
+ * This method is used to compute an approximation of the angle between two segments in 2D.
+ * return a slope between [0, 8] i.e. an equivalent angle between [0, 2Pi]
+ */
+inline double ComputePositiveSlope(const FPoint2D& StartPoint, const FPoint2D& EndPoint, double ReferenceSlope)
+{
+	double Slope = ComputeSlope(StartPoint, EndPoint);
+	Slope -= ReferenceSlope;
+	if (Slope < 0) Slope += 8;
+	return Slope;
+}
+
+/**
+ * Compute the positive slope between the segments [StartPoint, EndPoint1] and [StartPoint, EndPoint2]
+ * @return a slope between [0, 8] i.e. an equivalent angle between [0, 2Pi]
+ */
+inline double ComputePositiveSlope(const FPoint2D& StartPoint, const FPoint2D& EndPoint1, const FPoint2D& EndPoint2)
+{
+	double ReferenceSlope = ComputeSlope(StartPoint, EndPoint1);
+	double Slope = ComputeSlope(StartPoint, EndPoint2);
+	Slope -= ReferenceSlope;
+	return TransformIntoPositiveSlope(Slope);
+}
+
+inline double ClockwiseSlope(const FPoint2D& StartPoint, const FPoint2D& EndPoint, double ReferenceSlope)
+{
+	return 8. - ComputePositiveSlope(StartPoint, EndPoint, ReferenceSlope);
+}
+
+inline double CounterClockwiseSlope(const FPoint2D& StartPoint, const FPoint2D& EndPoint, double ReferenceSlope)
+{
+	return ComputePositiveSlope(StartPoint, EndPoint, ReferenceSlope);
+}
+
+/**
+ * Compute the oriented slope of a segment according to a reference slope
+ * @return a slope between [-4, 4] i.e. an equivalent angle between [-Pi, Pi]
+ */
+inline double ComputeOrientedSlope(const FPoint2D& StartPoint, const FPoint2D& EndPoint, double ReferenceSlope)
+{
+	return TransformIntoOrientedSlope(ComputePositiveSlope(StartPoint, EndPoint, ReferenceSlope));
+}
+
+/**
+ * Compute the positive slope between the segments [StartPoint, EndPoint1] and [StartPoint, EndPoint2]
+ * @return a slope between [-4, 4] i.e. an equivalent angle between [-Pi, Pi]
+ */
+inline double ComputeOrientedSlope(const FPoint2D& StartPoint, const FPoint2D& EndPoint1, const FPoint2D& EndPoint2)
+{
+	return TransformIntoOrientedSlope(ComputePositiveSlope(StartPoint, EndPoint1, EndPoint2));
+}
+
+/**
+ * return a slope between [0, 4] i.e. an angle between [0, Pi]
+ */
+inline double ComputeUnorientedSlope(const FPoint2D& StartPoint, const FPoint2D& EndPoint, double ReferenceSlope)
+{
+	return FMath::Abs(ComputeOrientedSlope(StartPoint, EndPoint, ReferenceSlope));
+}
+
+/**
+ * return a slope between [0, 1] relative to the nearest axis between horizontal or vertical axis i.e.
+ * ComputeUnorientedSlope => 0.5 return 0.5
+ * ComputeUnorientedSlope => 2.3 return 0.3
+ * ComputeUnorientedSlope => 3.6 return 0.4
+ */
+inline double ComputeSlopeRelativeToNearestAxis(const FPoint2D& StartPoint, const FPoint2D& EndPoint)
+{
+	double Slope = FMath::Abs(TransformIntoOrientedSlope(ComputeSlope(StartPoint, EndPoint)));
+	if (Slope > 2)
+	{
+		Slope = 4 - Slope;
+	}
+
+	// if slope close to 2 means segment close to IsoU, otherwise segment close to IsoV
+	// Wants a slope between 0 and 1 to manage either IsoU and IsoV
+	// Close to 0 means close to IsoU or IsoV
+	if (Slope > 1)
+	{
+		Slope = 2 - Slope;
+	}
+
+	return Slope;
+}
+
+/**
+ * return a slope between [0, 2] relative to reference Axis i.e.
+ * ComputeUnorientedSlope => 0.5 return 0.5
+ * ComputeUnorientedSlope => 2.3 return 1.7
+ */
+inline double ComputeSlopeRelativeToReferenceAxis(const FPoint2D& StartPoint, const FPoint2D& EndPoint, double ReferenceAxisSlope)
+{
+	double Slope = ComputeUnorientedSlope(StartPoint, EndPoint, ReferenceAxisSlope);
+	if (Slope > 2)
+	{
+		Slope = 4 - Slope;
+	}
+
+	return Slope;
+}
+
+/**
+ * return a slope between [0, 4] i.e. an angle between [0, Pi]
+ */
+inline double ComputeUnorientedSlope(const FPoint2D& StartPoint, const FPoint2D& EndPoint1, const FPoint2D& EndPoint2)
+{
+	return FMath::Abs(ComputeOrientedSlope(StartPoint, EndPoint1, EndPoint2));
+}
+
+/**
+ *                         P1
+ *          inside        /
+ *                       /   inside
+ *                      /
+ *    A -------------- B --------------- C
+ *                      \
+ *           Outside     \  Outside
+ *                        \
+ *                         P2
+ *
+ * Return true if the segment BP is inside the sector defined the half-lines [BA) and [BC) in the counterclockwise.
+ * Return false if ABP angle or PBC angle is too flat (smaller than FlatAngle)
+ */
+inline bool IsPointPInsideSectorABC(const FPoint2D& PointA, const FPoint2D& PointB, const FPoint2D& PointC, const FPoint2D& PointP, const double FlatAngle)
+{
+	double SlopWithNextBoundary = ComputeSlope(PointB, PointC);
+	double BoundaryDeltaSlope = ComputePositiveSlope(PointB, PointA, SlopWithNextBoundary);
+	double SegmentSlope = ComputePositiveSlope(PointB, PointP, SlopWithNextBoundary);
+	if (SegmentSlope < FlatAngle || SegmentSlope + FlatAngle > BoundaryDeltaSlope)
+	{
+		return false;
+	}
+	return true;
+}
+
+/**
+ *                         P1
+ *          inside        /
+ *                       /   inside
+ *                      /
+ *    A -------------- B --------------- C
+ *                      \
+ *           Outside     \  Outside
+ *                        \
+ *                         P2
+ *
+ * Return true if all of the segment BPi is inside the sector defined the half-lines [BA) and [BC) in the counterclockwise.
+ */
+inline bool ArePointsInsideSectorABC(const FPoint2D& PointA, const FPoint2D& PointB, const FPoint2D& PointC, const TArray<const FPoint2D*>& Points, const double FlatAngle = -DOUBLE_SMALL_NUMBER)
+{
+	double SlopWithNextBoundary = ComputeSlope(PointB, PointC);
+	double BoundaryDeltaSlope = ComputePositiveSlope(PointB, PointA, SlopWithNextBoundary);
+
+	return Algo::AllOf(Points, [&](const FPoint2D* PointP) {
+		double DeltaU = PointB.U - PointP->U;
+		double DeltaV = PointB.V - PointP->V;
+
+		if (FMath::Abs(DeltaU) < DOUBLE_SMALL_NUMBER && FMath::Abs(DeltaV) < DOUBLE_SMALL_NUMBER)
+		{
+			return true;
+		}
+
+		double SegmentSlope = ComputePositiveSlope(PointB, *PointP, SlopWithNextBoundary);
+		if (SegmentSlope < FlatAngle || SegmentSlope + FlatAngle > BoundaryDeltaSlope)
+		{
+			return false;
+		}
+		return true;
+		});
+}
+
+inline FPoint2D SlopeToVector(const double Slope)
+{
+	int32 SlopeStep = (int32)(Slope);
+	FPoint2D Vector;
+	switch (SlopeStep)
+	{
+	case 0:
+		// Delta = DeltaV / DeltaU;
+		Vector[0] = 1.;
+		Vector[1] = Slope;
+		break;
+	case 1:
+		// 2 - DeltaU / DeltaV;
+		Vector[0] = 2. - Slope;
+		Vector[1] = 1.;
+		break;
+	case 2:
+		// 2 - DeltaU / DeltaV;
+		Vector[0] = 2. - Slope;
+		Vector[1] = 1.;
+		break;
+	case 3:
+		// 4 + DeltaV / DeltaU;
+		Vector[0] = -1.;
+		Vector[1] = 4. - Slope;
+		break;
+	case 4:
+		// 4 + DeltaV / DeltaU;
+		Vector[0] = -1.;
+		Vector[1] = 4. - Slope;
+		break;
+	case 5:
+		// 6 - DeltaU / DeltaV;
+		Vector[0] = Slope - 6.;
+		Vector[1] = -1.;
+		break;
+	case 6:
+		// 6 - DeltaU / DeltaV // deltaU/deltaV <0
+		Vector[0] = Slope - 6.;
+		Vector[1] = -1;
+		break;
+	case 7:
+		// 8 + DeltaV / DeltaU; // deltaU/deltaV <0
+		Vector[0] = 1.;
+		Vector[1] = Slope - 8.;
+		break;
+	default:
+		break;
+	}
+
+	return Vector;
+}
+
+} // namespace UE::CADKernel	
diff --git a/Engine/Source/Runtime/Datasmith/CADKernel/Public/CADKernel/Mesh/MeshEnum.h b/Engine/Source/Runtime/Datasmith/CADKernel/Public/CADKernel/Mesh/MeshEnum.h
new file mode 100644
--- /dev/null
+++ b/Engine/Source/Runtime/Datasmith/CADKernel/Public/CADKernel/Mesh/MeshEnum.h
@@ -0,0 +1,52 @@
+// Copyright Epic Games, Inc. All Rights Reserved.
+#pragma once
+
+#include "CADKernel/Core/Types.h"
+
+namespace UE::CADKernel
+{
+/**
+ * 3 parametric spaces are defined :
+ *  -EGridSpace::Default2D the natural parametric space of the carrier surface i.e. for a Bezier surface the boundary of the space is [[0, 1], [0, 1]] and for a spherical surface [[0, 2Pi], [-Pi / 2, Pi / 2]]
+ *  -EGridSpace::UniformScaled.The average curvilinear length of the iso curves is calculated in each axis.The parametric space is then scaled in each axis such that
+ *            the average curvilinear length between two points P(A, V) and P(B, V) = B - A.For a spherical surface [[0, 2Pi.R / 2], [-Pi.R / 2, Pi.R / 2]] with R the radius of the sphere
+ * -EGridSpace::Scaled.The UnifomeScaled space is rescaled in its most disturbe axis e.g. if V is the most disturbe axis, whatever V, length of[P(A, V), P(B, V)] = (B - A).
+ *            This parametric space is closed to a good flattening of the surface.
+ */
+enum EGridSpace : uint8
+{
+	Default2D = 0,  // do not modified enum value as it is used for a static array
+	Scaled = 1,
+	UniformScaled = 2,
+	EndGridSpace
+};
+
+enum EMeshingState : uint8
+{
+	NotMeshed = 0,
+	PartiallyMeshed,
+	FullyMeshed
+};
+
+namespace GridSpace
+{
+
+inline FString ToString(EGridSpace Space)
+{
+	switch (Space)
+	{
+	case Default2D:
+		return FString(TEXT("Default2D"));
+	case Scaled:
+		return FString(TEXT("Scaled"));
+	case UniformScaled:
+		return FString(TEXT("UniformScaled"));
+	default:
+		return FString(TEXT(""));
+	}
+}
+
+} // namespace GridSpace
+
+} // namespace UE::CADKernel
+
diff --git a/Engine/Source/Runtime/Datasmith/CADKernel/Public/CADKernel/Mesh/Structure/EdgeSegment.h b/Engine/Source/Runtime/Datasmith/CADKernel/Public/CADKernel/Mesh/Structure/EdgeSegment.h
new file mode 100644
--- /dev/null
+++ b/Engine/Source/Runtime/Datasmith/CADKernel/Public/CADKernel/Mesh/Structure/EdgeSegment.h
@@ -0,0 +1,266 @@
+// Copyright Epic Games, Inc. All Rights Reserved.
+#pragma once
+
+#include "CADKernel/Core/Types.h"
+
+#include "CADKernel/Math/Geometry.h"
+#include "CADKernel/Core/HaveStates.h"
+#include "CADKernel/Math/Boundary.h"
+#include "CADKernel/Math/Point.h"
+#include "CADKernel/Math/SlopeUtils.h"
+
+namespace UE::CADKernel
+{
+class FTopologicalLoop;
+class FTopologicalEdge;
+class FGrid;
+class FThinZone;
+class FThinZone2DFinder;
+class FTopologicalVertex;
+
+class FEdgeSegment : public FHaveStates
+{
+private:
+	FTopologicalEdge* Edge;
+	double Coordinates[2];
+	FPoint2D Points[2];
+
+	FEdgeSegment* NextSegment;
+	FEdgeSegment* PreviousSegment;
+
+	FEdgeSegment* ClosedSegment;
+
+	FSurfacicBoundary Boundary;
+	double AxisMin;
+
+	double SquareDistanceToClosedSegment;
+	double Length;
+
+	FIdent ChainIndex;
+
+	FIdent Id;
+	static FIdent LastId;
+
+public:
+	FEdgeSegment()
+		: Edge(nullptr)
+		, NextSegment(nullptr)
+		, PreviousSegment(nullptr)
+		, ClosedSegment(nullptr)
+		, AxisMin(0.)
+		, SquareDistanceToClosedSegment(HUGE_VALUE)
+		, Length(-1.)
+		, ChainIndex(Ident::Undefined)
+		, Id(0)
+	{
+	};
+
+	FEdgeSegment(const FEdgeSegment& Segment) = default;
+
+	virtual ~FEdgeSegment() = default;
+
+	void SetBoundarySegment(bool bInIsInnerLoop, FTopologicalEdge* InEdge, double InStartU, double InEndU, const FPoint2D& InStartPoint, const FPoint2D& InEndPoint)
+	{
+		if (bInIsInnerLoop)
+		{
+			SetInner();
+		}
+
+		Edge = InEdge;
+		Coordinates[ELimit::Start] = InStartU;
+		Coordinates[ELimit::End] = InEndU;
+		Points[ELimit::Start] = InStartPoint;
+		Points[ELimit::End] = InEndPoint;
+		NextSegment = nullptr;
+		PreviousSegment = nullptr;
+		ClosedSegment = nullptr;
+
+		SquareDistanceToClosedSegment = HUGE_VAL;
+		Length = Points[ELimit::Start].Distance(Points[ELimit::End]);
+
+		Id = ++LastId;
+		ChainIndex = Ident::Undefined;
+
+		Boundary.Set(Points[ELimit::Start], Points[ELimit::End]);
+
+		AxisMin = Boundary[EIso::IsoU].Min + Boundary[EIso::IsoV].Min;
+	};
+
+	void UpdateReferences(TMap<int32, FEdgeSegment*>& Map)
+	{
+		TFunction<void(FEdgeSegment*&)> UpdateReference = [&](FEdgeSegment*& Reference)
+		{
+			if (Reference)
+			{
+				FEdgeSegment** NewReference = Map.Find(Reference->GetId());
+				if (NewReference)
+				{
+					Reference = *NewReference;
+				}
+				else
+				{
+					Reference = nullptr;
+				}
+			}
+		};
+
+		UpdateReference(NextSegment);
+		UpdateReference(PreviousSegment);
+		UpdateReference(ClosedSegment);
+	}
+
+	double GetAxeMin() const
+	{
+		return AxisMin;
+	}
+
+	FIdent GetChainIndex() const
+	{
+		return ChainIndex;
+	}
+
+	void SetChainIndex(FIdent index)
+	{
+		ChainIndex = index;
+	}
+
+	bool IsInner() const
+	{
+		return ((States & EHaveStates::IsInner) == EHaveStates::IsInner);
+	}
+
+	void SetInner()
+	{
+		States |= EHaveStates::IsInner;
+	}
+
+	FIdent GetId() const
+	{
+		return Id;
+	}
+
+	const FTopologicalEdge* GetEdge() const
+	{
+		return Edge;
+	}
+
+	FTopologicalEdge* GetEdge()
+	{
+		return Edge;
+	}
+
+	double GetLength() const
+	{
+		return Length;
+	}
+
+	FPoint2D GetCenter() const
+	{
+		return (Points[ELimit::Start] + Points[ELimit::End]) * 0.5;
+	}
+
+	FPoint2D ComputeEdgePoint(double EdgeParamU) const
+	{
+		double SegmentParamS = (EdgeParamU - Coordinates[ELimit::Start]) / (Coordinates[ELimit::End] - Coordinates[ELimit::Start]);
+		return Points[ELimit::Start] + (Points[ELimit::End] - Points[ELimit::Start]) * SegmentParamS;
+	}
+
+	const FPoint2D& GetExtemity(const ELimit Limit) const
+	{
+		return Points[Limit];
+	}
+
+	double GetCoordinate(const ELimit Limit) const
+	{
+		return Coordinates[Limit];
+	}
+
+	bool IsForward()
+	{
+		return (Coordinates[ELimit::End] - Coordinates[ELimit::Start]) >= 0;
+	}
+
+	/**
+	 * Compute the slope of the input Segment according to this.
+	 */
+	double ComputeUnorientedSlopeOf(const FEdgeSegment* Segment)
+	{
+		double ReferenceSlope = ComputeSlope(Points[ELimit::Start], Points[ELimit::End]);
+		return ComputeUnorientedSlope(Segment->Points[ELimit::Start], Segment->Points[ELimit::End], ReferenceSlope);
+	}
+
+	/**
+	 * Compute the slope of the Segment defined by the two input points according to this.
+	 */
+	double ComputeUnorientedSlopeOf(const FPoint2D& Middle, const FPoint2D& Projection)
+	{
+		double ReferenceSlope = ComputeSlope(Points[ELimit::Start], Points[ELimit::End]);
+		return ComputeUnorientedSlope(Projection, Middle, ReferenceSlope);
+	}
+
+	FEdgeSegment* GetNext() const
+	{
+		return NextSegment;
+	}
+
+	FEdgeSegment* GetPrevious() const
+	{
+		return PreviousSegment;
+	}
+
+	FEdgeSegment* GetCloseSegment() const
+	{
+		return ClosedSegment;
+	}
+
+	void ResetCloseData()
+	{
+		if (ClosedSegment->ClosedSegment == this)
+		{
+			ClosedSegment->ClosedSegment = nullptr;
+		}
+		ClosedSegment = nullptr;
+		SquareDistanceToClosedSegment = HUGE_VAL;
+	}
+
+	void SetCloseSegment(FEdgeSegment* InSegmentA, double InDistance)
+	{
+		ClosedSegment = InSegmentA;
+		SquareDistanceToClosedSegment = InDistance;
+
+		if (InDistance < InSegmentA->SquareDistanceToClosedSegment)
+		{
+			InSegmentA->ClosedSegment = this;
+			InSegmentA->SquareDistanceToClosedSegment = InDistance;
+		}
+	}
+
+	double GetCloseSquareDistance() const
+	{
+		return SquareDistanceToClosedSegment;
+	}
+
+	void SetNext(FEdgeSegment* Segment)
+	{
+		NextSegment = Segment;
+		Segment->SetPrevious(this);
+	}
+
+	double ComputeEdgeCoordinate(const double SegmentU) const
+	{
+		return Coordinates[ELimit::Start] + (Coordinates[ELimit::End] - Coordinates[ELimit::Start]) * SegmentU;
+	}
+
+	FPoint2D ProjectPoint(const FPoint2D& PointToProject, double& SegmentU) const
+	{
+		return ProjectPointOnSegment(PointToProject, Points[ELimit::Start], Points[ELimit::End], SegmentU, true);
+	}
+
+private:
+	void SetPrevious(FEdgeSegment* Segment)
+	{
+		PreviousSegment = Segment;
+	}
+};
+}
+
diff --git a/Engine/Source/Runtime/Datasmith/CADKernel/Public/CADKernel/UI/Display.h b/Engine/Source/Runtime/Datasmith/CADKernel/Public/CADKernel/UI/Display.h
--- a/Engine/Source/Runtime/Datasmith/CADKernel/Public/CADKernel/UI/Display.h
+++ b/Engine/Source/Runtime/Datasmith/CADKernel/Public/CADKernel/UI/Display.h
@@ -7,6 +7,12 @@
 #include "CADKernel/Math/Point.h"
 #include "CADKernel/UI/Visu.h"
 
+#ifdef CADKERNEL_DEV
+#include "CADKernel/Mesh/Meshers/IsoTriangulator/DefineForDebug.h"
+#else
+static const double DisplayScale = 1.;
+#endif
+
 class FString;
 
 namespace UE::CADKernel
@@ -143,6 +149,17 @@ void DisplayPoint(const TPoint& Point, EVisuProperty Property, FIdent Ident)
 	DrawPoint(Point, Property);
 }
 
+static void DisplayPoint2DWithScale(const FPoint2D& Point, EVisuProperty Property = EVisuProperty::BluePoint)
+{
+	DrawPoint(Point * DisplayScale, Property);
+}
+
+static void DisplayPoint2DWithScale(const FPoint2D& Point, EVisuProperty Property, FIdent Ident)
+{
+	F3DDebugSegment G(Ident);
+	DrawPoint(Point * DisplayScale, Property);
+}
+
 template<typename TPoint>
 void DisplayPoints(FString Message, const TArray<TPoint>& Points, EVisuProperty Property = EVisuProperty::BluePoint, bool bDisplay = true)
 {
@@ -192,6 +209,11 @@ CADKERNEL_API void Display2D(const FTopologicalFace& Face);
 CADKERNEL_API void Display2D(const FTopologicalLoop& Loop);
 CADKERNEL_API void Display2D(const FSurface& CarrierSurface);
 
+CADKERNEL_API void Display2DWithScale(const FTopologicalEdge& Edge, EVisuProperty Property = EVisuProperty::BlueCurve);
+CADKERNEL_API void Display2DWithScale(const FTopologicalFace& Face);
+CADKERNEL_API void Display2DWithScale(const FTopologicalLoop& Loop);
+CADKERNEL_API void Display2DWithScale(const FSurface& CarrierSurface);
+
 CADKERNEL_API void DisplayMesh(const FFaceMesh& Mesh);
 CADKERNEL_API void DisplayMesh(const FEdgeMesh& Mesh);
 CADKERNEL_API void DisplayMesh(const FVertexMesh& Mesh);
@@ -202,7 +224,7 @@ CADKERNEL_API void DisplayControlPolygon(const FCurve& Entity);
 CADKERNEL_API void DisplayControlPolygon(const FSurface& Entity);
 
 template<typename TPoint>
-void DisplaySegment(const TPoint& Point1, const TPoint& Point2, FIdent Ident = 0, EVisuProperty Property = EVisuProperty::Element, bool bWithOrientation = false)
+void DisplaySegment(const TPoint& Point1, const TPoint& Point2, FIdent Ident = 0, EVisuProperty Property = EVisuProperty::BlueCurve, bool bWithOrientation = false)
 {
 	F3DDebugSegment G(Ident);
 	if (bWithOrientation)
@@ -212,6 +234,17 @@ void DisplaySegment(const TPoint& Point1, const TPoint& Point2, FIdent Ident = 0
 	DrawSegment(Point1, Point2, Property);
 };
 
+template<typename TPoint>
+void DisplaySegmentWithScale(const TPoint& Point1, const TPoint& Point2, FIdent Ident = 0, EVisuProperty Property = EVisuProperty::BlueCurve, bool bWithOrientation = false)
+{
+	F3DDebugSegment G(Ident);
+	if (bWithOrientation)
+	{
+		DrawSegmentOrientation(Point1 * DisplayScale, Point2 * DisplayScale, Property);
+	}
+	DrawSegment(Point1 * DisplayScale, Point2 * DisplayScale, Property);
+};
+
 template<typename TPoint>
 void DisplayPolyline(const TArray<TPoint>& Points, EVisuProperty Property)
 {
@@ -220,6 +253,19 @@ void DisplayPolyline(const TArray<TPoint>& Points, EVisuProperty Property)
 	Close3DDebugSegment();
 }
 
+static void DisplayPolylineWithScale(const TArray<FPoint2D>& Points, EVisuProperty Property)
+{
+	TArray<FPoint2D> PointsWithScale = Points;
+	for (FPoint2D& Point : PointsWithScale)
+	{
+		Point *= DisplayScale;
+	}
+
+	Open3DDebugSegment(0);
+	Draw(PointsWithScale, Property);
+	Close3DDebugSegment();
+}
+
 template<typename TPoint>
 void DisplayOrientedPolyline(const TArray<TPoint>& Points, EVisuProperty Property)
 {
@@ -236,6 +282,7 @@ CADKERNEL_API void DrawQuadripode(double Height, double Base, FPoint& Centre, FP
 CADKERNEL_API void Draw(const FTopologicalEdge& Edge, EVisuProperty Property = EVisuProperty::BlueCurve);
 CADKERNEL_API void Draw(const FTopologicalFace& Face);
 CADKERNEL_API void Draw2D(const FTopologicalFace& Face);
+CADKERNEL_API void Draw2DWithScale(const FTopologicalFace& Face);
 CADKERNEL_API void Draw(const FShell& Shell);
 
 CADKERNEL_API void Draw(const FCurve& Curve, EVisuProperty Property = EVisuProperty::BlueCurve);
