[NO NON-PII COMMIT MESSAGE]

diff --git a/Engine/Source/Runtime/Datasmith/CADKernel/Private/CADKernel/Mesh/Meshers/IsoTriangulator/IsoNode.h b/Engine/Source/Runtime/Datasmith/CADKernel/Private/CADKernel/Mesh/Meshers/IsoTriangulator/IsoNode.h
--- a/Engine/Source/Runtime/Datasmith/CADKernel/Private/CADKernel/Mesh/Meshers/IsoTriangulator/IsoNode.h
+++ b/Engine/Source/Runtime/Datasmith/CADKernel/Private/CADKernel/Mesh/Meshers/IsoTriangulator/IsoNode.h
@@ -2,6 +2,7 @@
 #pragma once
 
 #include "CADKernel/Core/Types.h"
+#include "CADKernel/Math/SlopeUtils.h"
 #include "CADKernel/Mesh/MeshEnum.h"
 #include "CADKernel/Mesh/Structure/Grid.h"
 
