ThinFace preparatory work:
- minore refactoring

diff --git a/Engine/Plugins/Enterprise/DatasmithCADImporter/Source/CADInterfaces/Private/TechSoftFileParser.cpp b/Engine/Plugins/Enterprise/DatasmithCADImporter/Source/CADInterfaces/Private/TechSoftFileParser.cpp
--- a/Engine/Plugins/Enterprise/DatasmithCADImporter/Source/CADInterfaces/Private/TechSoftFileParser.cpp
+++ b/Engine/Plugins/Enterprise/DatasmithCADImporter/Source/CADInterfaces/Private/TechSoftFileParser.cpp
@@ -473,6 +473,86 @@ ECADParsingResult FTechSoftFileParser::Process()
 	return Result;
 }
 
+
+A3DStatus FTechSoftFileParser::AdaptBRepModel()
+{
+	TRACE_CPUPROFILER_EVENT_SCOPE(FTechSoftFileParserCADKernelTessellator::AdaptBRepModel);
+
+	const A3DUns32 ValidSurfaceCount = 9;
+	A3DUns32 AcceptedSurfaces[ValidSurfaceCount] = {
+		//kA3DTypeSurfBlend01,
+		//kA3DTypeSurfBlend02,
+		//kA3DTypeSurfBlend03,
+		kA3DTypeSurfNurbs,
+		kA3DTypeSurfCone,
+		kA3DTypeSurfCylinder,
+		kA3DTypeSurfCylindrical,
+		//kA3DTypeSurfOffset,
+		//kA3DTypeSurfPipe,
+		kA3DTypeSurfPlane,
+		kA3DTypeSurfRuled,
+		kA3DTypeSurfSphere,
+		kA3DTypeSurfRevolution,
+		//kA3DTypeSurfExtrusion,
+		//kA3DTypeSurfFromCurves,
+		kA3DTypeSurfTorus,
+		//kA3DTypeSurfTransform,
+	};
+
+	const A3DUns32 ValidCurveCount = 7;
+	A3DUns32 AcceptedCurves[ValidCurveCount] = {
+		//kA3DTypeCrvBase,
+		//kA3DTypeCrvBlend02Boundary,
+		kA3DTypeCrvNurbs,
+		kA3DTypeCrvCircle,
+		//kA3DTypeCrvComposite,
+		//kA3DTypeCrvOnSurf,
+		kA3DTypeCrvEllipse,
+		//kA3DTypeCrvEquation,
+		//kA3DTypeCrvHelix,
+		kA3DTypeCrvHyperbola,
+		//kA3DTypeCrvIntersection,
+		kA3DTypeCrvLine,
+		//kA3DTypeCrvOffset,
+		kA3DTypeCrvParabola,
+		kA3DTypeCrvPolyLine,
+		//kA3DTypeCrvTransform,
+	};
+
+	TUniqueTSObj<A3DCopyAndAdaptBrepModelData> CopyAndAdaptBrepModelData;
+	CopyAndAdaptBrepModelData->m_bUseSameParam = false;                        // If `A3D_TRUE`, surfaces will keep their parametrization when converted to NURBS.       
+	CopyAndAdaptBrepModelData->m_dTol = 1e-3;                                  // Tolerance value of resulting B-rep. The value is relative to the scale of the model.
+	CopyAndAdaptBrepModelData->m_bDeleteCrossingUV = false;                    // If `A3D_TRUE`, UV curves that cross seams of periodic surfaces are replaced by 3D curves 
+	CopyAndAdaptBrepModelData->m_bSplitFaces = true;                           // If `A3D_TRUE`, the faces with a periodic basis surface are split on parametric seams
+	CopyAndAdaptBrepModelData->m_bSplitClosedFaces = false;                    // If `A3D_TRUE`, the faces with a closed basis surface are split into faces at the parametric seam and mid-parameter
+	CopyAndAdaptBrepModelData->m_bForceComputeUV = true;                       // If `A3D_TRUE`, UV curves are computed from the B-rep data
+	CopyAndAdaptBrepModelData->m_bAllowUVCrossingSeams = true;                 // If `A3D_TRUE` and m_bForceComputeUV is set to `A3D_TRUE`, computed UV curves can cross seams.
+	CopyAndAdaptBrepModelData->m_bForceCompute3D = false;                      // If `A3D_TRUE`, 3D curves are computed from the B-rep data
+	CopyAndAdaptBrepModelData->m_bContinueOnError = true;                      // Continue processing even if an error occurs. Use \ref A3DCopyAndAdaptBrepModelAdvanced to get the error status.
+	CopyAndAdaptBrepModelData->m_bClampTolerantUVCurvesInsideUVDomain = false; // If `A3D_FALSE`, UV curves may stray outside the UV domain as long as the 3D edge tolerance is respected. If set to `A3D_TRUE`, the UV curves will be clamped to the UV domain (if the clamp still leaves them within the edge tolerance). */
+	CopyAndAdaptBrepModelData->m_bForceDuplicateGeometries = false;            // If `A3D_TRUE`, break the sharing of surfaces and curves into topologies.*/
+
+	CopyAndAdaptBrepModelData->m_uiAcceptableSurfacesSize = ValidSurfaceCount;
+	CopyAndAdaptBrepModelData->m_puiAcceptableSurfaces = &AcceptedSurfaces[0];
+	CopyAndAdaptBrepModelData->m_uiAcceptableCurvesSize = ValidCurveCount;
+	CopyAndAdaptBrepModelData->m_puiAcceptableCurves = &AcceptedCurves[0];
+
+	int32 ErrorCount = 0;
+	A3DCopyAndAdaptBrepModelErrorData* Errors = nullptr;
+	A3DStatus Ret = TechSoftInterface::AdaptBRepInModelFile(ModelFile.Get(), *CopyAndAdaptBrepModelData, ErrorCount, &Errors);
+	if ((Ret == A3D_SUCCESS || Ret == A3D_TOOLS_CONTINUE_ON_ERROR) && ErrorCount > 0)
+	{
+		// Add message about non-critical errors during the adaptation
+		CADFileData.LogMessage(FString::Printf(TEXT("File %s had %d non-critical error(s) during BRep adaptation step."), *CADFileData.GetCADFileDescription().GetFileName(), ErrorCount));
+	}
+	else if (Ret != A3D_SUCCESS)
+	{
+		CADFileData.LogWarning(FString::Printf(TEXT("File %s failed during BRep adaptation step."), *CADFileData.GetCADFileDescription().GetFileName(), ErrorCount));
+		return A3D_ERROR;
+	}
+	return A3D_SUCCESS;
+}
+
 void FTechSoftFileParser::SewModel()
 {
 	CADLibrary::TUniqueTSObj<A3DSewOptionsData> SewData;
@@ -2006,21 +2086,22 @@ void FTechSoftFileParser::ReadMaterialsAndColors()
 	}
 }
 
+#ifndef CADKERNEL_DEV
 void FTechSoftFileParser::CheckMemory()
 {
-	//CADFileData.GetRecord().StartMemoryUsed = FPlatformMemory::GetStats().UsedPhysical;
-	//uint64& MaxMemoryUsed = CADFileData.GetRecord().MaxMemoryUsed;
-	//while (bProcessIsRunning)
-	//{
-	//	FPlatformProcess::Sleep(0.1);
-	//	const uint64 MemoryUsed = FPlatformMemory::GetStats().UsedPhysical;
-	//	if (MaxMemoryUsed < MemoryUsed)
-	//	{
-	//		MaxMemoryUsed = MemoryUsed;
-	//	}
-	//}
+	CADFileData.GetRecord().StartMemoryUsed = FPlatformMemory::GetStats().UsedPhysical;
+	uint64& MaxMemoryUsed = CADFileData.GetRecord().MaxMemoryUsed;
+	while (bProcessIsRunning)
+	{
+		FPlatformProcess::Sleep(0.1);
+		const uint64 MemoryUsed = FPlatformMemory::GetStats().UsedPhysical;
+		if (MaxMemoryUsed < MemoryUsed)
+		{
+			MaxMemoryUsed = MemoryUsed;
+		}
+	}
 }
-
+#endif  
 
 #endif  
 
diff --git a/Engine/Plugins/Enterprise/DatasmithCADImporter/Source/CADInterfaces/Private/TechSoftFileParser.h b/Engine/Plugins/Enterprise/DatasmithCADImporter/Source/CADInterfaces/Private/TechSoftFileParser.h
--- a/Engine/Plugins/Enterprise/DatasmithCADImporter/Source/CADInterfaces/Private/TechSoftFileParser.h
+++ b/Engine/Plugins/Enterprise/DatasmithCADImporter/Source/CADInterfaces/Private/TechSoftFileParser.h
@@ -53,10 +53,10 @@ public:
 
 protected:
 
-	virtual A3DStatus AdaptBRepModel()
-	{
-		return A3DStatus::A3D_SUCCESS;
-	}
+	virtual A3DStatus AdaptBRepModel();
+	//{
+	//	return A3DStatus::A3D_SUCCESS;
+	//}
 
 	/**
 	 * If the tessellator is TechSoft, SewModel call TechSoftInterface::SewModel
@@ -144,8 +144,9 @@ private:
 	void ExtractGeneralTransformation(const A3DMiscTransformation* GeneralTransformation, FArchiveCADObject& Component);
 	void ExtractTransformation3D(const A3DMiscTransformation* CartesianTransformation, FArchiveCADObject& Component);
 
+#ifndef CADKERNEL_DEV
 	void CheckMemory();
-
+#endif  
 
 #endif
 
diff --git a/Engine/Source/Runtime/Datasmith/CADKernel/Private/CADKernel/Core/Factory.h b/Engine/Source/Runtime/Datasmith/CADKernel/Private/CADKernel/Core/Factory.h
--- a/Engine/Source/Runtime/Datasmith/CADKernel/Private/CADKernel/Core/Factory.h
+++ b/Engine/Source/Runtime/Datasmith/CADKernel/Private/CADKernel/Core/Factory.h
@@ -3,12 +3,13 @@
 
 #include "CADKernel/Core/Types.h"
 #include "CADKernel/Core/Group.h"
+#ifdef CADKERNEL_DEV
+#include "CADKernel/UI/DefineForDebug.h"
+#endif
 
 #include "Containers/List.h"
 
 
-#define DEBUG_FACTORY
-
 namespace UE::CADKernel
 {
 const int32 MaxSize = 256;
@@ -35,7 +36,9 @@ private:
 		{
 			FreeEntityStack.Add(&Entity);
 		}
+#ifdef DEBUG_FACTORY
 		AllocatedElementNum += MaxSize;
+#endif
 	}
 
 public:
@@ -48,18 +51,24 @@ public:
 	{
 		Entity->Clean();
 		FreeEntityStack.Add(Entity);
+#ifdef DEBUG_FACTORY
 		UsedElementNum--;
 		ensureCADKernel(UsedElementNum + FreeEntityStack.Num() == AllocatedElementNum);
+#endif
 	}
 
 	ElementType& New()
 	{
+#ifdef DEBUG_FACTORY
 		ensureCADKernel(UsedElementNum + FreeEntityStack.Num() == AllocatedElementNum);
+#endif
 		if (FreeEntityStack.IsEmpty())
 		{
 			FillFreeEntityStack();
 		}
+#ifdef DEBUG_FACTORY
 		UsedElementNum++;
+#endif
 		return *FreeEntityStack.Pop(false);
 	}
 };
diff --git a/Engine/Source/Runtime/Datasmith/CADKernel/Private/CADKernel/Geo/Curves/CurveMaker.cpp b/Engine/Source/Runtime/Datasmith/CADKernel/Private/CADKernel/Geo/Curves/CurveMaker.cpp
--- a/Engine/Source/Runtime/Datasmith/CADKernel/Private/CADKernel/Geo/Curves/CurveMaker.cpp
+++ b/Engine/Source/Runtime/Datasmith/CADKernel/Private/CADKernel/Geo/Curves/CurveMaker.cpp
@@ -1,17 +1,18 @@
 // Copyright Epic Games, Inc. All Rights Reserved.
 
-#include "CADKernel/Geo/Curves/BoundedCurve.h"
 #include "CADKernel/Geo/Curves/BezierCurve.h"
+//#include "CADKernel/Geo/Curves/BoundedCurve.h"
+//#include "CADKernel/Geo/Curves/CompositeCurve.h"
 #include "CADKernel/Geo/Curves/Curve.h"
+//#include "CADKernel/Geo/Curves/EllipseCurve.h"
+//#include "CADKernel/Geo/Curves/HyperbolaCurve.h"
 #include "CADKernel/Geo/Curves/NURBSCurve.h"
-#include "CADKernel/Geo/Curves/SegmentCurve.h"
-#include "CADKernel/Geo/Curves/EllipseCurve.h"
-#include "CADKernel/Geo/Curves/HyperbolaCurve.h"
-#include "CADKernel/Geo/Curves/ParabolaCurve.h"
-#include "CADKernel/Geo/Curves/CompositeCurve.h"
-#include "CADKernel/Geo/Curves/PolylineCurve.h"
-#include "CADKernel/Geo/Curves/RestrictionCurve.h"
-#include "CADKernel/Geo/Curves/SurfacicCurve.h"
+//#include "CADKernel/Geo/Curves/ParabolaCurve.h"
+//#include "CADKernel/Geo/Curves/PolylineCurve.h"
+//#include "CADKernel/Geo/Curves/RestrictionCurve.h"
+//#include "CADKernel/Geo/Curves/SegmentCurve.h"
+#include "CADKernel/Geo/Curves/SplineCurve.h"
+//#include "CADKernel/Geo/Curves/SurfacicCurve.h"
 
 namespace UE::CADKernel
 {
@@ -21,4 +22,24 @@ TSharedPtr<FCurve> FCurve::MakeNurbsCurve(FNurbsCurveData& InNurbsData)
 	return FEntity::MakeShared<UE::CADKernel::FNURBSCurve>(InNurbsData);
 }
 
+TSharedPtr<FCurve> FCurve::MakeBezierCurve(const TArray<FPoint>& InPoles)
+{
+	return FEntity::MakeShared<UE::CADKernel::FBezierCurve>(InPoles);
+}
+
+TSharedPtr<FCurve> FCurve::MakeSplineCurve(const TArray<FPoint>& InPoles)
+{
+	return FEntity::MakeShared<UE::CADKernel::FSplineCurve>(InPoles);
+}
+
+TSharedPtr<FCurve> FCurve::MakeSplineCurve(const TArray<FPoint>& InPoles, const TArray<FPoint>& InTangents)
+{
+	return FEntity::MakeShared<UE::CADKernel::FSplineCurve>(InPoles, InTangents);
+}
+
+TSharedPtr<FCurve> FCurve::MakeSplineCurve(const TArray<FPoint>& InPoles, const TArray<FPoint>& InArriveTangents, const TArray<FPoint>& InLeaveTangents)
+{
+	return FEntity::MakeShared<UE::CADKernel::FSplineCurve>(InPoles, InArriveTangents, InLeaveTangents);
+}
+
 }
\ No newline at end of file
diff --git a/Engine/Source/Runtime/Datasmith/CADKernel/Private/CADKernel/Geo/Curves/SplineCurve.cpp b/Engine/Source/Runtime/Datasmith/CADKernel/Private/CADKernel/Geo/Curves/SplineCurve.cpp
new file mode 100644
--- /dev/null
+++ b/Engine/Source/Runtime/Datasmith/CADKernel/Private/CADKernel/Geo/Curves/SplineCurve.cpp
@@ -0,0 +1,269 @@
+// Copyright Epic Games, Inc. All Rights Reserved.
+#include "CADKernel/Geo/Curves/SplineCurve.h"
+
+#include "CADKernel/Geo/GeoPoint.h"
+#include "CADKernel/Geo/Sampling/PolylineTools.h"
+#include "CADKernel/Math/BSpline.h"
+
+namespace UE::CADKernel
+{
+
+TSharedPtr<FEntityGeom> FSplineCurve::ApplyMatrix(const FMatrixH& InMatrix) const
+{
+	ensure(false);
+	//TArray<FPoint> TransformedPoles;
+	//TransformedPoles.Reserve(Poles.Num());
+
+	//for (const FPoint& Pole : Poles)
+	//{
+	//	TransformedPoles.Emplace(InMatrix.Multiply(Pole));
+	//}
+
+	return FEntity::MakeShared<FSplineCurve>(/*TransformedPoles*/);
+}
+
+void FSplineCurve::Offset(const FPoint& OffsetDirection)
+{
+	ensure(false);
+	//for (FPoint& Pole : Poles)
+	//{
+	//	Pole += OffsetDirection;
+	//}
+}
+
+#ifdef CADKERNEL_DEV
+FInfoEntity& FSplineCurve::GetInfo(FInfoEntity& Info) const
+{
+	return FCurve::GetInfo(Info);
+// 		.Add(TEXT("degre"), Poles.Num() - 1)
+// 		.Add(TEXT("poles"), Poles);
+}
+#endif
+
+void FSplineCurve::ExtendTo(const FPoint& Point)
+{
+	ensure(false);
+	//PolylineTools::ExtendTo(Poles, Point);
+}
+
+void FSplineCurve::SetSplinePoints(const TArray<FPoint>& Points)
+{
+	const int32 NumPoints = Points.Num();
+	Position.Points.Reset(NumPoints);
+	//Rotation.Points.Reset(NumPoints);
+	//Scale.Points.Reset(NumPoints);
+
+	double InputKey = 0.0;
+	for (const FPoint& Point : Points)
+	{
+		Position.Points.Emplace(InputKey, Point, FPoint::ZeroPoint, FPoint::ZeroPoint, CIM_CurveAuto);
+		//Rotation.Points.Emplace(InputKey, FQuat::Identity, FQuat::Identity, FQuat::Identity, CIM_CurveAuto);
+		//Scale.Points.Emplace(InputKey, FPoint::UnitPoint, FPoint::ZeroPoint, FPoint::ZeroPoint, CIM_CurveAuto);
+		InputKey += 1.0;
+	}
+
+	UpdateSpline();
+
+	Boundary.SetMin(0.);
+	Boundary.SetMax(InputKey);
+}
+
+void FSplineCurve::SetSplinePoints(const TArray<FPoint>& Points, const TArray<FPoint>& Tangents)
+{
+	const int32 NumPoints = Points.Num();
+	ensure(Tangents.Num() == NumPoints);
+
+	Position.Points.Reset(NumPoints);
+
+	double InputKey = 0.0;
+	for (int32 Index = 0; Index < Points.Num(); ++Index)
+	{
+		Position.Points.Emplace(InputKey, Points[Index], Tangents[Index], Tangents[Index], CIM_CurveAuto);
+		InputKey += 1.0;
+	}
+	//UpdateSpline();
+
+	Boundary.SetMin(0.);
+	Boundary.SetMax(InputKey);
+}
+
+void FSplineCurve::SetSplinePoints(const TArray<FPoint>& Points, const TArray<FPoint>& ArriveTangents, const TArray<FPoint>& LeaveTangents)
+{
+	const int32 NumPoints = Points.Num();
+	Position.Points.Reset(NumPoints);
+	ensure(false);
+
+	double InputKey = 0.0;
+	for (int32 Index = 0; Index < Points.Num(); ++Index)
+	{
+		//Position.Points.Emplace(InputKey, Points[Index], ArriveTangents[Index], LeaveTangents[Index], ? CIM_CurveAuto : CIM_CurveBreak);
+		InputKey += 1.0;
+	}
+	UpdateSpline();
+
+	Boundary.SetMin(0.);
+	Boundary.SetMax(InputKey);
+}
+
+
+
+void FSplineCurve::UpdateSpline(bool bInClosedLoop, bool bStationaryEndpoints, int32 ReparamStepsPerSegment)
+{
+	const int32 NumPoints = Position.Points.Num();
+	//check(Rotation.Points.Num() == NumPoints && Scale.Points.Num() == NumPoints);
+
+	// Ensure input keys are strictly ascending
+	for (int32 Index = 1; Index < NumPoints; Index++)
+	{
+		ensureAlways(Position.Points[Index - 1].InVal < Position.Points[Index].InVal);
+	}
+
+	// Ensure splines' looping status matches with that of the spline component
+	if (bInClosedLoop)
+	{
+		const double LastKey = Position.Points.Num() > 0 ? Position.Points.Last().InVal : 0.0;
+		const double LoopKey = /*bLoopPositionOverride ? LoopPosition :*/ LastKey + 1.0;
+		Position.SetLoopKey(LoopKey);
+		//Rotation.SetLoopKey(LoopKey);
+		//Scale.SetLoopKey(LoopKey);
+	}
+	else
+	{
+		Position.ClearLoopKey();
+		//Rotation.ClearLoopKey();
+		//Scale.ClearLoopKey();
+	}
+
+	// Automatically set the tangents on any CurveAuto keys
+	Position.AutoSetTangents(0.0, bStationaryEndpoints);
+	//Rotation.AutoSetTangents(0.0, bStationaryEndpoints);
+	//Scale.AutoSetTangents(0.0, bStationaryEndpoints);
+
+	// Now initialize the spline reparam table
+	const int32 NumSegments = bInClosedLoop ? NumPoints : FMath::Max(0, NumPoints - 1);
+
+	// Start by clearing it
+	ReparamTable.Points.Reset(NumSegments * ReparamStepsPerSegment + 1);
+	double AccumulatedLength = 0.0f;
+	for (int32 SegmentIndex = 0; SegmentIndex < NumSegments; ++SegmentIndex)
+	{
+		for (int32 Step = 0; Step < ReparamStepsPerSegment; ++Step)
+		{
+			const double Param = static_cast<double>(Step) / ReparamStepsPerSegment;
+			const double SegmentLength = (Step == 0) ? 0.0 : GetSegmentLength(SegmentIndex, Param, bInClosedLoop);
+
+			ReparamTable.Points.Emplace(SegmentLength + AccumulatedLength, SegmentIndex + Param, 0.0, 0.0, CIM_Linear);
+		}
+		AccumulatedLength += GetSegmentLength(SegmentIndex, 1.0, bInClosedLoop);
+	}
+
+	ReparamTable.Points.Emplace(AccumulatedLength, static_cast<double>(NumSegments), 0.0, 0.0, CIM_Linear);
+}
+
+FPoint FSplineCurve::GetLocationAtSplineInputKey(double Coordinate) const
+{
+	return Position.Eval((float)Coordinate, FPoint::ZeroPoint);
+}
+
+FPoint FSplineCurve::GetTangentAtSplineInputKey(double Coordinate) const
+{
+	return Position.EvalDerivative((float)Coordinate, FPoint::ZeroPoint);
+}
+
+FPoint FSplineCurve::GetDirectionAtSplineInputKey(double Coordinate) const
+{
+	return Position.EvalDerivative(Coordinate, FPoint::ZeroPoint).Normalize();
+}
+
+void FSplineCurve::EvaluatePoint(double Coordinate, FCurvePoint& OutPoint, int32 DerivativeOrder) const
+{
+
+
+	OutPoint.Point = GetLocationAtSplineInputKey(Coordinate);
+	if (DerivativeOrder > 0)
+	{
+		OutPoint.Gradient = GetTangentAtSplineInputKey(Coordinate);
+	}
+}
+
+double FSplineCurve::GetSplineLength() const
+{
+	const int32 NumPoints = ReparamTable.Points.Num();
+	if (NumPoints > 0)
+	{
+		return ReparamTable.Points.Last().InVal;
+	}
+	return 0.0f;
+}
+
+double FSplineCurve::GetSegmentLength(const int32 Index, const double Param, bool bInClosedLoop) const
+{
+	const int32 NumPoints = Position.Points.Num();
+	const int32 LastPoint = NumPoints - 1;
+
+	check(Index >= 0 && ((bInClosedLoop && Index < NumPoints) || (!bInClosedLoop && Index < LastPoint)));
+	check(Param >= 0.0 && Param <= 1.0);
+
+	// Evaluate the length of a Hermite spline segment.
+	// This calculates the integral of |dP/dt| dt, where P(t) is the spline equation with components (x(t), y(t), z(t)).
+	// This isn't solvable analytically, so we use a numerical method (Legendre-Gauss quadrature) which performs very well
+	// with functions of this type, even with very few samples.  In this case, just 5 samples is sufficient to yield a
+	// reasonable result.
+
+	struct FLegendreGaussCoefficient
+	{
+		double Abscissa;
+		double Weight;
+	};
+
+	static const FLegendreGaussCoefficient LegendreGaussCoefficients[] =
+	{
+		{ 0.0f, 0.5688889f },
+		{ -0.5384693f, 0.47862867f },
+		{ 0.5384693f, 0.47862867f },
+		{ -0.90617985f, 0.23692688f },
+		{ 0.90617985f, 0.23692688f }
+	};
+
+	const auto& StartPoint = Position.Points[Index];
+	const auto& EndPoint = Position.Points[Index == LastPoint ? 0 : Index + 1];
+
+	const auto& P0 = StartPoint.OutVal;
+	const auto& T0 = StartPoint.LeaveTangent;
+	const auto& P1 = EndPoint.OutVal;
+	const auto& T1 = EndPoint.ArriveTangent;
+
+	// Special cases for linear or constant segments
+	if (StartPoint.InterpMode == CIM_Linear)
+	{
+		return (P1 - P0).Length() * Param;
+	}
+	else if (StartPoint.InterpMode == CIM_Constant)
+	{
+		// Special case: constant interpolation acts like distance = 0 for all p in [0, 1[ but for p == 1, the distance returned is the linear distance between start and end
+		return Param == 1. ? (P1 - P0).Length() : 0.0;
+	}
+
+	// Cache the coefficients to be fed into the function to calculate the spline derivative at each sample point as they are constant.
+	const FPoint Coeff1 = ((P0 - P1) * 2.0 + T0 + T1) * 3.0;
+	const FPoint Coeff2 = (P1 - P0) * 6.0 - T0 * 4.0f - T1 * 2.0;
+	const FPoint Coeff3 = T0;
+
+	const double HalfParam = Param * 0.5;
+
+	double Length = 0.0;
+	for (const auto& LegendreGaussCoefficient : LegendreGaussCoefficients)
+	{
+		// Calculate derivative at each Legendre-Gauss sample, and perform a weighted sum
+		const double Alpha = HalfParam * (1.0 + LegendreGaussCoefficient.Abscissa);
+		const FPoint Derivative = ((Coeff1 * Alpha + Coeff2) * Alpha + Coeff3);
+		Length += Derivative.Length() * LegendreGaussCoefficient.Weight;
+	}
+	Length *= HalfParam;
+
+	return Length;
+}
+
+
+} // namespace UE::CADKernel
+
diff --git a/Engine/Source/Runtime/Datasmith/CADKernel/Private/CADKernel/Geo/Curves/SplineCurve.h b/Engine/Source/Runtime/Datasmith/CADKernel/Private/CADKernel/Geo/Curves/SplineCurve.h
new file mode 100644
--- /dev/null
+++ b/Engine/Source/Runtime/Datasmith/CADKernel/Private/CADKernel/Geo/Curves/SplineCurve.h
@@ -0,0 +1,117 @@
+// Copyright Epic Games, Inc. All Rights Reserved.
+#pragma once
+
+#include "CADKernel/Geo/Curves/Curve.h"
+
+#include "CADKernel/Math/Point.h"
+#include "Math/InterpCurve.h"
+
+DEFINE_INTERPCURVE_WRAPPER_STRUCT(FInterpCurveDouble, double)
+DEFINE_INTERPCURVE_WRAPPER_STRUCT(FInterpCurveFPoint, UE::CADKernel::FPoint)
+
+namespace UE::CADKernel
+{
+
+class CADKERNEL_API FSplineCurve : public FCurve
+{
+	friend class FEntity;
+
+protected:
+
+	FInterpCurveFPoint Position;
+
+	/** Spline built from rotation data. */
+	//FInterpCurveQuat Rotation;
+
+	/** Spline built from scale data. */
+	//FInterpCurveFPoint Scale;
+
+	/** Input: distance along curve, output: parameter that puts you there. */
+	FInterpCurveDouble ReparamTable;
+
+	/**
+	 * Whether the spline is to be considered as a closed loop.
+	 * Use SetClosedLoop() to set this property, and IsClosedLoop() to read it.
+	 */
+	bool bClosedLoop;
+
+	//bool bLoopPositionOverride;
+	//double LoopPosition;
+
+	FSplineCurve(const TArray<FPoint>& InPoles)
+	{
+		SetSplinePoints(InPoles);
+	}
+
+	FSplineCurve(const TArray<FPoint>& InPoles, const TArray<FPoint>& InTangents)
+	{
+		SetSplinePoints(InPoles, InTangents);
+	}
+
+	FSplineCurve(const TArray<FPoint>& InPoles, const TArray<FPoint>& InArriveTangents, const TArray<FPoint>& InLeaveTangents)
+	{
+		SetSplinePoints(InPoles, InArriveTangents, InLeaveTangents);
+	}
+
+	FSplineCurve() = default;
+
+	/** Get location along spline at the provided input coordinate value */
+	FPoint GetLocationAtSplineInputKey(double InCoordinate) const;
+
+	/** Get tangent along spline at the provided input coordinate value */
+	FPoint GetTangentAtSplineInputKey(double InCoordinate) const;
+
+	/** Get unit direction (normalized tangent) along spline at the provided input coordinate value */
+	FPoint GetDirectionAtSplineInputKey(double InCoordinate) const;
+
+public:
+
+	virtual void Serialize(FCADKernelArchive& Ar) override
+	{
+		FCurve::Serialize(Ar);
+		Ar << Position;
+		Ar << ReparamTable;
+		Ar << bClosedLoop;
+	}
+
+#ifdef CADKERNEL_DEV
+	virtual FInfoEntity& GetInfo(FInfoEntity&) const override;
+#endif
+
+	virtual ECurve GetCurveType() const override
+	{
+		return ECurve::Spline;
+	}
+
+	virtual TSharedPtr<FEntityGeom> ApplyMatrix(const FMatrixH& InMatrix) const override;
+	virtual void Offset(const FPoint& OffsetDirection) override;
+
+	virtual void EvaluatePoint(double Coordinate, FCurvePoint& OutPoint, int32 DerivativeOrder = 0) const override;
+
+	virtual void ExtendTo(const FPoint& Point) override;
+
+	/**
+	 * Update the spline's internal data according to the passed-in params
+	 * @param	bClosedLoop				Whether the spline is to be considered as a closed loop.
+	 * @param	bStationaryEndpoints	Whether the endpoints of the spline are considered stationary when traversing the spline at non-constant velocity.  Essentially this sets the endpoints' tangents to zero vectors.
+	 * @param	ReparamStepsPerSegment	Number of steps per spline segment to place in the re parameterization table
+	 */
+	void UpdateSpline(bool bClosedLoop = false, bool bStationaryEndpoints = false, int32 ReparamStepsPerSegment = 10);
+
+	/** Returns the length of the specified spline segment up to the parametric value given */
+	double GetSegmentLength(const int32 Index, const double Param, bool bClosedLoop = false) const;
+
+	/** Returns total length along this spline */
+	double GetSplineLength() const;
+
+	/** Sets the spline to an array of points */
+	void SetSplinePoints(const TArray<FPoint>& InPoints);
+	void SetSplinePoints(const TArray<FPoint>& InPoints, const TArray<FPoint>& InTangents);
+	void SetSplinePoints(const TArray<FPoint>& InPoints, const TArray<FPoint>& InArriveTangents, const TArray<FPoint>& InLeaveTangents);
+
+	const FInterpCurveFPoint& GetSplinePointsPosition() const { return Position; }
+
+};
+
+} // namespace UE::CADKernel
+
diff --git a/Engine/Source/Runtime/Datasmith/CADKernel/Private/CADKernel/Mesh/Criteria/CriteriaGrid.cpp b/Engine/Source/Runtime/Datasmith/CADKernel/Private/CADKernel/Mesh/Criteria/CriteriaGrid.cpp
--- a/Engine/Source/Runtime/Datasmith/CADKernel/Private/CADKernel/Mesh/Criteria/CriteriaGrid.cpp
+++ b/Engine/Source/Runtime/Datasmith/CADKernel/Private/CADKernel/Mesh/Criteria/CriteriaGrid.cpp
@@ -1,7 +1,14 @@
 // Copyright Epic Games, Inc. All Rights Reserved.
 #include "CADKernel/Mesh/Criteria/CriteriaGrid.h"
 
+#include "CADKernel/Mesh/Criteria/Criterion.h"
+#include "CADKernel/Geo/Surfaces/Surface.h"
 #include "CADKernel/Topo/TopologicalFace.h"
+#include "CADKernel/UI/Display.h"
+
+#ifdef CADKERNEL_DEV
+#include "CADKernel/UI/DefineForDebug.h"
+#endif
 
 namespace UE::CADKernel
 {
@@ -28,8 +35,7 @@ void FCriteriaGrid::Init()
 	ComputeMiddlePointsCoordinates(EIso::IsoU);
 	ComputeMiddlePointsCoordinates(EIso::IsoV);
 
-	const bool bWithNormals = false;
-	EvaluatePointGrid(CoordinateGrid, bWithNormals);
+	EvaluatePointGrid(CoordinateGrid, false);
 }
 
 FCriteriaGrid::FCriteriaGrid(FTopologicalFace& InFace)
diff --git a/Engine/Source/Runtime/Datasmith/CADKernel/Private/CADKernel/Mesh/Meshers/BowyerWatsonTriangulator.h b/Engine/Source/Runtime/Datasmith/CADKernel/Private/CADKernel/Mesh/Meshers/BowyerWatsonTriangulator.h
--- a/Engine/Source/Runtime/Datasmith/CADKernel/Private/CADKernel/Mesh/Meshers/BowyerWatsonTriangulator.h
+++ b/Engine/Source/Runtime/Datasmith/CADKernel/Private/CADKernel/Mesh/Meshers/BowyerWatsonTriangulator.h
@@ -6,7 +6,7 @@
 #include "CADKernel/Mesh/Meshers/IsoTriangulator/IsoCell.h"
 #include "CADKernel/UI/Display.h"
 #ifdef CADKERNEL_DEV
-#include "CADKernel/Mesh/Meshers/IsoTriangulator/DefineForDebug.h"
+#include "CADKernel/UI/DefineForDebug.h"
 #endif
 
 namespace UE::CADKernel
@@ -68,7 +68,7 @@ public:
 			for (const TPair<int32, FPoint2D>& Vertex : Vertices)
 			{
 				//F3DDebugSession _(TEXT("Vertex"));
-				DisplayPoint(Vertex.Value * DisplayScale, EVisuProperty::YellowPoint, Vertex.Key);
+				DisplayPoint2DWithScale(Vertex.Value , EVisuProperty::YellowPoint, Vertex.Key);
 			}
 			//Wait();
 		}
diff --git a/Engine/Source/Runtime/Datasmith/CADKernel/Private/CADKernel/Mesh/Meshers/IsoTriangulator.cpp b/Engine/Source/Runtime/Datasmith/CADKernel/Private/CADKernel/Mesh/Meshers/IsoTriangulator.cpp
--- a/Engine/Source/Runtime/Datasmith/CADKernel/Private/CADKernel/Mesh/Meshers/IsoTriangulator.cpp
+++ b/Engine/Source/Runtime/Datasmith/CADKernel/Private/CADKernel/Mesh/Meshers/IsoTriangulator.cpp
@@ -94,8 +94,8 @@ bool FIsoTriangulator::Triangulate()
 
 #ifdef DEBUG_ISOTRIANGULATOR
 	DisplayLoops(TEXT("FIsoTrianguler::LoopSegments after fix intersection"), true, false);
-	Grid.DisplayIsoSegments(TEXT("FIsoTrianguler::LoopSegments orientation"), DisplaySpace, LoopSegments, false, true, EVisuProperty::YellowCurve);
-	Grid.DisplayIsoSegments(TEXT("FIsoTrianguler::LoopSegments"), DisplaySpace, LoopSegments, true, false, EVisuProperty::YellowCurve);
+	Grid.DisplayIsoSegments(TEXT("FIsoTrianguler::LoopSegments orientation"), DisplaySpace, LoopSegments, false, true, false, EVisuProperty::YellowCurve);
+	Grid.DisplayIsoSegments(TEXT("FIsoTrianguler::LoopSegments"), DisplaySpace, LoopSegments, true, false, false, EVisuProperty::YellowCurve);
 	Wait(false);
 #endif
 
@@ -117,7 +117,7 @@ bool FIsoTriangulator::Triangulate()
 #ifdef DEBUG_ISOTRIANGULATOR
 	{
 		F3DDebugSession A(bDisplay, TEXT("FIsoTrianguler::FinalInnerSegments"));
-		Grid.DisplayIsoSegments(DisplaySpace, FinalInnerSegments, true, false, EVisuProperty::BlueCurve);
+		Grid.DisplayIsoSegments(DisplaySpace, FinalInnerSegments, true, false, false, EVisuProperty::BlueCurve);
 	}
 	InnerToOuterSegmentsIntersectionTool.Display(bDisplay, TEXT("FIsoTrianguler::IntersectionTool InnerToOuter"), EVisuProperty::RedCurve);
 	Chronos.TriangulateDuration1 = FChrono::Elapse(StartTime);
@@ -143,14 +143,15 @@ bool FIsoTriangulator::Triangulate()
 	// 	           - find the shortest segment to connect each connected loop by Delaunay
 	// =============================================================================================================
 #ifdef DEBUG_ISOTRIANGULATOR
-	Grid.DisplayIsoSegments(TEXT("FIsoTrianguler::Final To Loops Before"), DisplaySpace, FinalToLoops, false, false, EVisuProperty::YellowCurve);
+	Grid.DisplayIsoSegments(TEXT("FIsoTrianguler::Final To Loops Before"), DisplaySpace, FinalToLoops, false, false, false, EVisuProperty::YellowCurve);
 	//Wait(bDisplay);
 #endif
 
 	ConnectCellLoops();
 
 #ifdef DEBUG_ISOTRIANGULATOR
-	Grid.DisplayIsoSegments(TEXT("FIsoTrianguler::Final Iso ToLink Inner To Loops"), DisplaySpace, FinalToLoops, false, false, EVisuProperty::YellowCurve);
+	Grid.DisplayIsoSegments(TEXT("FIsoTrianguler::Final Iso ToLink Inner To Loops"), DisplaySpace, FinalToLoops, false, false, true, EVisuProperty::YellowCurve);
+	Wait(bDisplay);
 #endif
 
 	// =============================================================================================================
@@ -524,7 +525,7 @@ void FIsoTriangulator::BuildInnerSegments()
 #ifdef DEBUG_BUILDINNERSEGMENTS
 			if (bDisplay)
 			{
-				DisplaySegment(Grid.GetInner2DPoint(EGridSpace::UniformScaled, IndexNode1) * DisplayScale, Grid.GetInner2DPoint(EGridSpace::UniformScaled, IndexNode2) * DisplayScale, IndexNode1, EVisuProperty::PinkCurve);
+				DisplaySegmentWithScale(Grid.GetInner2DPoint(EGridSpace::UniformScaled, IndexNode1), Grid.GetInner2DPoint(EGridSpace::UniformScaled, IndexNode2), IndexNode1, EVisuProperty::PinkCurve);
 			}
 #endif
 			AddToInnerToOuterSegmentsIntersectionTool(IndexNode1, IndexNode2, InType);
@@ -537,7 +538,7 @@ void FIsoTriangulator::BuildInnerSegments()
 			if (bDisplay)
 			{
 				F3DDebugSession _(bDisplay, TEXT("Segment"));
-				DisplaySegment(Grid.GetInner2DPoint(EGridSpace::UniformScaled, IndexNode1) * DisplayScale, Grid.GetInner2DPoint(EGridSpace::UniformScaled, IndexNode2) * DisplayScale, IndexNode1, EVisuProperty::RedCurve);
+				DisplaySegmentWithScale(Grid.GetInner2DPoint(EGridSpace::UniformScaled, IndexNode1), Grid.GetInner2DPoint(EGridSpace::UniformScaled, IndexNode2), IndexNode1, EVisuProperty::RedCurve);
 			}
 #endif
 			return;
@@ -548,7 +549,7 @@ void FIsoTriangulator::BuildInnerSegments()
 #ifdef DEBUG_BUILDINNERSEGMENTS
 			if (bDisplay)
 			{
-				DisplaySegment(Grid.GetInner2DPoint(EGridSpace::UniformScaled, IndexNode1) * DisplayScale, Grid.GetInner2DPoint(EGridSpace::UniformScaled, IndexNode2) * DisplayScale, IndexNode1, EVisuProperty::PinkCurve);
+				DisplaySegmentWithScale(Grid.GetInner2DPoint(EGridSpace::UniformScaled, IndexNode1), Grid.GetInner2DPoint(EGridSpace::UniformScaled, IndexNode2), IndexNode1, EVisuProperty::PinkCurve);
 			}
 #endif
 			if (LoopSegmentsIntersectionTool.DoesIntersect(Grid.GetInner2DPoint(EGridSpace::UniformScaled, IndexNode1), Grid.GetInner2DPoint(EGridSpace::UniformScaled, IndexNode2))
@@ -569,7 +570,7 @@ void FIsoTriangulator::BuildInnerSegments()
 #ifdef DEBUG_BUILDINNERSEGMENTS
 			if (bDisplay)
 			{
-				DisplaySegment(Grid.GetInner2DPoint(EGridSpace::UniformScaled, IndexNode1) * DisplayScale, Grid.GetInner2DPoint(EGridSpace::UniformScaled, IndexNode2) * DisplayScale, IndexNode1, EVisuProperty::BlueCurve);
+				DisplaySegmentWithScale(Grid.GetInner2DPoint(EGridSpace::UniformScaled, IndexNode1), Grid.GetInner2DPoint(EGridSpace::UniformScaled, IndexNode2), IndexNode1, EVisuProperty::BlueCurve);
 			}
 #endif
 			AddToInnerSegments(IndexNode1, IndexNode2, InType);
@@ -581,7 +582,7 @@ void FIsoTriangulator::BuildInnerSegments()
 #ifdef DEBUG_BUILDINNERSEGMENTS
 			if (bDisplay)
 			{
-				DisplaySegment(Grid.GetInner2DPoint(EGridSpace::UniformScaled, IndexNode1) * DisplayScale, Grid.GetInner2DPoint(EGridSpace::UniformScaled, IndexNode2) * DisplayScale, IndexNode1, EVisuProperty::GreenCurve);
+				DisplaySegmentWithScale(Grid.GetInner2DPoint(EGridSpace::UniformScaled, IndexNode1), Grid.GetInner2DPoint(EGridSpace::UniformScaled, IndexNode2), IndexNode1, EVisuProperty::GreenCurve);
 			}
 #endif
 			return;
@@ -590,7 +591,7 @@ void FIsoTriangulator::BuildInnerSegments()
 #ifdef DEBUG_BUILDINNERSEGMENTS
 		if (bDisplay)
 		{
-			DisplaySegment(Grid.GetInner2DPoint(EGridSpace::UniformScaled, IndexNode1) * DisplayScale, Grid.GetInner2DPoint(EGridSpace::UniformScaled, IndexNode2) * DisplayScale, IndexNode1, EVisuProperty::YellowCurve);
+			DisplaySegmentWithScale(Grid.GetInner2DPoint(EGridSpace::UniformScaled, IndexNode1), Grid.GetInner2DPoint(EGridSpace::UniformScaled, IndexNode2), IndexNode1, EVisuProperty::YellowCurve);
 		}
 #endif
 		AddToInnerToOuterSegmentsIntersectionTool(IndexNode1, IndexNode2, InType);
@@ -1229,38 +1230,38 @@ void FIsoTriangulator::ConnectCellLoops()
 		ConnectCellSubLoopsByNeighborhood(Cell);
 
 #ifdef DEBUG_CONNECT_CELL_LOOPS
-		Grid.DisplayIsoSegments(TEXT("ConnectCellLoops::Step 0"), EGridSpace::UniformScaled, Cell.FinalSegments, false, false, EVisuProperty::YellowCurve);
+		Grid.DisplayIsoSegments(TEXT("ConnectCellLoops::Step 0"), EGridSpace::UniformScaled, Cell.FinalSegments, false, false, false, EVisuProperty::YellowCurve);
 #endif
 
 		if (Cell.bHasOuterLoop)
 		{
 			FindIsoSegmentToLinkOuterLoopNodes(Cell);
 #ifdef DEBUG_CONNECT_CELL_LOOPS
-			Grid.DisplayIsoSegments(TEXT("ConnectCellLoops::Step 1"), EGridSpace::UniformScaled, Cell.FinalSegments, false, false, EVisuProperty::YellowCurve);
+			Grid.DisplayIsoSegments(TEXT("ConnectCellLoops::Step 1"), EGridSpace::UniformScaled, Cell.FinalSegments, false, false, false, EVisuProperty::YellowCurve);
 #endif
 
 			if (Cell.CandidateSegments.Num() == 0)
 			{
 				FindSegmentToLinkOuterLoopNodes(Cell);
 #ifdef DEBUG_CONNECT_CELL_LOOPS
-				Grid.DisplayIsoSegments(TEXT("ConnectCellLoops::Step 2"), EGridSpace::UniformScaled, Cell.FinalSegments, false, false, EVisuProperty::YellowCurve);
+				Grid.DisplayIsoSegments(TEXT("ConnectCellLoops::Step 2"), EGridSpace::UniformScaled, Cell.FinalSegments, false, false, false, EVisuProperty::YellowCurve);
 #endif
 
 				FindSegmentToLinkOuterToInnerLoopNodes(Cell);
 #ifdef DEBUG_CONNECT_CELL_LOOPS
-				Grid.DisplayIsoSegments(TEXT("ConnectCellLoops::Step 3"), EGridSpace::UniformScaled, Cell.FinalSegments, false, false, EVisuProperty::YellowCurve);
+				Grid.DisplayIsoSegments(TEXT("ConnectCellLoops::Step 3"), EGridSpace::UniformScaled, Cell.FinalSegments, false, false, false, EVisuProperty::YellowCurve);
 #endif
 			}
 		}
 		ConnectCellCornerToInnerLoop(Cell);
 #ifdef DEBUG_CONNECT_CELL_LOOPS
-		Grid.DisplayIsoSegments(TEXT("ConnectCellLoops::Step 4"), EGridSpace::UniformScaled, Cell.FinalSegments, false, false, EVisuProperty::YellowCurve);
+		Grid.DisplayIsoSegments(TEXT("ConnectCellLoops::Step 4"), EGridSpace::UniformScaled, Cell.FinalSegments, false, false, false, EVisuProperty::YellowCurve);
 #endif
 
 		FinalToLoops.Append(Cell.FinalSegments);
 
 #ifdef DEBUG_CONNECT_CELL_LOOPS
-		Grid.DisplayIsoSegments(TEXT("ConnectCellLoops::Final Iso ToLink Inner To Loops"), EGridSpace::UniformScaled, FinalToLoops, false, false, EVisuProperty::YellowCurve);
+		Grid.DisplayIsoSegments(TEXT("ConnectCellLoops::Final Iso ToLink Inner To Loops"), EGridSpace::UniformScaled, FinalToLoops, false, false, false, EVisuProperty::YellowCurve);
 #endif
 	}
 }
@@ -3234,14 +3235,6 @@ void FIsoTriangulator::TryToConnectTwoSubLoopsWithTheMostIsoSegment(FCell& Cell,
 bool FIsoTriangulator::TryToCreateSegment(FCell& Cell, FLoopNode* NodeA, const FPoint2D& ACoordinates, FIsoNode* NodeB, const FPoint2D& BCoordinates, const double FlatAngle)
 {
 
-#ifdef DEBUG_TRY_TO_CREATE_SEGMENT
-	{
-		F3DDebugSession _(TEXT("Test"));
-		DisplaySegment(NodeA->Get2DPoint(EGridSpace::UniformScaled, Grid), NodeB->Get2DPoint(EGridSpace::UniformScaled, Grid), 0, EVisuProperty::RedCurve);
-		//Wait(false);
-	}
-#endif
-
 	if (NodeA->GetSegmentConnectedTo(NodeB))
 	{
 		return false;
@@ -3267,6 +3260,14 @@ bool FIsoTriangulator::TryToCreateSegment(FCell& Cell, FLoopNode* NodeA, const F
 		return false;
 	}
 
+#ifdef DEBUG_TRY_TO_CREATE_SEGMENT
+	{
+		F3DDebugSession _(TEXT("Test"));
+		DisplaySegmentWithScale(NodeA->Get2DPoint(EGridSpace::UniformScaled, Grid), NodeB->Get2DPoint(EGridSpace::UniformScaled, Grid), 0, EVisuProperty::RedCurve);
+		Wait(false);
+	}
+#endif
+
 	// Is Outside and not too flat at NodeA
 	if (NodeA->IsSegmentBeInsideFace(BCoordinates, Grid, FlatAngle))
 	{
diff --git a/Engine/Source/Runtime/Datasmith/CADKernel/Private/CADKernel/Mesh/Meshers/IsoTriangulator.h b/Engine/Source/Runtime/Datasmith/CADKernel/Private/CADKernel/Mesh/Meshers/IsoTriangulator.h
--- a/Engine/Source/Runtime/Datasmith/CADKernel/Private/CADKernel/Mesh/Meshers/IsoTriangulator.h
+++ b/Engine/Source/Runtime/Datasmith/CADKernel/Private/CADKernel/Mesh/Meshers/IsoTriangulator.h
@@ -11,7 +11,7 @@
 #include "CADKernel/UI/Visu.h"
 
 #ifdef CADKERNEL_DEV
-#include "CADKernel/Mesh/Meshers/IsoTriangulator/DefineForDebug.h"
+#include "CADKernel/UI/DefineForDebug.h"
 #include "CADKernel/Mesh/Meshers/MesherReport.h"
 #endif
 
diff --git a/Engine/Source/Runtime/Datasmith/CADKernel/Private/CADKernel/Mesh/Meshers/IsoTriangulator/IntersectionSegmentTool.h b/Engine/Source/Runtime/Datasmith/CADKernel/Private/CADKernel/Mesh/Meshers/IsoTriangulator/IntersectionSegmentTool.h
--- a/Engine/Source/Runtime/Datasmith/CADKernel/Private/CADKernel/Mesh/Meshers/IsoTriangulator/IntersectionSegmentTool.h
+++ b/Engine/Source/Runtime/Datasmith/CADKernel/Private/CADKernel/Mesh/Meshers/IsoTriangulator/IntersectionSegmentTool.h
@@ -15,11 +15,19 @@ class FIsoInnerNode;
 class FIsoSegment;
 class FPoint2D;
 
+enum class EConnectionType : uint8
+{
+	DoesntStartFrom = 0,
+	StartFrom,
+	SuperimposedByOrOn,
+	SameSegment,
+};
+
 namespace IntersectionToolBase
 {
 struct FSegment
 {
-	const TSegment<FPoint2D> Segment2D;
+	const FSegment2D Segment2D;
 
 	/**
 	 * Segment's axis aligned bounding box
@@ -56,34 +64,74 @@ struct FSegment
 		return nullptr;
 	}
 
-	bool DoesItStartFrom(const FIsoNode* StartNode, const FIsoNode* EndNode) const
+	EConnectionType IsSuperimposed(const FSegment2D& SegmentAB, const FSegment2D& SegmentCD, bool bSameOrientation) const
 	{
-		if (GetFirstNode() == StartNode || GetSecondNode() == StartNode)
+		const FPoint2D AB = SegmentAB.GetVector();
+		const FPoint2D CD = SegmentCD.GetVector(); 
+		const double ParallelCoef = AB ^ CD;
+		if (FMath::IsNearlyZero(ParallelCoef, DOUBLE_KINDA_SMALL_NUMBER))
 		{
-			return true;
+			const double OrientationCoef = AB * CD;
+			if ((OrientationCoef >= 0) == bSameOrientation)
+			{
+				return EConnectionType::SuperimposedByOrOn;
+			}
 		}
+		return EConnectionType::StartFrom;
+	};
 
-		if (GetFirstNode() == EndNode || GetSecondNode() == EndNode)
+	EConnectionType DoesItStartFromAndSuperimposed(const FIsoNode* StartNode, const FPoint2D* EndPoint, const FSegment2D& InSegment) const
+	{
+		if (GetFirstNode() == StartNode)
 		{
-			return true;
+			constexpr bool bSameOrientation = true;
+			return IsSuperimposed(Segment2D, InSegment, bSameOrientation);
+		}
+		if (GetSecondNode() == StartNode)
+		{
+			constexpr bool bNotSameOrientation = true;
+			return IsSuperimposed(Segment2D, InSegment, bNotSameOrientation);
 		}
 
-		return false;
+		return EConnectionType::DoesntStartFrom;
 	}
 
-	bool DoesItStartFrom(const FIsoNode* StartNode, const FPoint2D* EndPoint) const
+	EConnectionType DoesItStartFromAndSuperimposed(const FPoint2D* StartPoint, const FPoint2D* EndPoint, const FSegment2D& InSegment) const
 	{
-		if (GetFirstNode() == StartNode || GetSecondNode() == StartNode)
+		return EConnectionType::DoesntStartFrom;
+	}
+
+	EConnectionType DoesItStartFromAndSuperimposed(const FIsoNode* StartNode, const FIsoNode* EndNode, const FSegment2D& InSegment) const
+	{
+		if (GetFirstNode() == StartNode)
 		{
-			return true;
+			if (GetSecondNode() == EndNode)
+			{
+				return EConnectionType::SameSegment;
+			}
+			return IsSuperimposed(Segment2D, InSegment, true);
 		}
 
-		return false;
-	}
+		if (GetSecondNode() == EndNode)
+		{
+			return IsSuperimposed(Segment2D, InSegment, true);
+		}
 
-	bool DoesItStartFrom(const FPoint2D* StartPoint, const FPoint2D* EndPoint) const
-	{
-		return false;
+		if (GetFirstNode() == EndNode )
+		{
+			if (GetSecondNode() == StartNode)
+			{
+				return EConnectionType::SameSegment;
+			}
+			return IsSuperimposed(Segment2D, InSegment, false);
+		}
+
+		if (GetSecondNode() == StartNode)
+		{
+			return IsSuperimposed(Segment2D, InSegment, false);
+		}
+
+		return EConnectionType::DoesntStartFrom;
 	}
 
 	bool IsFullyBefore(const FSegment& Segment) const
@@ -262,11 +310,18 @@ public:
 				{
 					break;
 				}
-}
+			}
 
-			if (Segment.DoesItStartFrom(StartExtremity, EndExtremity))
+			switch (Segment.DoesItStartFromAndSuperimposed(StartExtremity, EndExtremity, InSegment.Segment2D))
 			{
+			case EConnectionType::SameSegment:
+			case EConnectionType::StartFrom:
 				continue;
+			case EConnectionType::SuperimposedByOrOn:
+				return Segment.IsoSegment;
+			case EConnectionType::DoesntStartFrom:
+			default:
+				break;
 			}
 
 			if (Segment.DoesItIntersect(InSegment))
@@ -306,11 +361,24 @@ public:
 				{
 					break;
 				}
-}
+			}
 
-			if (Segment.DoesItStartFrom(StartExtremity, EndExtremity))
+			switch (Segment.DoesItStartFromAndSuperimposed(StartExtremity, EndExtremity, InSegment.Segment2D))
 			{
+			case EConnectionType::StartFrom:
+			case EConnectionType::SameSegment:
+				continue;
+
+			case EConnectionType::SuperimposedByOrOn:
+				++IntersectionCount;
+				if (OutIntersectedSegments)
+				{
+					OutIntersectedSegments->Add(Segment.GetIsoSegment());
+				}
 				continue;
+			case EConnectionType::DoesntStartFrom:
+			default:
+				break;
 			}
 
 			if (Segment.DoesItIntersect(InSegment))
@@ -326,7 +394,7 @@ public:
 		return IntersectionCount;
 	}
 
-#ifdef CADKERNEL_DEV
+#ifdef CADKERNEL_DEBUG
 	virtual void Display(bool bDisplay, const TCHAR* Message, EVisuProperty Property = EVisuProperty::BlueCurve) const
 	{
 		if (!bDisplay)
diff --git a/Engine/Source/Runtime/Datasmith/CADKernel/Private/CADKernel/Mesh/Meshers/IsoTriangulator/IsoCell.h b/Engine/Source/Runtime/Datasmith/CADKernel/Private/CADKernel/Mesh/Meshers/IsoTriangulator/IsoCell.h
--- a/Engine/Source/Runtime/Datasmith/CADKernel/Private/CADKernel/Mesh/Meshers/IsoTriangulator/IsoCell.h
+++ b/Engine/Source/Runtime/Datasmith/CADKernel/Private/CADKernel/Mesh/Meshers/IsoTriangulator/IsoCell.h
@@ -1,6 +1,7 @@
 // Copyright Epic Games, Inc. All Rights Reserved.
 #pragma once
 
+#include "CADKernel/Core/Factory.h"
 #include "CADKernel/Core/Types.h"
 #include "CADKernel/Mesh/MeshEnum.h"
 #include "CADKernel/Mesh/Meshers/IsoTriangulator/IsoNode.h"
diff --git a/Engine/Source/Runtime/Datasmith/CADKernel/Private/CADKernel/Mesh/Meshers/IsoTriangulator/IsoSegment.h b/Engine/Source/Runtime/Datasmith/CADKernel/Private/CADKernel/Mesh/Meshers/IsoTriangulator/IsoSegment.h
--- a/Engine/Source/Runtime/Datasmith/CADKernel/Private/CADKernel/Mesh/Meshers/IsoTriangulator/IsoSegment.h
+++ b/Engine/Source/Runtime/Datasmith/CADKernel/Private/CADKernel/Mesh/Meshers/IsoTriangulator/IsoSegment.h
@@ -200,7 +200,8 @@ public:
 
 	bool HasTriangleOnRightAndLeft()
 	{
-		return (States & EIsoSegmentStates::RightTriangle) == EIsoSegmentStates::RightTriangle && (States & EIsoSegmentStates::LeftTriangle) == EIsoSegmentStates::LeftTriangle;
+		constexpr EIsoSegmentStates BoothSides = EIsoSegmentStates::RightTriangle | EIsoSegmentStates::LeftTriangle;
+		return (States & BoothSides) == BoothSides;
 	}
 
 	void SetHasTriangleOnRight()
diff --git a/Engine/Source/Runtime/Datasmith/CADKernel/Private/CADKernel/Mesh/Meshers/ParametricMesher.cpp b/Engine/Source/Runtime/Datasmith/CADKernel/Private/CADKernel/Mesh/Meshers/ParametricMesher.cpp
--- a/Engine/Source/Runtime/Datasmith/CADKernel/Private/CADKernel/Mesh/Meshers/ParametricMesher.cpp
+++ b/Engine/Source/Runtime/Datasmith/CADKernel/Private/CADKernel/Mesh/Meshers/ParametricMesher.cpp
@@ -9,6 +9,7 @@
 #include "CADKernel/Mesh/Structure/EdgeMesh.h"
 #include "CADKernel/Mesh/Structure/FaceMesh.h"
 #include "CADKernel/Mesh/Structure/ModelMesh.h"
+#include "CADKernel/Mesh/Structure/ThinZone2DFinder.h"
 #include "CADKernel/Mesh/Structure/VertexMesh.h"
 #include "CADKernel/Topo/Body.h"
 #include "CADKernel/Topo/Model.h"
@@ -16,7 +17,9 @@
 #include "CADKernel/Topo/TopologicalFace.h"
 #include "CADKernel/Utils/Util.h"
 
+#ifndef CADKERNEL_DEV
 #include "Async/ParallelFor.h"
+#endif
 
 #ifdef CADKERNEL_DEV
 #include "CADKernel/Mesh/Meshers/MesherReport.h"
@@ -29,6 +32,7 @@ namespace UE::CADKernel
 
 FParametricMesher::FParametricMesher(FModelMesh& InMeshModel)
 	: MeshModel(InMeshModel)
+	, bThinZoneMeshing(false)
 {
 }
 
@@ -101,22 +105,29 @@ void FParametricMesher::PreMeshingTasks()
 
 	const TArray<TSharedPtr<FCriterion>>& Criteria = GetMeshModel().GetCriteria();
 
-#ifdef CADKERNEL_DEV
-	for(FTopologicalFace* Face : Faces)
-	{
-#else
+	// ============================================================================================================
+	//      Apply Surface Criteria
+	// ============================================================================================================
+
 	TArray<FTopologicalFace*>& LocalFaces = Faces;
-	ParallelFor(Faces.Num(), [&LocalFaces, &Criteria](int32 Index)
-		{
-			FTopologicalFace* Face = LocalFaces[Index];
+#ifndef CADKERNEL_DEV
+	ParallelFor(Faces.Num(), [&LocalFaces, &Criteria, &bThinZone = bThinZoneMeshing](int32 Index)
+#else
+	bool bThinZone = bThinZoneMeshing;
+	for (int32 Index = 0; Index < Faces.Num(); ++Index)
 #endif
-			if (Face == nullptr || Face->IsDeleted() || Face->IsMeshed())
-			{
-				return;
-			}
-			ApplyFaceCriteria(*Face, Criteria);
+	{
+		FTopologicalFace* Face = LocalFaces[Index];
+		if (Face == nullptr || Face->IsNotMeshable())
+		{
+			return;
+		}
+		ApplyFaceCriteria(*Face, Criteria, bThinZone);
+		if (!Face->IsDeletedOrDegenerated())
+		{
 			Face->ComputeSurfaceSideProperties();
 		}
+	}
 #ifndef CADKERNEL_DEV
 	);
 #endif
@@ -175,7 +186,7 @@ void FParametricMesher::MeshEntities()
 
 }
 
-void FParametricMesher::ApplyFaceCriteria(FTopologicalFace& Face, const TArray<TSharedPtr<FCriterion>>& Criteria)
+void FParametricMesher::ApplyFaceCriteria(FTopologicalFace& Face, const TArray<TSharedPtr<FCriterion>>& Criteria, bool bThinZoneMeshing)
 {
 	if (Face.IsApplyCriteria())
 	{
@@ -192,6 +203,38 @@ void FParametricMesher::ApplyFaceCriteria(FTopologicalFace& Face, const TArray<T
 
 	Face.InitDeltaUs();
 	Face.ApplyCriteria(Criteria, Grid);
+
+	if (bThinZoneMeshing)
+	{
+		Grid.ScaleGrid();
+
+#ifdef DEBUG_ONLY_SURFACE_TO_DEBUG
+		if (Grid.bDisplay)
+		{
+			Grid.DisplayGridPoints(EGridSpace::Default2D);
+			Grid.DisplayGridPoints(EGridSpace::UniformScaled);
+			Grid.DisplayGridPoints(EGridSpace::Scaled);
+		}
+#endif
+		FThinZone2DFinder ThinZoneFinder(Grid, Face);
+
+		// Size (length of segment of the loop sampling) is equal to MinimalElementLength / ElementRatio
+		// With this ratio each edges of the mesh should be defined by at least 3 segments. 
+		// This should ensure to identified all thin zones according to the mesh size and minimizing the size of the loop sampling
+		constexpr double ElementRatio = 3.;
+		const double Size = Face.GetEstimatedMinimalElementLength() / ElementRatio;
+		const bool bHasThinZones = ThinZoneFinder.SearchThinZones(Size);
+		if (bHasThinZones)
+		{
+			Face.SetHasThinZone();
+			Face.MoveThinZones(ThinZoneFinder.GetThinZones());
+		}
+	}
+
+	if (Face.IsDegenerated())
+	{
+		Face.Remove();
+	}
 }
 
 void FParametricMesher::ApplyEdgeCriteria(FTopologicalEdge& Edge)
@@ -231,7 +274,7 @@ void FParametricMesher::Mesh(FTopologicalFace& Face)
 	bDisplay = (Face.GetId() == FaceToDebug);
 #endif
 
-	if (Face.IsDeleted() || Face.IsMeshed())
+	if (Face.IsNotMeshable())
 	{
 		return;
 	}
@@ -249,20 +292,18 @@ void FParametricMesher::Mesh(FTopologicalFace& Face)
 	FTimePoint GenerateCloudStartTime = FChrono::Now();
 
 	FGrid Grid(Face, MeshModel);
-	GenerateCloud(Grid);
-
-	FDuration GenerateCloudDuration = FChrono::Elapse(GenerateCloudStartTime);
-
-	if (Grid.IsDegenerated())
+	if (!GenerateCloud(Grid) || Grid.IsDegenerated())
 	{
 #ifdef CADKERNEL_DEV
 		MesherReport.Logs.AddDegeneratedGrid();
- #endif
+#endif
 		FMessage::Printf(EVerboseLevel::Log, TEXT("The meshing of the surface %d failed due to a degenerated grid\n"), Face.GetId());
 		Face.SetMeshed();
 		return;
 	}
 
+	FDuration GenerateCloudDuration = FChrono::Elapse(GenerateCloudStartTime);
+
 	FTimePoint IsoTriangulerStartTime = FChrono::Now();
 
 	TSharedRef<FFaceMesh> SurfaceMesh = StaticCastSharedRef<FFaceMesh>(Face.GetOrCreateMesh(MeshModel));
@@ -297,49 +338,26 @@ void FParametricMesher::Mesh(FTopologicalFace& Face)
 
 }
 
-void FParametricMesher::GenerateCloud(FGrid& Grid)
+bool FParametricMesher::GenerateCloud(FGrid& Grid)
 {
 	Grid.DefineCuttingParameters();
 	if (!Grid.GeneratePointCloud())
 	{
-		return;
+		return false;
 	}
 
-	bool bFindThinZone = false;
-	if (bFindThinZone)
+	if (bThinZoneMeshing)
 	{
 		FTimePoint StartTime = FChrono::Now();
-#ifdef THIN_ZONES
-		Grid.SearchThinZones();
-#endif
-
 		if (Grid.GetFace().HasThinZone())
 		{
-#ifdef DEBUG_THIN_ZONES
-			{
-				F3DDebugSession _(TEXT("Thin Surface"));
-				Display(Grid.GetSurface());
-			}
-#endif
-			FTimePoint MeshThinZonesTime = FChrono::Now();
-			MeshThinZoneEdges(Grid);
-#ifdef CADKERNEL_DEV
-			MesherReport.Chronos.GlobalMeshThinZones += FChrono::Elapse(MeshThinZonesTime);
-#endif
+			MeshThinZoneEdges(Grid.GetFace());
 		}
 #ifdef CADKERNEL_DEV
 		MesherReport.Chronos.GlobalThinZones += FChrono::Elapse(StartTime);
 #endif
 	}
 
-#ifdef DEBUG_THIN_ZONES
-	if(Grid.bDisplay)
-	{
-		Grid.DisplayInnerDomainPoints(TEXT("FGrid::PointCloud 2D"), Grid.GetInner2DPoints(EGridSpace::Default2D));
-		Wait();
-	}
-#endif
-
 	FTimePoint StartTime = FChrono::Now();
 	MeshFaceLoops(Grid);
 
@@ -348,6 +366,8 @@ void FParametricMesher::GenerateCloud(FGrid& Grid)
 #ifdef CADKERNEL_DEV
 	MesherReport.Chronos.GlobalMeshAndGetLoopNodes += FChrono::Elapse(StartTime);
 #endif
+
+	return true;
 }
 
 void FParametricMesher::MeshFaceLoops(FGrid& Grid)
@@ -384,6 +404,15 @@ static void FillImposedIsoCuttingPoints(TArray<double>& UEdgeSetOfIntersectionWi
 			continue;
 		}
 
+		// Remove coordinate inside thin zone
+		for (FLinearBoundary ThinZone : Edge.GetThinZoneBounds())
+		{
+			if (ThinZone.Contains(InterU))
+			{
+				continue;
+			}
+		}
+
 		// Remove nearly duplicate 
 		if (InterU - PreviousU < EdgeToleranceGeo)
 		{
@@ -456,7 +485,7 @@ void FParametricMesher::Mesh(FTopologicalEdge& InEdge, const FTopologicalFace& F
 			// In this case, the loop is flat i.e. in 2d the meshes of the 2d line and 2d curve are coincident
 			// So the grid is degenerated and the surface is not meshed
 			// to avoid this case, the Edge is virtually meshed i.e. the nodes inside the edge have the id of the mesh of the vertices.
-			InEdge.SetAsVirtuallyMeshed();
+			InEdge.SetVirtuallyMeshedMarker();
 		}
 
 		if (ActiveEdge.IsThinPeak())
@@ -481,7 +510,7 @@ void FParametricMesher::Mesh(FTopologicalEdge& InEdge, const FTopologicalFace& F
 	ApplyEdgeCriteria(InEdge);
 
 #ifdef DEBUG_MESH_EDGE
-	if(bDisplay)
+	if (bDisplay)
 	{
 		F3DDebugSession _(FString::Printf(TEXT("EdgePointsOnDomain %d"), InEdge.GetId()));
 		Display2D(InEdge);
@@ -518,16 +547,33 @@ void FParametricMesher::Mesh(FTopologicalEdge& InEdge, const FTopologicalFace& F
 		}
 		Wait();
 	}
+	{
+		F3DDebugSession _(FString::Printf(TEXT("Thin Zone")));
+		for (FLinearBoundary ThinZone : InEdge.GetThinZoneBounds())
+		{
+			TArray<double> Coords;
+			Coords.Add(ThinZone.GetMin());
+			Coords.Add(ThinZone.GetMax());
+			TArray<FPoint2D> ThinZone2D;
+			InEdge.Approximate2DPoints(Coords, ThinZone2D);
+			DisplaySegment(ThinZone2D[0], ThinZone2D[1], EVisuProperty::YellowCurve);
+		}
+		Wait();
+
+	}
 #endif
 
 	FLinearBoundary EdgeBounds = InEdge.GetBoundary();
 
 	TArray<double>& DeltaUs = InEdge.GetDeltaUMaxs();
 
+	InEdge.SortImposedCuttingPoints();
+	const TArray<FImposedCuttingPoint>& EdgeImposedCuttingPoints = InEdge.GetImposedCuttingPoints();
+
 	// build a edge mesh compiling inner surface cutting (based on criteria applied on the surface) and edge cutting (based on criteria applied on the curve)
 	TArray<FCuttingPoint> ImposedIsoCuttingPoints;
 	{
-		int32 NbImposedCuttingPoints = InEdge.GetImposedCuttingPoints().Num() + EdgeIntersectionWithIsoU_Coordinates.Num() + EdgeIntersectionWithIsoV_Coordinates.Num() + 2;
+		int32 NbImposedCuttingPoints = EdgeImposedCuttingPoints.Num() + EdgeIntersectionWithIsoU_Coordinates.Num() + EdgeIntersectionWithIsoV_Coordinates.Num() + 2;
 		ImposedIsoCuttingPoints.Reserve(NbImposedCuttingPoints);
 	}
 
@@ -536,26 +582,12 @@ void FParametricMesher::Mesh(FTopologicalEdge& InEdge, const FTopologicalFace& F
 	ImposedIsoCuttingPoints.Emplace(EdgeBounds.GetMin(), ECoordinateType::VertexCoordinate, -1, ExtremityTolerances[0]);
 	ImposedIsoCuttingPoints.Emplace(EdgeBounds.GetMax(), ECoordinateType::VertexCoordinate, -1, ExtremityTolerances[1]);
 
-	double MinDeltaU = HUGE_VALUE;
-	for (const double& DeltaU : DeltaUs)
-	{
-		MinDeltaU = FMath::Min(MinDeltaU, DeltaU);
-	}
-
-#ifdef DEBUG_MESH_EDGE
-	if (bDisplay)
+	int32 Index = 0;
+	for (const FImposedCuttingPoint& CuttingPoint : EdgeImposedCuttingPoints)
 	{
-		F3DDebugSession _(FString::Printf(TEXT("Edge %d"), InEdge.GetId()));
-		for (int32 Index = 0; Index < InEdge.GetImposedCuttingPoints().Num(); ++Index)
-		{
-			const FImposedCuttingPoint& CuttingPoint = InEdge.GetImposedCuttingPoints()[Index];
-			ImposedIsoCuttingPoints.Emplace(CuttingPoint.Coordinate, ECoordinateType::ImposedCoordinate, CuttingPoint.OppositNodeIndex, MinDeltaU * AThird);
-			FCurvePoint Point;
-			InEdge.EvaluatePoint(CuttingPoint.Coordinate, 0, Point);
-			DisplayPoint(Point.Point);
-		}
+		double CuttingPointDeltaU = InEdge.GetDeltaUFor(CuttingPoint.Coordinate, Index);
+		ImposedIsoCuttingPoints.Emplace(CuttingPoint.Coordinate, ECoordinateType::ImposedCoordinate, CuttingPoint.OppositNodeIndex, CuttingPointDeltaU * AThird);
 	}
-#endif
 
 	// Add Edge intersection with inner surface grid Iso
 	double EdgeTolerance = FMath::Min(ExtremityTolerances[0], ExtremityTolerances[1]);
@@ -571,7 +603,7 @@ void FParametricMesher::Mesh(FTopologicalEdge& InEdge, const FTopologicalFace& F
 
 	ImposedIsoCuttingPoints.Sort([](const FCuttingPoint& Point1, const FCuttingPoint& Point2) { return Point1.Coordinate < Point2.Coordinate; });
 
-	auto MergeImposedCuttingPoints = [&](int32& Index, int32& NewIndex, ECoordinateType NewType)
+	TFunction<void(int32&, int32&, ECoordinateType)> MergeImposedCuttingPoints = [&](int32& Index, int32& NewIndex, ECoordinateType NewType)
 	{
 		double DeltaU = FMath::Max(ImposedIsoCuttingPoints[NewIndex].IsoDeltaU, ImposedIsoCuttingPoints[Index].IsoDeltaU);
 		if (ImposedIsoCuttingPoints[NewIndex].Type <= ImposedCoordinate && ImposedIsoCuttingPoints[Index].Type <= ImposedCoordinate)
@@ -632,17 +664,17 @@ void FParametricMesher::Mesh(FTopologicalEdge& InEdge, const FTopologicalFace& F
 	if (ImposedIsoCuttingPoints.Num() > 1)
 	{
 		int32 NewIndex = 0;
-		for (int32 Index = 1; Index < ImposedIsoCuttingPoints.Num(); ++Index)
+		for (int32 Andex = 1; Andex < ImposedIsoCuttingPoints.Num(); ++Andex)
 		{
-			if (ImposedIsoCuttingPoints[Index].Type > ECoordinateType::ImposedCoordinate)
+			if (ImposedIsoCuttingPoints[Andex].Type > ECoordinateType::ImposedCoordinate)
 			{
 				bool bIsDelete = false;
 				for (const FLinearBoundary& ThinZone : InEdge.GetThinZoneBounds())
 				{
-					if (ThinZone.Contains(ImposedIsoCuttingPoints[Index].Coordinate))
+					if (ThinZone.Contains(ImposedIsoCuttingPoints[Andex].Coordinate))
 					{
 						bIsDelete = true;
-						continue;
+						break; // or copntinue
 					}
 				}
 				if (bIsDelete)
@@ -651,18 +683,18 @@ void FParametricMesher::Mesh(FTopologicalEdge& InEdge, const FTopologicalFace& F
 				}
 			}
 
-			if (ImposedIsoCuttingPoints[NewIndex].Type == ECoordinateType::ImposedCoordinate || ImposedIsoCuttingPoints[Index].Type == ECoordinateType::ImposedCoordinate)
+			if (ImposedIsoCuttingPoints[NewIndex].Type == ECoordinateType::ImposedCoordinate || ImposedIsoCuttingPoints[Andex].Type == ECoordinateType::ImposedCoordinate)
 			{
-				MergeImposedCuttingPoints(Index, NewIndex, ECoordinateType::ImposedCoordinate);
+				MergeImposedCuttingPoints(Andex, NewIndex, ECoordinateType::ImposedCoordinate);
 			}
-			else if (ImposedIsoCuttingPoints[NewIndex].Type != ImposedIsoCuttingPoints[Index].Type)
+			else if (ImposedIsoCuttingPoints[NewIndex].Type != ImposedIsoCuttingPoints[Andex].Type)
 			{
-				MergeImposedCuttingPoints(Index, NewIndex, ECoordinateType::IsoUVCoordinate);
+				MergeImposedCuttingPoints(Andex, NewIndex, ECoordinateType::IsoUVCoordinate);
 			}
 			else
 			{
 				++NewIndex;
-				ImposedIsoCuttingPoints[NewIndex] = ImposedIsoCuttingPoints[Index];
+				ImposedIsoCuttingPoints[NewIndex] = ImposedIsoCuttingPoints[Andex];
 			}
 		}
 		ImposedIsoCuttingPoints.SetNum(NewIndex + 1);
@@ -678,13 +710,22 @@ void FParametricMesher::Mesh(FTopologicalEdge& InEdge, const FTopologicalFace& F
 		ImposedIsoCuttingPoints.Emplace(EdgeBounds.GetMax(), ECoordinateType::VertexCoordinate, -1, InEdge.GetDeltaUMaxs().Last() * AQuarter);
 	}
 
-	// max vertex of the edge
-	int32 MaxNumberOfVertex = (int32)((EdgeBounds.GetMax() - EdgeBounds.GetMin()) / MinDeltaU) + 5;
-
 	// Final array of the edge mesh vertex 
 	TArray<FCuttingPoint>& FinalEdgeCuttingPointCoordinates = InEdge.GetCuttingPoints();
+	{
+		// max count of vertex
+		double MinDeltaU = HUGE_VALUE;
+		for (const double& DeltaU : DeltaUs)
+		{
+			if (DeltaU < MinDeltaU)
+			{
+				MinDeltaU = DeltaU;
+			}
+		}
 
-	FinalEdgeCuttingPointCoordinates.Empty(ImposedIsoCuttingPoints.Num() + MaxNumberOfVertex);
+		int32 MaxNumberOfVertex = FMath::IsNearlyZero(MinDeltaU) ? 5 : (int32)((EdgeBounds.GetMax() - EdgeBounds.GetMin()) / MinDeltaU) + 5;
+		FinalEdgeCuttingPointCoordinates.Empty(ImposedIsoCuttingPoints.Num() + MaxNumberOfVertex);
+	}
 
 #ifdef DEBUG_GETPREFERREDUVCOORDINATESFROMNEIGHBOURS
 	TArray<double> CuttingPoints2;
@@ -717,36 +758,69 @@ void FParametricMesher::Mesh(FTopologicalEdge& InEdge, const FTopologicalFace& F
 	{
 		TArray<double> CuttingPoints;
 		FMesherTools::ComputeFinalCuttingPointsWithImposedCuttingPoints(InEdge.GetCrossingPointUs(), InEdge.GetDeltaUMaxs(), ImposedIsoCuttingPoints, CuttingPoints);
+		int32 ImposedIndex = 0;
+		int32 ImposedIsoCuttingPointsCount = ImposedIsoCuttingPoints.Num();
 		for (const double& Coordinate : CuttingPoints)
 		{
-			FinalEdgeCuttingPointCoordinates.Emplace(Coordinate, ECoordinateType::OtherCoordinate);
+			if (FMath::IsNearlyEqual(ImposedIsoCuttingPoints[ImposedIndex].Coordinate, Coordinate))
+			{
+				FinalEdgeCuttingPointCoordinates.Emplace(ImposedIsoCuttingPoints[ImposedIndex]);
+				++ImposedIndex;
+			}
+			else
+			{
+				while (ImposedIndex < ImposedIsoCuttingPointsCount && ImposedIsoCuttingPoints[ImposedIndex].Coordinate < Coordinate)
+				{
+					++ImposedIndex;
+				}
+				FinalEdgeCuttingPointCoordinates.Emplace(Coordinate, ECoordinateType::OtherCoordinate);
+			}
 		}
 
 #ifdef DEBUG_GETPREFERREDUVCOORDINATESFROMNEIGHBOURS
+		if (InEdge.IsThinZone() && EdgeImposedCuttingPoints.Num())
 		{
 			F3DDebugSession G(TEXT("Mesh(TSharedRef<FEdge> InEdge"));
 			{
 				F3DDebugSession G(TEXT("U From Iso"));
 				for (const FCuttingPoint& CuttingU : ImposedIsoCuttingPoints)
 				{
-					UE::CADKernel::DisplayPoint(FPoint2D(CuttingU.Coordinate, 0.0));
+					if (CuttingU.OppositNodeIndex >= 0)
+					{
+						UE::CADKernel::DisplayPoint(FPoint2D(CuttingU.Coordinate, 0.0), EVisuProperty::RedPoint);
+					}
+					else
+					{
+						UE::CADKernel::DisplayPoint(FPoint2D(CuttingU.Coordinate, 0.0));
+					}
 				}
 			}
+			{
+				F3DDebugSession _(InEdge.GetThinZoneBounds().Num() != 0, FString::Printf(TEXT("Thin Zone"), InEdge.GetId()));
+				for (FLinearBoundary ThinZone : InEdge.GetThinZoneBounds())
+				{
+					UE::CADKernel::DisplayPoint(FPoint2D(ThinZone.GetMin(), 0.01), EVisuProperty::BluePoint);
+					UE::CADKernel::DisplayPoint(FPoint2D(ThinZone.GetMax(), 0.01), EVisuProperty::BluePoint);
+					DisplaySegment(FPoint2D(ThinZone.GetMin(), 0.01), FPoint2D(ThinZone.GetMax(), 0.01), EVisuProperty::BlueCurve);
+				}
+				Wait();
+			}
+
 			{
 				F3DDebugSession G(TEXT("U From Criteria"));
 				for (double CuttingU : CuttingPoints2)
 				{
-					UE::CADKernel::DisplayPoint(FPoint2D(CuttingU, 0.05), EVisuProperty::NonManifoldEdge);
+					UE::CADKernel::DisplayPoint(FPoint2D(CuttingU, 0.02), EVisuProperty::RedPoint);
 				}
 			}
 			{
 				F3DDebugSession G(TEXT("U Final (Criteria & Iso)"));
 				for (double CuttingU : CuttingPoints)
 				{
-					UE::CADKernel::DisplayPoint(FPoint2D(CuttingU, 0.1), EVisuProperty::PurplePoint);
+					UE::CADKernel::DisplayPoint(FPoint2D(CuttingU, 0.04), EVisuProperty::YellowPoint);
 				}
 			}
-			//Wait();
+			Wait();
 		}
 #endif
 
@@ -1064,7 +1138,7 @@ void FParametricMesher::MeshSurfaceByFront(TArray<FCostToFace>& QuadTrimmedSurfa
 
 	for (FTopologicalFace* Face : Faces)
 	{
-		if (Face == nullptr)
+		if (Face == nullptr || Face->IsDeletedOrDegenerated())
 		{
 			continue;
 		}
@@ -1079,6 +1153,8 @@ void FParametricMesher::MeshSurfaceByFront(TArray<FCostToFace>& QuadTrimmedSurfa
 	TArray<FTopologicalFace*> SecondChoiceOfCandidateFacesForMesh; // first in first out
 	SecondChoiceOfCandidateFacesForMesh.Reserve(100);
 
+	static bool bStop = false;
+
 	TFunction<void(FTopologicalFace&)> MeshFace = [&](FTopologicalFace& Face)
 	{
 #ifdef DISPLAYDEBUGMESHFACEBYFACESTEP
@@ -1093,7 +1169,7 @@ void FParametricMesher::MeshSurfaceByFront(TArray<FCostToFace>& QuadTrimmedSurfa
 		{
 			F3DDebugSession A(FString::Printf(TEXT("Mesh %d"), Face.GetId()));
 			DisplayMesh(*Face.GetOrCreateMesh(GetMeshModel()));
-			//Wait();
+			Wait(bStop);
 		}
 #endif
 
@@ -1174,7 +1250,7 @@ void FParametricMesher::MeshSurfaceByFront(TArray<FCostToFace>& QuadTrimmedSurfa
 
 	TFunction<void(FTopologicalFace&)> MeshFacesByFront = [&](FTopologicalFace& Face)
 	{
-		if (Face.IsMeshed())
+		if (Face.IsNotMeshable())
 		{
 			return;
 		}
@@ -1199,7 +1275,7 @@ void FParametricMesher::MeshSurfaceByFront(TArray<FCostToFace>& QuadTrimmedSurfa
 				for (; Index < CandidateFacesForMesh.Num(); ++Index)
 				{
 					FTopologicalFace* CandidateSurface = CandidateFacesForMesh[Index];
-					if (CandidateSurface->IsMeshed())
+					if (CandidateSurface->IsNotMeshable())
 					{
 						CandidateFacesForMesh.RemoveAt(Index);
 						--Index;
@@ -1224,7 +1300,7 @@ void FParametricMesher::MeshSurfaceByFront(TArray<FCostToFace>& QuadTrimmedSurfa
 					for (; Index < CandidateFacesForMesh.Num(); ++Index)
 					{
 						FTopologicalFace* CandidateSurface = CandidateFacesForMesh[Index];
-						if (CandidateSurface->IsMeshed())
+						if (CandidateSurface->IsNotMeshable())
 						{
 							CandidateFacesForMesh.RemoveAt(Index);
 							--Index;
@@ -1249,7 +1325,7 @@ void FParametricMesher::MeshSurfaceByFront(TArray<FCostToFace>& QuadTrimmedSurfa
 			for (int32 Index = 0; Index < SecondChoiceOfCandidateFacesForMesh.Num(); ++Index)
 			{
 				FTopologicalFace* CandidateSurface = SecondChoiceOfCandidateFacesForMesh[Index];
-				if (CandidateSurface->IsMeshed())
+				if (CandidateSurface->IsNotMeshable())
 				{
 					SecondChoiceOfCandidateFacesForMesh.RemoveAt(Index);
 					--Index;
@@ -1280,7 +1356,7 @@ void FParametricMesher::MeshSurfaceByFront(TArray<FCostToFace>& QuadTrimmedSurfa
 	// the the other surface
 	for (FTopologicalFace* Face : Faces)
 	{
-		if (Face != nullptr && !Face->IsMeshed())
+		if (Face != nullptr && Face->IsMeshable())
 		{
 			MeshFacesByFront(*Face);
 		}
@@ -1288,171 +1364,79 @@ void FParametricMesher::MeshSurfaceByFront(TArray<FCostToFace>& QuadTrimmedSurfa
 
 }
 
-// =========================================================================================================================================================================================================
-// =========================================================================================================================================================================================================
-// =========================================================================================================================================================================================================
-//
-//
-//                                                                            NOT YET IMPLEMENTED
-//
-//
-// =========================================================================================================================================================================================================
-// =========================================================================================================================================================================================================
-// =========================================================================================================================================================================================================
-
-void FParametricMesher::MeshThinZoneEdges(FGrid& Grid)
+void FParametricMesher::MeshThinZoneEdges(FTopologicalFace& Face)
 {
-#ifdef UNUSED
+	TArray<FThinZone2D>& ThinZones = Face.GetThinZones();
+
+	{
 #ifdef DEBUG_MESHTHINSURF
-	Open3DDebugSession(FString::Printf(TEXT("thin Surfaces cutting on surf %d"), Grid.GetFace()->GetId()));
+		if(bDisplay)
+		{
+			F3DDebugSession _(bDisplay, FString::Printf(TEXT("Mesh Thin Face %d"), Face.GetId()));
+		}
 #endif
 
-	const TArray<FThinZone2D>& ThinZones = Grid.GetThinZones();
-
-	FTimePoint MeshStartTime = FChrono::Now();
-
-	for (const FThinZone2D& Zone : ThinZones)
-	{
-		bool bFirstSideIsPartiallyMeshed = Zone.GetFirstSide().IsPartiallyMeshed();
-		bool bSecondSideIsPartiallyMeshed = Zone.GetSecondSide().IsPartiallyMeshed();
+		for (FThinZone2D& Zone : ThinZones)
+		{
+			FindThinZoneBoundary(Zone.GetFirstSide());
+			FindThinZoneBoundary(Zone.GetSecondSide());
+		}
 
-		if (bFirstSideIsPartiallyMeshed && bSecondSideIsPartiallyMeshed)
+		for (FThinZone2D& Zone : ThinZones)
 		{
-			// the most meshed edge is meshed first
-			double FirstSideMeshedLength = Zone.GetFirstSide().GetMeshedLength();
-			double SecondSideMeshedLength = Zone.GetSecondSide().GetMeshedLength();
-			if (FirstSideMeshedLength > SecondSideMeshedLength)
+			EMeshingState bFirstSideMeshingState = Zone.GetFirstSide().GetMeshingState();
+			EMeshingState bSecondSideMeshingState = Zone.GetSecondSide().GetMeshingState();
+
+			if (bFirstSideMeshingState == EMeshingState::FullyMeshed)
 			{
-				bSecondSideIsPartiallyMeshed = false;
+				MeshThinZoneSide(Zone.GetFirstSide());
 			}
-			else
+			else if (bSecondSideMeshingState == EMeshingState::FullyMeshed)
 			{
-				bFirstSideIsPartiallyMeshed = false;
+				MeshThinZoneSide(Zone.GetSecondSide());
 			}
-		}
-
-		if (!bFirstSideIsPartiallyMeshed && !bSecondSideIsPartiallyMeshed)
-		{
-			if (Zone.GetFirstSide().GetLength() > Zone.GetSecondSide().GetLength())
+			else if (bFirstSideMeshingState == EMeshingState::PartiallyMeshed)
 			{
-				GetThinZoneBoundary(Zone.GetFirstSide());
-				GetThinZoneBoundary(Zone.GetSecondSide());
-
 				MeshThinZoneSide(Zone.GetFirstSide());
 			}
-			else
+			else if (bSecondSideMeshingState == EMeshingState::PartiallyMeshed)
 			{
-				GetThinZoneBoundary(Zone.GetFirstSide());
-				GetThinZoneBoundary(Zone.GetSecondSide());
-
 				MeshThinZoneSide(Zone.GetSecondSide());
 			}
-		}
-		else if (bFirstSideIsPartiallyMeshed && !bSecondSideIsPartiallyMeshed)
-		{
-			MeshThinZoneSide(Zone.GetFirstSide());
-			GetThinZoneBoundary(Zone.GetSecondSide());
-		}
-		else if (!bFirstSideIsPartiallyMeshed && bSecondSideIsPartiallyMeshed)
-		{
-			MeshThinZoneSide(Zone.GetSecondSide());
-			GetThinZoneBoundary(Zone.GetFirstSide());
-		}
-	}
-
-	// if the extremity of the thin zone are connected by a short edges path, the edges path are not discretized to avoid a well discretized edge connecting to thin sides
-
-#ifdef DEBUG_MESHTHINSURF
-	Close3DDebugSession();
-#endif
-
-#ifdef DEBUG_MESHTHINSURF
-	Open3DDebugSession(FString("Mesh of ThinZone 2D of surf " + Utils::ToFString(Grid.GetFace()->GetId())));
-	for (const FThinZone2D& Zone : ThinZones)
-	{
-		TArray<TSharedPtr<FTopologicalEdge>> OutEdges;
-		Zone.GetFirstSide().GetEdges(OutEdges);
-		for (TSharedPtr<FTopologicalEdge> Edge : OutEdges)
-		{
-			TSharedPtr<FTopologicalEdge> ActiveEdge = StaticCastSharedRef<FTopologicalEdge>(Edge->GetLinkActiveEntity());
-
-			TArray<double> ImposedCuttingPointEdgeCoordinates;
-			GetCuttingPointCoordinates(ActiveEdge->GetImposedCuttingPoints(), ImposedCuttingPointEdgeCoordinates);
-
-			TArray<double> ImposedCuttingPointU;
-			Edge->TransformActiveEdgeCoordinatesToLocalCoordinates(ImposedCuttingPointEdgeCoordinates, ImposedCuttingPointU);
-
-			TArray<FPoint> ImposedCuttingPoint2D;
-			Edge->Approximate2DPoints(ImposedCuttingPointU, ImposedCuttingPoint2D);
-			for (FPoint Point : ImposedCuttingPoint2D)
+			else if (Zone.GetFirstSide().Length() > Zone.GetSecondSide().Length())
 			{
-				Display(Point);
+				MeshThinZoneSide(Zone.GetFirstSide());
 			}
-		}
-
-		Zone.GetSecondSide().GetEdges(OutEdges);
-		for (TSharedPtr<FTopologicalEdge> Edge : OutEdges)
-		{
-			TSharedPtr<FTopologicalEdge> ActiveEdge = StaticCastSharedRef<FTopologicalEdge>(Edge->GetLinkActiveEntity());
-
-			TArray<double> ImposedCuttingPointEdgeCoordinates;
-			GetCuttingPointCoordinates(ActiveEdge->GetImposedCuttingPoints(), ImposedCuttingPointEdgeCoordinates);
-
-			TArray<double> ImposedCuttingPointCoordinates;
-			Edge->TransformActiveEdgeCoordinatesToLocalCoordinates(ImposedCuttingPointEdgeCoordinates, ImposedCuttingPointCoordinates);
-
-			TArray<FPoint> ImposedCuttingPoint2D;
-			Edge->Approximate2DPoints(ImposedCuttingPointCoordinates, ImposedCuttingPoint2D);
-			for (FPoint Point : ImposedCuttingPoint2D)
+			else
 			{
-				Display(Point);
+				MeshThinZoneSide(Zone.GetSecondSide());
 			}
 		}
 	}
-	Close3DDebugSession();
-#endif
-#ifdef CADKERNEL_DEV
-	MesherReport.Chronos.GlobalMeshThinZones += FChrono::Elapse(MeshStartTime);
-#endif
-#endif
+
+	// if the extremity of the thin zone are connected by a short edges path, the edges path are not discretized to avoid a well discretized edge connecting to thin sides
+
 }
 
-static void AddActiveEdgeThinZone(TSharedPtr<FTopologicalEdge> Edge, TSharedPtr<FTopologicalEdge> ActiveEdge, FLinearBoundary& SideEdgeCoordinate)
+void FParametricMesher::FindThinZoneBoundary(FThinZoneSide& Side)
 {
-	TArray<double> SideEdgeBound;
-	SideEdgeBound.SetNum(2);
-	SideEdgeBound[0] = SideEdgeCoordinate.GetMin();
-	SideEdgeBound[1] = SideEdgeCoordinate.GetMax();
-
-	TArray<double> ActiveEdgeThinZone;
-	Edge->TransformActiveEdgeCoordinatesToLocalCoordinates(SideEdgeBound, ActiveEdgeThinZone);
-	FLinearBoundary ThinZoneBoundary(ActiveEdgeThinZone[0], ActiveEdgeThinZone[1]);
-	ActiveEdge->AddThinZone(ThinZoneBoundary);
-};
-
-void FParametricMesher::GetThinZoneBoundary(const FThinZoneSide& Side)
-{
-#ifdef UNUSED
-	TSharedPtr<FTopologicalEdge> Edge = nullptr;
-	TSharedPtr<FTopologicalEdge> ActiveEdge = nullptr;
+	FTopologicalEdge* Edge = nullptr;
 	FLinearBoundary SideEdgeCoordinate;
 
-	for (const FEdgeSegment* EdgeSegment : Side.GetSegments())
+	for (FEdgeSegment& EdgeSegment : Side.GetSegments())
 	{
-		double UMin = EdgeSegment->GetCoordinate(ELimit::Start);
-		double UMax = EdgeSegment->GetCoordinate(ELimit::End);
+		double UMin = EdgeSegment.GetCoordinate(ELimit::Start);
+		double UMax = EdgeSegment.GetCoordinate(ELimit::End);
 		GetMinMax(UMin, UMax);
 
-		if (Edge != EdgeSegment->GetEdge())
+		if (Edge != EdgeSegment.GetEdge())
 		{
 			if (Edge)
 			{
-				AddActiveEdgeThinZone(Edge, ActiveEdge, SideEdgeCoordinate);
+				Edge->AddThinZone(SideEdgeCoordinate);
 			}
 
-			Edge = EdgeSegment->GetEdge();
-			ActiveEdge = StaticCastSharedRef<FTopologicalEdge>(Edge->GetLinkActiveEntity());
-
+			Edge = EdgeSegment.GetEdge();
 			SideEdgeCoordinate.Set(UMin, UMax);
 		}
 		else
@@ -1460,206 +1444,165 @@ void FParametricMesher::GetThinZoneBoundary(const FThinZoneSide& Side)
 			SideEdgeCoordinate.ExtendTo(UMin, UMax);
 		}
 	};
-	AddActiveEdgeThinZone(Edge, ActiveEdge, SideEdgeCoordinate);
-#endif
+	Edge->AddThinZone(SideEdgeCoordinate);
 }
 
-void FParametricMesher::MeshThinZoneSide(const FThinZoneSide& Side)
+#define DEBUG_MESHTHINSURF
+void FParametricMesher::MeshThinZoneSide(FThinZoneSide& Side)
 {
-#ifdef UNUSED
 	typedef TFunction<bool(double, double)> CompareMethode;
 
-	TSharedPtr<FTopologicalEdge> Edge;
-	TSharedPtr<FTopologicalEdge> ActiveEdge;
+	FTopologicalEdge* Edge = nullptr;
+
 	int32 Index = 0;
 	int32 Increment = 1;
+	double UMin = 0.;
+	double UMax = 0.;
+
 	TArray<double> EdgeCuttingPointCoordinates;
-	FLinearBoundary SideEdgeCoordinate;
 	const TArray<int32>* NodeIndices = nullptr;
 
-	TFunction<void(int32&, int32&, const FEdgeSegment*, double, double)> AddImposedCuttingPoint = [&](int32& Index, int32& Increment, const FEdgeSegment* EdgeSegment, double UMin, double UMax)
+	TFunction<void(const FEdgeSegment&)> AddImposedCuttingPoint = [&](const FEdgeSegment& EdgeSegment)
 	{
-#ifdef DebugMeshThinSurf
-		DisplaySegment(EdgeSegment->GetExtemity(ELimit::End), EdgeSegment->GetExtemity(ELimit::Start));
-#endif
 		for (; Index >= 0 && Index < EdgeCuttingPointCoordinates.Num(); Index += Increment)
 		{
-			if (EdgeCuttingPointCoordinates[Index] < UMin)
+			if (EdgeCuttingPointCoordinates[Index] < UMin || EdgeCuttingPointCoordinates[Index] > UMax)
 			{
 				break;
 			}
-			if (EdgeCuttingPointCoordinates[Index] > UMax)
-			{
-				break;
-			}
-			FPoint CuttingPoint3D = EdgeSegment->ComputeEdgePoint(EdgeCuttingPointCoordinates[Index]);
 
-			const FEdgeSegment* ClosedSegment = EdgeSegment->GetClosedSegment();
+			FPoint2D CuttingPoint = EdgeSegment.ComputeEdgePoint(EdgeCuttingPointCoordinates[Index]);
+
+			FEdgeSegment* ClosedSegment = EdgeSegment.GetCloseSegment();
 			if (ClosedSegment == nullptr)
 			{
 #ifdef CADKERNEL_DEV
+				ensureCADKernel(false);
 				Wait();
 #endif
 				continue;
 			}
+
 			double OppositeCuttingPointSegmentU;
-			FPoint OppositeCuttingPoint3D = ClosedSegment->ProjectPoint(CuttingPoint3D, OppositeCuttingPointSegmentU);
+			FPoint2D OppositeCuttingPoint = ClosedSegment->ProjectPoint(CuttingPoint, OppositeCuttingPointSegmentU);
 
 			double OppositeCuttingPointU = 0;
-			TSharedPtr<FTopologicalEdge> OpositEdge = nullptr;
-			if (OppositeCuttingPointSegmentU == 0 && ClosedSegment->GetPrevious()->GetClosedSegment())
+			FTopologicalEdge* OppositeEdge = nullptr;
+			if (FMath::IsNearlyZero(OppositeCuttingPointSegmentU) && ClosedSegment->GetPrevious()->GetCloseSegment())
 			{
-				OppositeCuttingPoint3D = ClosedSegment->GetPrevious()->ProjectPoint(CuttingPoint3D, OppositeCuttingPointSegmentU);
-				OppositeCuttingPointU = ClosedSegment->GetPrevious()->ComputeEdgeCoordinate(OppositeCuttingPointSegmentU);
-				OpositEdge = ClosedSegment->GetPrevious()->GetEdge();
+				FEdgeSegment* PreviousClosedSegment = ClosedSegment->GetPrevious();
+				OppositeCuttingPoint = PreviousClosedSegment->ProjectPoint(CuttingPoint, OppositeCuttingPointSegmentU);
+				OppositeCuttingPointU = PreviousClosedSegment->ComputeEdgeCoordinate(OppositeCuttingPointSegmentU);
+				OppositeEdge = PreviousClosedSegment->GetEdge();
 			}
-			else if (OppositeCuttingPointSegmentU == 1 && ClosedSegment->GetNext()->GetClosedSegment())
+			else if (FMath::IsNearlyEqual(OppositeCuttingPointSegmentU, 1.) && ClosedSegment->GetNext()->GetCloseSegment())
 			{
-				OppositeCuttingPoint3D = ClosedSegment->GetNext()->ProjectPoint(CuttingPoint3D, OppositeCuttingPointSegmentU);
-				OppositeCuttingPointU = ClosedSegment->GetNext()->ComputeEdgeCoordinate(OppositeCuttingPointSegmentU);
-				OpositEdge = ClosedSegment->GetNext()->GetEdge();
+				FEdgeSegment* NextClosedSegment = ClosedSegment->GetNext();
+				OppositeCuttingPoint = NextClosedSegment->ProjectPoint(CuttingPoint, OppositeCuttingPointSegmentU);
+				OppositeCuttingPointU = NextClosedSegment->ComputeEdgeCoordinate(OppositeCuttingPointSegmentU);
+				OppositeEdge = NextClosedSegment->GetEdge();
 			}
 			else
 			{
-				OpositEdge = ClosedSegment->GetEdge();
+				OppositeEdge = ClosedSegment->GetEdge();
 				OppositeCuttingPointU = ClosedSegment->ComputeEdgeCoordinate(OppositeCuttingPointSegmentU);
 			}
 
-			double OppositeActiveEdgeCuttingPointU = OpositEdge->TransformLocalCoordinateToActiveEdgeCoordinate(OppositeCuttingPointU);
-
+			if (OppositeEdge != Edge)
+			{
+				OppositeEdge->AddImposedCuttingPointU(OppositeCuttingPointU, (*NodeIndices)[Index]);
 #ifdef DEBUG_MESHTHINSURF
-			DisplayPoint(CuttingPoint3D);
-			DisplaySegment(ClosedSegment->GetExtemity(ELimit::End), ClosedSegment->GetExtemity(ELimit::Start));
-			DisplayPoint(OppositeCuttingPoint3D);
-			DisplaySegment(CuttingPoint3D, OppositeCuttingPoint3D);
+				DisplayPoint2DWithScale(CuttingPoint, EVisuProperty::RedPoint);
+				DisplaySegmentWithScale(ClosedSegment->GetExtemity(ELimit::End), ClosedSegment->GetExtemity(ELimit::Start));
+				DisplayPoint2DWithScale(OppositeCuttingPoint);
+				DisplaySegmentWithScale(CuttingPoint, OppositeCuttingPoint);
+				Wait(false);
 #endif
-			OpositEdge->GetLinkActiveEdge()->AddImposedCuttingPointU(OppositeActiveEdgeCuttingPointU, (*NodeIndices)[Index]);
-		}
-	};
-
-	TFunction<void(TArray<double>&, double, int32&, CompareMethode)> FindFirstIndexForward = [](TArray<double>& EdgeCuttingPointU, double ULimit, int32& OutIndex, CompareMethode Compare)
-	{
-		for (; OutIndex < EdgeCuttingPointU.Num(); ++OutIndex)
-		{
-			if (Compare(ULimit, EdgeCuttingPointU[OutIndex]))
-			{
-				break;
 			}
 		}
 	};
 
-	TFunction<void(TArray<double>&, double, int32&, CompareMethode)> FindFirstIndexBackward = [](TArray<double>& EdgeCuttingPointU, double ULimit, int32& OutIndex, CompareMethode Compare)
+	TFunction<void(double, CompareMethode)> FindFirstIndex = [&](double ULimit, CompareMethode Compare)
 	{
-		for (; OutIndex >= 0; --OutIndex)
+		for (; Index >= 0 && Index < EdgeCuttingPointCoordinates.Num(); Index += Increment)
 		{
-			if (Compare(ULimit, EdgeCuttingPointU[OutIndex]))
+			if (Compare(ULimit, EdgeCuttingPointCoordinates[Index]))
 			{
 				break;
 			}
 		}
 	};
 
-	TFunction<void(const FEdgeSegment* EdgeSegment)> Process = [&](const FEdgeSegment* EdgeSegment)
+#ifdef DEBUG_MESHTHINSURF
+	Open3DDebugSession(TEXT("MeshThinZoneSide"));
+#endif
+
+	for (FEdgeSegment& EdgeSegment : Side.GetSegments())
 	{
-		double UMin = EdgeSegment->GetCoordinate(ELimit::Start);
-		double UMax = EdgeSegment->GetCoordinate(ELimit::End);
+		UMin = EdgeSegment.GetCoordinate(ELimit::Start);
+		UMax = EdgeSegment.GetCoordinate(ELimit::End);
 		GetMinMax(UMin, UMax);
 
-		if (Edge != EdgeSegment->GetEdge())
+		if (Edge != EdgeSegment.GetEdge())
 		{
+#ifdef DEBUG_MESHTHINSURF
 			if (Edge)
 			{
-				AddActiveEdgeThinZone(Edge, ActiveEdge, SideEdgeCoordinate);
-#ifdef DebugMeshThinSurf
 				Close3DDebugSession();
+			}
+			Open3DDebugSession(TEXT("Projection of mesh"));
 #endif
+			Edge = EdgeSegment.GetEdge();
+
+			if (!Edge->IsMeshed())
+			{
+				Mesh(*Edge, *Edge->GetFace());
 			}
 
-			Edge = EdgeSegment->GetEdge();
-			ActiveEdge = StaticCastSharedRef<FTopologicalEdge>(Edge->GetLinkActiveEntity());
+			FEdgeMesh& EdgeMesh = *Edge->GetOrCreateMesh(MeshModel);
+			NodeIndices = &EdgeMesh.EdgeVerticesIndex;
+			GetCuttingPointCoordinates(Edge->GetCuttingPoints(), EdgeCuttingPointCoordinates);
 
-			SideEdgeCoordinate.Set(UMin, UMax);
 
-			if (!ActiveEdge->IsMeshed())
+			if (EdgeCuttingPointCoordinates.Num() == 0)
 			{
-				Mesh(*Edge.ToSharedRef(), *Edge->GetFace());
-#ifdef DebugMeshThinSurf
-				Open3DDebugSession(TEXT("Mesh of Edge"));
-				DisplayMesh((const TSharedPtr<FEdgeMesh>&) ActiveEdge->GetMesh());
-				Close3DDebugSession();
-#endif
+				TArray<FPoint>& NodeCoordinates = EdgeMesh.GetNodeCoordinates();
+				TArray<FPoint> ProjectedPoints;
+				Edge->ProjectPoints(NodeCoordinates, EdgeCuttingPointCoordinates, ProjectedPoints);
+				if (EdgeCuttingPointCoordinates.Num() > 1 && EdgeCuttingPointCoordinates[0] > EdgeCuttingPointCoordinates[1])
+				{
+					Algo::Reverse(EdgeCuttingPointCoordinates);
+				}
+				EdgeCuttingPointCoordinates.EmplaceAt(0, Edge->GetStartCurvilinearCoordinates());
+				EdgeCuttingPointCoordinates.Emplace(Edge->GetEndCurvilinearCoordinates());
 			}
-#ifdef DebugMeshThinSurf
-			Open3DDebugSession(TEXT("Projection of mesh"));
-#endif
-			NodeIndices = &ActiveEdge->GetOrCreateMesh(MeshModel)->EdgeVerticesIndex;
 
-			TArray<double> CuttingPointCoordinates;
-			GetCuttingPointCoordinates(ActiveEdge->GetCuttingPoints(), CuttingPointCoordinates);
+			ensureCADKernel(EdgeCuttingPointCoordinates[0] < EdgeCuttingPointCoordinates[1]);
 
-			Edge->TransformActiveEdgeCoordinatesToLocalCoordinates(CuttingPointCoordinates, EdgeCuttingPointCoordinates);
-
-			if ((EdgeCuttingPointCoordinates[0] < EdgeCuttingPointCoordinates[1]) == (EdgeSegment->GetCoordinate(ELimit::Start) < EdgeSegment->GetCoordinate(ELimit::End)))
+			bool bEdgeIsForward = EdgeSegment.IsForward();
+			if (bEdgeIsForward)
 			{
 				Index = 0;
-				if (EdgeCuttingPointCoordinates[0] < EdgeCuttingPointCoordinates[1])
-				{
-					FindFirstIndexForward(EdgeCuttingPointCoordinates, UMin, Index, [](double Value1, double Value2) {return (Value1 < Value2); });
-				}
-				else
-				{
-					FindFirstIndexForward(EdgeCuttingPointCoordinates, UMax, Index, [](double Value1, double Value2) {return (Value1 > Value2); });
-				}
 				Increment = 1;
+				FindFirstIndex(UMin, [](double Value1, double Value2) {return (Value1 < Value2); });
 			}
 			else
 			{
 				Index = (int32)EdgeCuttingPointCoordinates.Num() - 1;
-				if (EdgeCuttingPointCoordinates[0] < EdgeCuttingPointCoordinates[1])
-				{
-					FindFirstIndexBackward(EdgeCuttingPointCoordinates, UMax, Index, [](double Value1, double Value2) {return (Value1 > Value2); });
-				}
-				else
-				{
-					FindFirstIndexBackward(EdgeCuttingPointCoordinates, UMax, Index, [](double Value1, double Value2) {return (Value1 < Value2); });
-				}
 				Increment = -1;
+				FindFirstIndex(UMax, [](double Value1, double Value2) {return (Value1 > Value2); });
 			}
 		}
-		else
-		{
-			SideEdgeCoordinate.ExtendTo(UMin, UMax);
-		}
-
-		AddImposedCuttingPoint(Index, Increment, EdgeSegment, UMin, UMax);
-	};
 
-#ifdef DebugMeshThinSurf
-	Open3DDebugSession(TEXT("MeshThinZoneSide"));
-#endif
-
-	if (Side.IsFirstSide())
-	{
-		for (const FEdgeSegment* EdgeSegment : Side.GetSegments())
-		{
-			Process(EdgeSegment);
-		}
-	}
-	else
-	{
-		const TArray<FEdgeSegment*>& Segments = Side.GetSegments();
-		for (int32 SegmentIndex = Segments.Num() - 1; SegmentIndex >= 0; --SegmentIndex)
-		{
-			Process(Segments[SegmentIndex]);
-		}
+		AddImposedCuttingPoint(EdgeSegment);
 	}
-	AddActiveEdgeThinZone(Edge, ActiveEdge, SideEdgeCoordinate);
 
-#ifdef DebugMeshThinSurf
+#ifdef DEBUG_MESHTHINSURF
 	Close3DDebugSession();
 	Close3DDebugSession();
+	//Wait();
 #endif
-#endif
+
 }
 
 #ifdef DEBUG_INTERSECTEDGEISOS
diff --git a/Engine/Source/Runtime/Datasmith/CADKernel/Private/CADKernel/Mesh/Structure/Grid.cpp b/Engine/Source/Runtime/Datasmith/CADKernel/Private/CADKernel/Mesh/Structure/Grid.cpp
--- a/Engine/Source/Runtime/Datasmith/CADKernel/Private/CADKernel/Mesh/Structure/Grid.cpp
+++ b/Engine/Source/Runtime/Datasmith/CADKernel/Private/CADKernel/Mesh/Structure/Grid.cpp
@@ -22,7 +22,7 @@ FGrid::FGrid(FTopologicalFace& InFace, FModelMesh& InMeshModel)
 	, MinimumElementSize(Tolerance3D * 2.)
 	, MeshModel(InMeshModel)
 {
-#ifdef DEBUG_ONLY_SURFACE_TO_DEBUG
+#ifdef CADKERNEL_DEBUG
 	bDisplay = (InFace.GetId() == FaceToDebug);
 	Open3DDebugSession(bDisplay, FString::Printf(TEXT("Grid %d"), InFace.GetId()));
 #endif
@@ -320,7 +320,7 @@ bool FGrid::GeneratePointCloud()
 {
 	FTimePoint StartTime = FChrono::Now();
 
-	if (CheckIf2DGridIsDegenerate())
+	if(CheckIf2DGridIsDegenerate())
 	{
 		return false;
 	}
@@ -364,10 +364,10 @@ void FGrid::FindPointsCloseToLoop()
 		if (bDisplay)
 		{
 			F3DDebugSession _(*FString::Printf(TEXT("Cell %d"), CellIndex++));
-			DisplayPoint2DWithScale(Points2D[EGridSpace::UniformScaled][GlobalIndex]);
-			DisplayPoint2DWithScale(Points2D[EGridSpace::UniformScaled][GlobalIndex - 1]);
-			DisplayPoint2DWithScale(Points2D[EGridSpace::UniformScaled][GlobalIndex - 1 - CuttingCount[EIso::IsoU]]);
-			DisplayPoint2DWithScale(Points2D[EGridSpace::UniformScaled][GlobalIndex - CuttingCount[EIso::IsoU]]);
+			DisplayPoint2DWithScale(Points2D[EGridSpace::UniformScaled][GlobalIndex] );
+			DisplayPoint2DWithScale(Points2D[EGridSpace::UniformScaled][GlobalIndex - 1] );
+			DisplayPoint2DWithScale(Points2D[EGridSpace::UniformScaled][GlobalIndex - 1 - CuttingCount[EIso::IsoU]] );
+			DisplayPoint2DWithScale(Points2D[EGridSpace::UniformScaled][GlobalIndex - CuttingCount[EIso::IsoU]] );
 			Wait(bWaitCell);
 		}
 	};
@@ -528,7 +528,7 @@ void FGrid::FindPointsCloseToLoop()
 			if (bDisplay)
 			{
 				F3DDebugSession _(TEXT("SEG"));
-				DisplaySegmentWithScale(*PointB, *PointA);
+				DisplaySegmentWithScale(*PointB , *PointA );
 			}
 #endif
 
@@ -1068,10 +1068,17 @@ void FGrid::GetMeshOfLoop(const FTopologicalLoop& Loop)
 	}
 }
 
+//#define DEBUG_GET_BOUNDARY_MESH
 bool FGrid::GetMeshOfLoops()
 {
+	int32 ThinZoneNum = 0;
+	if (Face.HasThinZone())
+	{
+		ThinZoneNum = Face.GetThinZones().Num();
+	}
+
 	int32 LoopCount = Face.GetLoops().Num();
-	FaceLoops2D[EGridSpace::Default2D].Reserve(LoopCount);
+	FaceLoops2D[EGridSpace::Default2D].Reserve(LoopCount + ThinZoneNum);
 
 	FaceLoops3D.Reserve(LoopCount);
 	NormalsOfFaceLoops.Reserve(LoopCount);
@@ -1104,6 +1111,30 @@ bool FGrid::GetMeshOfLoops()
 		return false;
 	}
 
+	TFunction<void(TArray<FPoint2D>&, const TArray<FEdgeSegment>&)> AddThinZoneInLoopPoints = [](TArray<FPoint2D>& LoopPoints, const TArray<FEdgeSegment>& Segments)
+	{
+		const FEdgeSegment& Segment = Segments[0];
+		LoopPoints.Emplace(Segment.GetEdge()->Approximate2DPoint(Segment.GetCoordinate(ELimit::Start)));
+		for (const FEdgeSegment& SegmentIter : Segments)
+		{
+			LoopPoints.Emplace(SegmentIter.GetEdge()->Approximate2DPoint(SegmentIter.GetCoordinate(ELimit::End)));
+		}
+	};
+
+	if (ThinZoneNum)
+	{
+		for (const FThinZone2D& ThinZone : Face.GetThinZones())
+		{
+			int32 PointNum = ThinZone.GetFirstSide().GetSegments().Num();
+			PointNum += ThinZone.GetSecondSide().GetSegments().Num();
+			TArray<FPoint2D>& LoopPoints = FaceLoops2D[EGridSpace::Default2D].Emplace_GetRef();
+			LoopPoints.Reserve(PointNum + 4);
+
+			AddThinZoneInLoopPoints(LoopPoints, ThinZone.GetFirstSide().GetSegments());
+			AddThinZoneInLoopPoints(LoopPoints, ThinZone.GetSecondSide().GetSegments());
+		}
+	}
+
 	// Fit boundaries to Surface bounds.
 	const FSurfacicBoundary& Bounds = Face.GetBoundary();
 	for (TArray<FPoint2D>& Loop : FaceLoops2D[EGridSpace::Default2D])
@@ -1154,7 +1185,7 @@ bool FGrid::CheckIf2DGridIsDegenerate() const
 		for (int32 Index = 1; Index < CoordinateGrid[Iso].Num(); ++Index)
 		{
 			double Delta = CoordinateGrid[Iso][Index] - CoordinateGrid[Iso][Index - 1];
-			if (Delta > MaxDelta)
+			if(Delta > MaxDelta)
 			{
 				MaxDelta = Delta;
 			}
diff --git a/Engine/Source/Runtime/Datasmith/CADKernel/Private/CADKernel/Mesh/Structure/Grid.h b/Engine/Source/Runtime/Datasmith/CADKernel/Private/CADKernel/Mesh/Structure/Grid.h
--- a/Engine/Source/Runtime/Datasmith/CADKernel/Private/CADKernel/Mesh/Structure/Grid.h
+++ b/Engine/Source/Runtime/Datasmith/CADKernel/Private/CADKernel/Mesh/Structure/Grid.h
@@ -10,7 +10,7 @@
 #include "CADKernel/UI/Display.h"
 
 #ifdef CADKERNEL_DEV
-#include "CADKernel/Mesh/Meshers/IsoTriangulator/DefineForDebug.h"
+#include "CADKernel/UI/DefineForDebug.h"
 #endif
 
 
@@ -94,13 +94,10 @@ protected:
 	 */
 	TArray<ENodeMarker> NodeMarkers;
 
-public:
-	FGridChronos Chronos;
-
 public:
 	FGrid(FTopologicalFace& InFace, FModelMesh& InShellMesh);
 
-#ifndef CADKERNEL_DEV
+#ifndef CADKERNEL_DEBUG
 	virtual ~FGrid() = default;
 #else
 	virtual ~FGrid()
@@ -424,7 +421,7 @@ public:
 	// ======================================================================================================================================================================================================================
 	// Display Methodes   ================================================================================================================================================================================================
 	// ======================================================================================================================================================================================================================
-#ifdef CADKERNEL_DEV
+#ifdef CADKERNEL_DEBUG
 	bool bDisplay = false;
 
 	void DisplayIsoNode(EGridSpace Space, const int32 PointIndex, FIdent Ident = 0, EVisuProperty Property = EVisuProperty::BluePoint) const;
@@ -482,7 +479,7 @@ public:
 
 	void DisplayNodes(const TCHAR* Message, EGridSpace Space, const TArray<const FIsoNode*>& Nodes, EVisuProperty Property) const;
 
-	void DisplayGridPoints(EGridSpace DisplaySpace) const;
+	virtual void DisplayGridPoints(EGridSpace DisplaySpace) const override;
 	void DisplayInnerPoints(TCHAR* Message, EGridSpace DisplaySpace) const;
 
 	void DisplayGridNormal() const;
diff --git a/Engine/Source/Runtime/Datasmith/CADKernel/Private/CADKernel/Mesh/Structure/GridBase.cpp b/Engine/Source/Runtime/Datasmith/CADKernel/Private/CADKernel/Mesh/Structure/GridBase.cpp
--- a/Engine/Source/Runtime/Datasmith/CADKernel/Private/CADKernel/Mesh/Structure/GridBase.cpp
+++ b/Engine/Source/Runtime/Datasmith/CADKernel/Private/CADKernel/Mesh/Structure/GridBase.cpp
@@ -108,7 +108,7 @@ bool FGridBase::ScaleGrid()
 
 	double MeanLengthV = GetMean(LengthsV);
 	double MeanLengthU = GetMean(LengthsU);
-	if (MeanLengthV < Tolerance3D || MeanLengthU < Tolerance3D)
+	if(MeanLengthV < Tolerance3D || MeanLengthU < Tolerance3D)
 	{
 		SetAsDegenerated();
 		return false;
@@ -210,7 +210,9 @@ bool FGridBase::ScaleGrid()
 			}
 		}
 	}
+#ifdef CADKERNEL_DEV
 	Chronos.ScaleGridDuration = FChrono::Elapse(StartTime);
+#endif
 	return true;
 }
 
@@ -249,6 +251,13 @@ void FGridBase::EvaluatePointGrid(const FCoordinateGrid& CoordGrid, bool bWithNo
 	}
 }
 
+#ifdef CADKERNEL_DEV
+void FGrid::PrintTimeElapse() const
+{
+	Chronos.PrintTimeElapse();
+}
+#endif
+
 }
 
 
diff --git a/Engine/Source/Runtime/Datasmith/CADKernel/Private/CADKernel/Mesh/Structure/GridBase.h b/Engine/Source/Runtime/Datasmith/CADKernel/Private/CADKernel/Mesh/Structure/GridBase.h
--- a/Engine/Source/Runtime/Datasmith/CADKernel/Private/CADKernel/Mesh/Structure/GridBase.h
+++ b/Engine/Source/Runtime/Datasmith/CADKernel/Private/CADKernel/Mesh/Structure/GridBase.h
@@ -10,7 +10,7 @@
 #include "CADKernel/UI/Display.h"
 
 #ifdef CADKERNEL_DEV
-#include "CADKernel/Mesh/Meshers/IsoTriangulator/DefineForDebug.h"
+#include "CADKernel/UI/DefineForDebug.h"
 #endif
 
 
@@ -104,6 +104,8 @@ protected:
 
 public:
 	FGridChronos Chronos;
+#ifdef CADKERNEL_DEV
+#endif
 
 public:
 	FGridBase(FTopologicalFace& InFace);
@@ -162,7 +164,7 @@ protected:
 			PointU0V0 +
 			(PointU0V0 - PointU1V0) * (U0MinusU / U1MinusU0) +
 			(PointU0V0 - PointU0V1) * (V0MinusV / V1MinusV0) +
-			(PointU0V0 - PointU0V1 + PointU1V1 - PointU1V0) * ((U0MinusU * V0MinusV) / (U1MinusU0 * V1MinusV0));	
+			(PointU0V0 - PointU0V1 + PointU1V1 - PointU1V0) * ((U0MinusU * V0MinusV) / (U1MinusU0 * V1MinusV0));
 	};
 
 public:
@@ -227,7 +229,7 @@ public:
 		return Points2D[(int32)Space];
 	}
 
-	const FTopologicalFace& GetFace() const
+const FTopologicalFace& GetFace() const
 	{
 		return Face;
 	}
@@ -245,7 +247,18 @@ public:
 		return IndexV * CuttingCount[EIso::IsoU] + IndexU;
 	}
 
-#ifdef CADKERNEL_DEV
+#ifdef WIP
+	/**
+	 * @return the Index of the position in the arrays of a point [IndexU, IndexV] of the grid
+	 */
+	void UVIndexFromGlobalIndex(int32 GLobalIndex, int32& OutIndexU, int32& OutIndexV) const
+	{
+		OutIndexU = GLobalIndex % CuttingCount[EIso::IsoU];
+		OutIndexV = GLobalIndex / CuttingCount[EIso::IsoU];
+	}
+#endif
+
+#ifdef CADKERNEL_DEBUG
 
 	// ======================================================================================================================================================================================================================
 	// Display Methodes   ================================================================================================================================================================================================
@@ -255,7 +268,9 @@ public:
 	void DisplayIsoNode(EGridSpace Space, const int32 PointIndex, FIdent Ident = 0, EVisuProperty Property = EVisuProperty::BluePoint) const;
 	virtual void DisplayGridPoints(EGridSpace DisplaySpace) const;
 	void DisplayInnerPoints(TCHAR* Message, EGridSpace DisplaySpace) const;
+#endif
 
+#ifdef CADKERNEL_DEV
 	void PrintTimeElapse() const;
 #endif
 };
diff --git a/Engine/Source/Runtime/Datasmith/CADKernel/Private/CADKernel/Mesh/Structure/LoopCleaner.h b/Engine/Source/Runtime/Datasmith/CADKernel/Private/CADKernel/Mesh/Structure/LoopCleaner.h
--- a/Engine/Source/Runtime/Datasmith/CADKernel/Private/CADKernel/Mesh/Structure/LoopCleaner.h
+++ b/Engine/Source/Runtime/Datasmith/CADKernel/Private/CADKernel/Mesh/Structure/LoopCleaner.h
@@ -447,7 +447,7 @@ private:
 		SegmentCount--;
 	}
 
-#ifdef CADKERNEL_DEV		
+#ifdef CADKERNEL_DEBUG
 	void DisplayIntersection(const TPair<double, double>& Intersection)
 	{
 		if (bDisplay)
diff --git a/Engine/Source/Runtime/Datasmith/CADKernel/Private/CADKernel/Mesh/Structure/ThinZone2D.h b/Engine/Source/Runtime/Datasmith/CADKernel/Private/CADKernel/Mesh/Structure/ThinZone2D.h
deleted file mode 100644
diff --git a/Engine/Source/Runtime/Datasmith/CADKernel/Private/CADKernel/Mesh/Structure/ThinZone2DDisplay.cpp b/Engine/Source/Runtime/Datasmith/CADKernel/Private/CADKernel/Mesh/Structure/ThinZone2DDisplay.cpp
new file mode 100644
--- /dev/null
+++ b/Engine/Source/Runtime/Datasmith/CADKernel/Private/CADKernel/Mesh/Structure/ThinZone2DDisplay.cpp
@@ -0,0 +1,232 @@
+// Copyright Epic Games, Inc. All Rights Reserved.
+
+#include "CADKernel/Mesh/Structure/ThinZone2DFinder.h"
+
+#ifdef CADKERNEL_DEBUG
+
+#include "CADKernel/Math/Point.h"
+#include "CADKernel/Mesh/Structure/Grid.h"
+#include "CADKernel/Mesh/Structure/EdgeSegment.h"
+#include "CADKernel/Topo/TopologicalEdge.h"
+#include "CADKernel/Topo/TopologicalFace.h"
+#include "CADKernel/Topo/TopologicalLoop.h"
+#include "CADKernel/UI/Display.h"
+
+namespace UE::CADKernel
+{
+
+//#define DEBUG_SPLIT_BY_EDGE
+void FThinZone2DFinder::DisplayLoopSegments()
+{
+	if (!bDisplay)
+	{
+		return;
+	}
+
+	FTopologicalEdge* currentEdge = LoopSegments[0]->GetEdge();
+
+	F3DDebugSession _(TEXT("BoundarySegment"));
+
+#ifdef DEBUG_SPLIT_BY_EDGE
+	F3DDebugSession E(FString::Printf(TEXT("Edge %d"), currentEdge->GetId()));
+#endif
+	for (FEdgeSegment* EdgeSegment : LoopSegments)
+	{
+#ifdef DEBUG_SPLIT_BY_EDGE
+		if (currentEdge != EdgeSegment->GetEdge())
+		{
+			currentEdge = EdgeSegment->GetEdge();
+			Close3DDebugSession();
+			Open3DDebugSession(FString::Printf(TEXT("Edge %d"), currentEdge->GetId()));
+		}
+#endif
+		ThinZone::DisplayEdgeSegment(EdgeSegment, EVisuProperty::GreenCurve);
+	}
+	Wait();
+}
+
+void FThinZone2DFinder::DisplayCloseSegments()
+{
+	if (!bDisplay)
+	{
+		return;
+	}
+
+	F3DDebugSession _(TEXT("Close Segments"));
+	for (const FEdgeSegment* Segment : LoopSegments)
+	{
+		if (Segment->GetCloseSegment())
+		{
+			ThinZone::DisplayEdgeSegmentAndProjection(Segment, EVisuProperty::BlueCurve, EVisuProperty::BlueCurve, EVisuProperty::RedCurve);
+		}
+	}
+}
+
+void ThinZone::DisplayEdgeSegmentAndProjection(const FEdgeSegment* Segment, EVisuProperty SegColor, EVisuProperty OppositeColor, EVisuProperty ProjectionColor)
+{
+	FEdgeSegment* CloseSegment = Segment->GetCloseSegment();
+	DisplayEdgeSegmentAndProjection(Segment, CloseSegment, SegColor, OppositeColor, ProjectionColor);
+}
+
+void ThinZone::DisplayEdgeSegmentAndProjection(const FEdgeSegment* Segment, const FEdgeSegment* CloseSegment, EVisuProperty SegColor, EVisuProperty OppositeColor, EVisuProperty ProjectionColor)
+{
+	DisplayEdgeSegment(Segment, SegColor);
+	DisplayEdgeSegment(CloseSegment, OppositeColor);
+
+	double Coordinate;
+	FPoint2D Projection = ProjectPointOnSegment(Segment->GetCenter(), CloseSegment->GetExtemity(ELimit::Start), CloseSegment->GetExtemity(ELimit::End), Coordinate, true);
+
+	DisplaySegmentWithScale(Projection, Segment->GetCenter(), Segment->GetId(), ProjectionColor);
+}
+
+void ThinZone::DisplayEdgeSegment(const FEdgeSegment* EdgeSegment, EVisuProperty Color, int32 Index)
+{
+	DisplaySegmentWithScale(EdgeSegment->GetExtemity(ELimit::Start), EdgeSegment->GetExtemity(ELimit::End), Index, Color);
+}
+
+void ThinZone::DisplayEdgeSegment(const FEdgeSegment* EdgeSegment, EVisuProperty Color)
+{
+	DisplayEdgeSegment(EdgeSegment, Color, EdgeSegment->GetId());
+}
+
+void ThinZone::DisplayThinZoneSide(const TArray<FEdgeSegment*>& Side, int32 Index, EVisuProperty Color, bool bSplitBySegment)
+{
+	F3DDebugSession A(FString::Printf(TEXT("Side %d"), Index));
+	for (const FEdgeSegment* EdgeSegment : Side)
+	{
+		F3DDebugSession A(bSplitBySegment, TEXT("Seg"));
+		DisplayEdgeSegment(EdgeSegment, Color);
+	}
+}
+
+void ThinZone::DisplayThinZoneSide(const FThinZoneSide& Side, int32 Index, EVisuProperty Color, bool bSplitBySegment)
+{
+	F3DDebugSession A(FString::Printf(TEXT("Side %d"), Index));
+	for (const FEdgeSegment& EdgeSegment : Side.GetSegments())
+	{
+		F3DDebugSession A(bSplitBySegment, TEXT("Seg"));
+		DisplayEdgeSegment(&EdgeSegment, Color);
+	}
+}
+
+static EVisuProperty GetRandomColor(int32 Index)
+{
+	static TArray<EVisuProperty> ZoneColor;
+	ZoneColor.Add(EVisuProperty::OrangeCurve);
+	ZoneColor.Add(EVisuProperty::BlueCurve);
+	ZoneColor.Add(EVisuProperty::GreenCurve);
+	ZoneColor.Add(EVisuProperty::PurpleCurve);
+	ZoneColor.Add(EVisuProperty::RedCurve);
+	ZoneColor.Add(EVisuProperty::YellowCurve);
+	const uint32 ColorCount = ZoneColor.Num();
+
+	const int32 ColorIndex = Index % ColorCount;
+	return ZoneColor[ColorIndex];
+}
+
+void FThinZone2DFinder::DisplaySegmentsOfThinZone()
+{
+	if (!bDisplay)
+	{
+		return;
+	}
+
+	F3DDebugSession _(TEXT("Segments Of ThinZones"));
+
+	FIdent LastIndex = Ident::Undefined - 1;
+	EVisuProperty Color;
+
+	for (FEdgeSegment* EdgeSegment : LoopSegments)
+	{
+		FIdent ChainIndex = EdgeSegment->GetChainIndex();
+		if (ChainIndex != LastIndex && ChainIndex != Ident::Undefined)
+		{
+			if(LastIndex != Ident::Undefined - 1)
+			{
+				Close3DDebugSession();
+			}
+			Open3DDebugSession(FString::Printf(TEXT("ChainIndex %d"), ChainIndex));
+			Color = GetRandomColor(ChainIndex);
+			LastIndex = ChainIndex;
+		}
+		if (ChainIndex != Ident::Undefined)
+		{
+			ThinZone::DisplayEdgeSegment(EdgeSegment, Color, ChainIndex);
+		}
+	}
+	if (LastIndex != Ident::Undefined - 1)
+	{
+		Close3DDebugSession();
+	}
+	Wait();
+}
+
+void ThinZone::DisplayThinZoneSides(const TArray<TArray<FEdgeSegment*>>& ThinZoneSides)
+{
+	F3DDebugSession _(TEXT("ThinZone Sides"));
+	int32 ChainIndex = 0;
+	EVisuProperty Color;
+	for (const TArray<FEdgeSegment*>& ThinZoneSide : ThinZoneSides)
+	{
+		Color = GetRandomColor(ChainIndex);
+		DisplayThinZoneSide(ThinZoneSide, ChainIndex++, Color);
+	}
+}
+
+void ThinZone::DisplayThinZones(const TArray<FThinZone2D>& ThinZones)
+{
+	if (ThinZones.Num() > 0)
+	{
+		F3DDebugSession _(TEXT("Thin Zones"));
+
+		int32 index = 0;
+		for (const FThinZone2D& Zone : ThinZones)
+		{
+			FString Title;
+			index++;
+			EVisuProperty VisuProperty = EVisuProperty::BlueCurve;
+			switch (Zone.GetCategory())
+			{
+			case EThinZone2DType::Undefined:
+				VisuProperty = EVisuProperty::Iso;
+				Title = TEXT("Zone UNDEFINED");
+				break;
+			case EThinZone2DType::Global:
+				VisuProperty = EVisuProperty::BluePoint;
+				Title = TEXT("Zone GLOBAL");
+				break;
+			case EThinZone2DType::PeakStart:
+				VisuProperty = EVisuProperty::RedPoint;
+				Title = TEXT("Zone PEAK start");
+				break;
+			case EThinZone2DType::PeakEnd:
+				VisuProperty = EVisuProperty::OrangePoint;
+				Title = TEXT("Zone PEAK end");
+				break;
+			case EThinZone2DType::Butterfly:
+				VisuProperty = EVisuProperty::YellowPoint;
+				Title = TEXT("Zone BUTTERFLY");
+				break;
+			case EThinZone2DType::BetweenLoops:
+				VisuProperty = EVisuProperty::PinkPoint;
+				Title = TEXT("Zone BETWEEN_CONTOUR");
+				break;
+			default:
+				Title = TEXT("Zone Unknown");
+			}
+
+			Zone.Display(Title, VisuProperty);
+		}
+	}
+}
+
+void FThinZone2D::Display(const FString& Title, EVisuProperty VisuProperty) const
+{
+	F3DDebugSession A(Title);
+	ThinZone::DisplayThinZoneSide(GetFirstSide(), 0, VisuProperty);
+	ThinZone::DisplayThinZoneSide(GetSecondSide(), 1, VisuProperty);
+}
+
+}
+
+#endif
diff --git a/Engine/Source/Runtime/Datasmith/CADKernel/Private/CADKernel/Mesh/Structure/ThinZone2DFinder.cpp b/Engine/Source/Runtime/Datasmith/CADKernel/Private/CADKernel/Mesh/Structure/ThinZone2DFinder.cpp
new file mode 100644
--- /dev/null
+++ b/Engine/Source/Runtime/Datasmith/CADKernel/Private/CADKernel/Mesh/Structure/ThinZone2DFinder.cpp
@@ -0,0 +1,1148 @@
+// Copyright Epic Games, Inc. All Rights Reserved.
+
+#include "CADKernel/Mesh/Structure/ThinZone2DFinder.h"
+
+#include "CADKernel/Math/Point.h"
+#include "CADKernel/Mesh/Structure/Grid.h"
+#include "CADKernel/Mesh/Structure/EdgeSegment.h"
+#include "CADKernel/Topo/TopologicalEdge.h"
+#include "CADKernel/Topo/TopologicalFace.h"
+#include "CADKernel/Topo/TopologicalLoop.h"
+#include "CADKernel/UI/Display.h"
+
+namespace UE::CADKernel
+{
+
+FIdent FEdgeSegment::LastId = 0;
+
+void FThinZone2DFinder::FindCloseSegments()
+{
+	TFunction<double()> ComputeMaxSegmentLength = [&]()-> double
+	{
+		double MaxLength = 0;
+		for (const FEdgeSegment* Segment : LoopSegments)
+		{
+			double Length = Segment->GetLength();
+			if (Length > MaxLength)
+			{
+				MaxLength = Length;
+			}
+		}
+		MaxLength *= 1.01;
+		return MaxLength;
+	};
+
+	const double MaxSegmentLength = ComputeMaxSegmentLength();
+	const double MaxSpace = FMath::Max(MaxSegmentLength, FinderTolerance * 1.01);
+	const double MaxSquareSpace = 4 * FMath::Square(MaxSpace);
+	const double MaxSpacePlusLength = 1.1 * (MaxSpace + MaxSegmentLength);
+
+	// Copy of loop segments to generate a sorted array of segments
+	TArray<FEdgeSegment*> SortedThinZoneSegments = LoopSegments;
+	Algo::Sort(SortedThinZoneSegments, [](FEdgeSegment* SegmentA, FEdgeSegment* SegmentB)
+		{
+			return (SegmentA->GetAxeMin() < SegmentB->GetAxeMin());
+		});
+
+#ifdef DEBUG_FIND_CLOSE_SEGMENTS
+	F3DDebugSession _(bDisplay, ("Close Segments"));
+	int32 SegmentID = 1291;
+	int32 SegmentIndex = 0;
+#endif
+
+	// For each segment, the nearest segment is search
+	// If segment is from an inner loop, the nearest could not be from the same loop
+	int32 SegmentToBeginIndex = 0;
+	for (FEdgeSegment* Segment : SortedThinZoneSegments)
+	{
+		FTopologicalLoop* SegmentLoop = nullptr;
+		if (Segment->IsInner())
+		{
+			SegmentLoop = Segment->GetEdge()->GetLoop();
+		}
+		else
+		{
+			SegmentLoop = nullptr;
+		}
+
+		const FPoint2D SegmentMiddle = Segment->GetCenter();
+
+		FPoint2D ClosePoint, ClosePoint2;
+		FEdgeSegment* CloseSegment = nullptr;
+
+		//     DistanceMin between Candidate and Segment with Candidate is before 
+		//     DistanceMin = (Segment.SegmentMin - Candidate.SegmentMax) = Segment.SegmentMin - (Candidate.SegmentMin + MaxSegmentLength) 
+		//
+		//     Candidate                              Segment
+		//     *---*                                  x-----x
+		//     | Candidate.SegmentMin                 | Segment.SegmentMin
+		//         | Candidate.SegmentMax = Candidate.SegmentMin + MaxSegmentLength
+		//
+		//     We want DistanceMin < MaxSpace
+		//     Segment.SegmentMin - (Candidate.SegmentMin + MaxSegmentLength) < MaxSpace
+		//     Candidate.SegmentMin > Segment.SegmentMin - MaxSegmentLength - MaxSpace
+
+		//     DistanceMin between Candidate and Segment, if Candidate is after
+		//     DistanceMin = (Candidate.SegmentMin - Segment.SegmentMax) = Candidate.SegmentMin - (Segment.SegmentMin + MaxSegmentLength) 
+		//
+		//     Segment                                Candidate
+		//     *--x--*                                  x-----x
+		//        | Segment.Middle                      | Candidate.SegmentMin
+		//
+		//     We want DistanceMin < MaxSpace + MaxSegmentLength
+		//     Candidate.SegmentMin - Segment.Middle < MaxSpace + MaxSegmentLength
+		//     Candidate.SegmentMin < Segment.Middle + MaxSpace + MaxSegmentLength
+		//
+		//     MaxSpacePlusLength = MaxSpace + MaxSegmentLength;
+
+		const double SegmentMinMin = Segment->GetAxeMin() - MaxSpacePlusLength;
+		const double SegmentMiddleMax = SegmentMiddle.DiagonalAxisCoordinate() + MaxSpacePlusLength;
+		double MinSquareThickness = HUGE_VALUE;
+
+#ifdef DEBUG_FIND_CLOSE_SEGMENTS
+		if (SegmentIndex == SegmentID)
+		{
+			F3DDebugSession _(bDisplay, FString::Printf(TEXT("Segment %d"), SegmentIndex));
+			ThinZone::DisplayEdgeSegment(Segment, EVisuProperty::RedCurve);
+			Wait();
+		}
+#endif
+		for (int32 CandidateIndex = SegmentToBeginIndex; CandidateIndex < SortedThinZoneSegments.Num(); ++CandidateIndex)
+		{
+
+			FEdgeSegment* Candidate = SortedThinZoneSegments[CandidateIndex];
+			if (Candidate == Segment)
+			{
+				continue;
+			}
+
+			// Inner boundary are not check for thin inner boundary
+			if (SegmentLoop == Candidate->GetEdge()->GetLoop())
+			{
+				continue;
+			}
+
+
+			const double CandidateSegmentMin = Candidate->GetAxeMin();
+			// If min point of candidate segment - maxSpace (= max length of the segments + Tolerance) is smaller than Min point of current segment then the distance between both segments cannot be smaller the Tolerance
+			if (CandidateSegmentMin < SegmentMinMin)
+			{
+				SegmentToBeginIndex = CandidateIndex;
+				continue;
+			}
+
+			// If Min point of current segment + maxSpace is smaller than Min point of candidate segment then the distance between both segments cannot be smaller the Tolerance.
+			// As segments are sorted, next segments are not close to the current segment
+			if (CandidateSegmentMin > SegmentMiddleMax)
+			{
+				break;
+			}
+
+			const FPoint2D& FirstPointCandidate = Candidate->GetExtemity(ELimit::Start);
+
+			// if the distance of FirstPointCandidate with the middle of current segment is biggest than MaxSpace, then the projection of the point cannot be smaller than the tolerance, 
+			{
+				const double SquareDistance = SegmentMiddle.SquareDistance(FirstPointCandidate);
+				if (SquareDistance > MaxSquareSpace)
+				{
+					continue;
+				}
+			}
+
+			// check the angle between segments. As they are opposite, the cosAngle as to be close to -1
+			{
+				const double Slope = Segment->ComputeUnorientedSlopeOf(Candidate);
+				if (Slope < 3.) // Angle < 3Pi/4 (135 deg)
+				{
+					continue;
+				}
+			}
+
+#ifdef DEBUG_FIND_CLOSE_SEGMENTS
+			if (SegmentIndex == SegmentID)
+			{
+				F3DDebugSession _(bDisplay, FString::Printf(TEXT("Candidate %d %d"), SegmentIndex, CandidateIndex));
+				ThinZone::DisplayEdgeSegmentAndProjection(Segment, Candidate, EVisuProperty::BlueCurve, EVisuProperty::BlueCurve, EVisuProperty::RedCurve);
+				Wait();
+  			}
+			Wait(SegmentIndex == SegmentID + 1);
+#endif
+			const FPoint2D& SecondPointCandidate = Candidate->GetExtemity(ELimit::End);
+			double Coordinate;
+			const FPoint2D Projection = ProjectPointOnSegment(SegmentMiddle, FirstPointCandidate, SecondPointCandidate, Coordinate, true);
+
+			const double SquareDistance = SegmentMiddle.SquareDistance(Projection);
+			if (SquareDistance > SquareFinderTolerance)
+			{
+				continue;
+			}
+
+			if (MinSquareThickness > SquareDistance)
+			{
+				// check the angle between segment and Middle-Projection. As they are opposite, the cosAngle as to be close to 0
+				const double Slope = Segment->ComputeUnorientedSlopeOf(SegmentMiddle, Projection);
+				if (Slope < 1. || Slope > 3.)
+				{
+					continue;
+				}
+
+				MinSquareThickness = SquareDistance;
+				ClosePoint = Projection;
+				CloseSegment = Candidate;
+			}
+		}
+
+		if (CloseSegment)
+		{
+			Segment->SetCloseSegment(CloseSegment, MinSquareThickness);
+
+#ifdef DEBUG_FIND_CLOSE_SEGMENTS
+			if (bDisplay)
+			{
+				if (SegmentIndex == SegmentID)
+				{
+					Wait();
+				}
+
+				F3DDebugSession _(bDisplay, FString::Printf(TEXT("Close Segment %d"), SegmentIndex));
+				ThinZone::DisplayEdgeSegmentAndProjection(Segment, CloseSegment, EVisuProperty::BlueCurve, EVisuProperty::BlueCurve, EVisuProperty::RedCurve);
+				SegmentIndex++;
+				//Wait(SegmentIndex > 22);
+			}
+#endif
+		}
+	}
+}
+
+//#define DEBUG_FIND_CLOSE_SEGMENTS
+void FThinZone2DFinder::FindCloseSegments(const TArray<FEdgeSegment*> Segments, const TArray< const TArray<FEdgeSegment*>*> OppositeSides)
+{
+#ifdef DEBUG_FIND_CLOSE_SEGMENTS
+	F3DDebugSession _(TEXT("FindCloseSegments"));
+	{
+		F3DDebugSession _(TEXT("OppositeSide"));
+		for (const TArray<FEdgeSegment*>* OppositeSidePtr : OppositeSides)
+		{
+			ThinZone::DisplayThinZoneSide(*OppositeSidePtr, 0, EVisuProperty::YellowCurve);
+		}
+	}
+#endif
+
+	int32 SegmentToBeginIndex = 0;
+	for (FEdgeSegment* Segment : Segments)
+	{
+
+#ifdef DEBUG_FIND_CLOSE_SEGMENTS
+		{
+			F3DDebugSession _(TEXT("Segment"));
+			ThinZone::DisplayEdgeSegment(Segment, Segment->HasMarker1() ? EVisuProperty::RedCurve : EVisuProperty::BlueCurve);
+		}
+#endif
+
+		if (!Segment->HasMarker1())
+		{
+			continue;
+		}
+		Segment->ResetMarker1();
+
+		FPoint2D SegmentMiddle = Segment->GetCenter();
+
+		FPoint2D ClosePoint, ClosePoint2;
+		FEdgeSegment* CloseSegment = nullptr;
+
+		double MinSquareThickness = HUGE_VALUE;
+
+#ifdef DEBUG_FIND_CLOSE_SEGMENTS_2
+		//if (SegmentIndex == SegmentID)
+		F3DDebugSession _(bDisplay, TEXT("Segment"));
+		{
+			F3DDebugSession _(bDisplay, TEXT("Segment"));
+			DisplaySegmentWithScale(Segment->GetExtemity(ELimit::Start), Segment->GetExtemity(ELimit::End), 0, EVisuProperty::YellowCurve);
+		}
+#endif
+
+		for (const TArray<FEdgeSegment*>* OppositeSidePtr : OppositeSides)
+		{
+			const TArray<FEdgeSegment*>& OppositeSide = *OppositeSidePtr;
+			for (FEdgeSegment* Candidate : OppositeSide)
+			{
+				// check the angle between segments. As they are opposite, the Angle as to be close to Pi i.e. Slop ~= 4
+				double Slope = Segment->ComputeUnorientedSlopeOf(Candidate);
+				if (Slope < 3.) // Angle < 3Pi/4 (135 deg)
+				{
+					continue;
+				}
+
+				const FPoint2D& FirstPointCandidate = Candidate->GetExtemity(ELimit::Start);
+				const FPoint2D& SecondPointCandidate = Candidate->GetExtemity(ELimit::End);
+
+				double Coordinate;
+				const FPoint2D Projection = ProjectPointOnSegment(SegmentMiddle, FirstPointCandidate, SecondPointCandidate, Coordinate, true);
+
+				const double SquareDistance = SegmentMiddle.SquareDistance(Projection);
+				if (MinSquareThickness > SquareDistance)
+				{
+					// check the angle between segment and Middle-Projection. As they are opposite, the cosAngle as to be close to 0
+					Slope = Segment->ComputeUnorientedSlopeOf(SegmentMiddle, Projection);
+					if (Slope < 1. || Slope > 3.)
+					{
+						continue;
+					}
+
+					MinSquareThickness = SquareDistance;
+					ClosePoint = Projection;
+					CloseSegment = Candidate;
+				}
+			}
+		}
+
+		if (CloseSegment)
+		{
+			Segment->SetCloseSegment(CloseSegment, MinSquareThickness);
+#ifdef DEBUG_FIND_CLOSE_SEGMENTS
+			if (bDisplay)
+			{
+				F3DDebugSession _(bDisplay, TEXT("Close Segment"));
+				ThinZone::DisplayEdgeSegmentAndProjection(Segment, CloseSegment, EVisuProperty::BlueCurve, EVisuProperty::BlueCurve, EVisuProperty::RedCurve);
+			}
+#endif
+		}
+	}
+}
+
+//#define DEBUG_LINK_CLOSE_SEGMENTS
+void FThinZone2DFinder::LinkCloseSegments()
+{
+	TArray<FEdgeSegment*> ThinZoneSegments;
+	ThinZoneSegments.Reserve(LoopSegments.Num());
+
+	FIdent ChainIndex = Ident::Undefined;
+	FTopologicalLoop* Loop = nullptr;
+	bool bThinZone = false;
+
+	// Add in ThinZoneSegments all segments of thin zone 
+	for (FEdgeSegment* EdgeSegment : LoopSegments)
+	{
+		if (Loop != EdgeSegment->GetEdge()->GetLoop())
+		{
+			Loop = EdgeSegment->GetEdge()->GetLoop();
+			if(bThinZone)
+			{
+				ChainIndex++;
+				bThinZone = false;
+			}
+		}
+
+		if (const FEdgeSegment* CloseSegment = EdgeSegment->GetCloseSegment())
+		{
+			// Pick case i.e. all the segments are connected, 
+			//    _____ . ____ .
+			//   /
+			//  . --- . ---- . --- .
+			// As soon as a CloseSegment chain index is set and is equal to index, this mean that the segment is on the other side of the pick.
+			if (CloseSegment->GetChainIndex() == ChainIndex)
+			{
+				ChainIndex++; // new chain
+			}
+
+			EdgeSegment->SetChainIndex(ChainIndex);
+			ThinZoneSegments.Add(EdgeSegment);
+			bThinZone = true;
+		}
+		else if (bThinZone)
+		{
+			ChainIndex++;  // new chain
+			bThinZone = false;
+		}
+	}
+
+#ifdef DEBUG_LINK_CLOSE_SEGMENTS
+	{
+		DisplaySegmentsOfThinZone();
+	}
+#endif
+
+	if (ThinZoneSegments.IsEmpty())
+	{
+		return;
+	}
+
+	// Fill ThinZoneSides
+	{
+		// Reserve
+		ThinZoneSides.SetNum(ChainIndex);
+		ChainIndex = Ident::Undefined;
+		int32 SegmentCount = 0;
+		for (FEdgeSegment* EdgeSegment : ThinZoneSegments)
+		{
+			if (ChainIndex != EdgeSegment->GetChainIndex())
+			{
+				if (SegmentCount)
+				{
+					TArray<FEdgeSegment*>& ThinZoneSide = ThinZoneSides[ChainIndex];
+					ThinZoneSide.Reserve(ThinZoneSide.Max() + SegmentCount);
+					SegmentCount = 0;
+				}
+				ChainIndex = EdgeSegment->GetChainIndex();
+			}
+			SegmentCount++;
+		}
+		ThinZoneSides.Emplace_GetRef().Reserve(SegmentCount);
+
+		// Fill
+		ChainIndex = Ident::Undefined;
+		TArray<FEdgeSegment*>* ThinZoneSide = nullptr;
+		for (FEdgeSegment* EdgeSegment : ThinZoneSegments)
+		{
+			if (ChainIndex != EdgeSegment->GetChainIndex())
+			{
+				ChainIndex = EdgeSegment->GetChainIndex();
+				ThinZoneSide = &ThinZoneSides[ChainIndex];
+			}
+			ThinZoneSide->Add(EdgeSegment);
+		}
+	}
+}
+
+void FThinZone2DFinder::ImproveThinSide()
+{
+	TFunction<TArray<FEdgeSegment*>(const FEdgeSegment*)> GetComplementary = [](const FEdgeSegment* EdgeSegment) ->TArray<FEdgeSegment*>
+	{
+		// We allow to extend a thin zone up to 4 times the local thickness if it lets to link another ThinZone
+		constexpr double ComplementaryFactor = 4.;
+
+		TArray<FEdgeSegment*> ComplementaryEdges;
+		ComplementaryEdges.Reserve(100);
+
+		double Length = 0;
+		double MaxLength = FMath::Sqrt(EdgeSegment->GetCloseSquareDistance()) * ComplementaryFactor;
+
+		for (FEdgeSegment* PreviousSegment = EdgeSegment->GetPrevious();; PreviousSegment = PreviousSegment->GetPrevious())
+		{
+			if (PreviousSegment->GetCloseSegment())
+			{
+				break;
+			}
+			Length += PreviousSegment->GetLength();
+			ComplementaryEdges.Add(PreviousSegment);
+
+			if (Length > MaxLength)
+			{
+				ComplementaryEdges.Empty();
+				break;
+			}
+		}
+
+		Algo::Reverse(ComplementaryEdges);
+		return MoveTemp(ComplementaryEdges);
+	};
+
+	TFunction<void(TArray<FEdgeSegment*>&, TArray<FEdgeSegment*>&)> MergeChains = [&](TArray<FEdgeSegment*>& ThinZoneSide, TArray<FEdgeSegment*>& ComplementaryEdges)
+	{
+		FEdgeSegment* PreviousSegment = ComplementaryEdges[0]->GetPrevious();
+		const FIdent ChainId = ThinZoneSide[0]->GetChainIndex();
+
+		TArray<FEdgeSegment*> NewChain = MoveTemp(ThinZoneSides[PreviousSegment->GetChainIndex()]);
+		NewChain.Reserve(NewChain.Num() + ThinZoneSide.Num() + ComplementaryEdges.Num());
+		NewChain.Append(ComplementaryEdges);
+		NewChain.Append(ThinZoneSide);
+		ThinZoneSide = MoveTemp(NewChain);
+
+		for (FEdgeSegment* Segment : ThinZoneSide)
+		{
+			Segment->SetChainIndex(ChainId);
+		}
+	};
+
+	TFunction<void(TArray<FEdgeSegment*>&)> MergeWithPreviousChain = [&](TArray<FEdgeSegment*>& ThinZoneSide)
+	{
+		const FEdgeSegment* PreviousSegment = ThinZoneSide[0]->GetPrevious();
+		const FIdent ChainId = ThinZoneSide[0]->GetChainIndex();
+
+		TArray<FEdgeSegment*> NewChain = MoveTemp(ThinZoneSides[PreviousSegment->GetChainIndex()]);
+		NewChain.Reserve(NewChain.Num() + ThinZoneSide.Num());
+		NewChain.Append(ThinZoneSide);
+		ThinZoneSide = MoveTemp(NewChain);
+
+		for (FEdgeSegment* Segment : ThinZoneSide)
+		{
+			Segment->SetChainIndex(ChainId);
+		}
+	};
+
+	TFunction<void(TArray<FEdgeSegment*>&)> SetMarker1ToComplementarySegments = [&](TArray<FEdgeSegment*>& Complementary)
+	{
+		for (FEdgeSegment* Segment : Complementary)
+		{
+			Segment->SetMarker1();
+		}
+	};
+
+	// two adjacent chains can be:
+	// - connected. In this case, the both chains are merged
+	// - separated by a small chain slightly too far to be considered as close. If this kind of chain is find, the three chains are merged.
+	for (TArray<FEdgeSegment*>& Side : ThinZoneSides)
+	{
+		if (Side.IsEmpty())
+		{
+			continue;
+		}
+
+		const FEdgeSegment* Segment0Side = Side[0];
+		if (Segment0Side->GetPrevious()->GetCloseSegment() != nullptr)
+		{
+			// Connected case
+			// 
+			//                  Segment0Side->GetPrevious()            
+			//                  | Segment0Side            
+			//       Side(n-1)  | |  Side(n)       Side(n-1) = ThinZoneSides[n-1]
+			//    #--------------#-----------#
+			//
+			//    #--------------------------#
+			//                  | |
+			//                  | Segment0Side->GetCloseSegment()
+			//                  Segment0Side->GetPrevious()->GetCloseSegment()
+
+			FIdent CloseOfPreviousSideIndex = Segment0Side->GetPrevious()->GetCloseSegment()->GetChainIndex();
+			if (CloseOfPreviousSideIndex != Side[0]->GetChainIndex() && CloseOfPreviousSideIndex == Segment0Side->GetCloseSegment()->GetChainIndex())
+			{
+				MergeWithPreviousChain(Side);
+			}
+		}
+		else
+		{
+			TArray<FEdgeSegment*> Complementary = GetComplementary(Segment0Side);
+			if (Complementary.Num())
+			{
+				// Case2 : separated by a small chain slightly too far to be considered as close
+
+				//                  PreviousComplementary[0]
+				//                  | Complementary[0]     Seg0Side            Side(n)   = ThinZoneSides[n]
+				//       Side(n-1)  | |   Complementary    |  Side(n)          Side(n-1) = ThinZoneSides[n-1]
+				//    #--------------#--------------------#-----------#
+				//
+				//    #-----------------------------------------------#
+				//                  |                      |
+				//                  |                      Seg0Side->GetCloseSegment()
+				//                  ClosePreviousComplementary[0]
+				const FEdgeSegment* PreviousComplementary = Complementary[0]->GetPrevious();
+				const FEdgeSegment* ClosePreviousComplementary = PreviousComplementary->GetCloseSegment();
+
+
+				// Case Side(n-1) and Side(n) are closed together
+				// 
+				//                               PreviousComplementary[0]
+				//       Side(n-1)               | 
+				//    #---------------------------#
+				//                                |
+				//                                |  Complementary
+				//    #---------------------------#
+				//                Side(n)        |       
+				//                               ClosePreviousComplementary 
+				if (ClosePreviousComplementary->GetChainIndex() == Segment0Side->GetChainIndex())
+				{
+					// Do nothing
+					continue;
+				}
+
+				FIdent CloseSideIndex = Segment0Side->GetCloseSegment()->GetChainIndex();
+				if (CloseSideIndex == ClosePreviousComplementary->GetChainIndex())
+				{
+					MergeChains(Side, Complementary);
+					SetMarker1ToComplementarySegments(Complementary);
+				}
+				else
+				{
+					//                  PreviousComplementary[0]
+					//                  | Complementary[0]     Seg0Side            Side(n)   = ThinZoneSides[n]
+					//       Side(n-1)  | |   Complementary    |  Side(n)          Side(n-1) = ThinZoneSides[n-1]
+					//    #--------------#--------------------#-----------#  ->
+					//
+					//    #--------------#--------------------#-----------#  <-
+					//                  |  ComplementaryClose  |
+					//                  |                      Seg0Side->GetCloseSegment() == PreviousComplementaryClose[0]
+					//                  ClosePreviousComplementary[0]
+					TArray<FEdgeSegment*> ComplementaryClose = GetComplementary(ClosePreviousComplementary);
+					if (ComplementaryClose.Num())
+					{
+						const FEdgeSegment* PreviousComplementaryClose = ComplementaryClose[0]->GetPrevious(); // == Seg0Side->GetCloseSegment()
+						const FEdgeSegment* CloseOfPreviousComplementaryClose = PreviousComplementaryClose->GetCloseSegment(); // == Seg0Side
+
+						if (CloseOfPreviousComplementaryClose->GetChainIndex() == Segment0Side->GetChainIndex())
+						{
+							MergeChains(Side, Complementary);
+							SetMarker1ToComplementarySegments(Complementary);
+
+							TArray<FEdgeSegment*>& SideOfClosePreviousComplementary = ThinZoneSides[ClosePreviousComplementary->GetChainIndex()];
+							MergeChains(SideOfClosePreviousComplementary, ComplementaryClose);
+							SetMarker1ToComplementarySegments(ComplementaryClose);
+						}
+					}
+				}
+			}
+		}
+	}
+
+	TFunction<const TArray<const TArray<FEdgeSegment*>*>(const TArray<FEdgeSegment*>&)> FindOppositeSides = [&](const TArray<FEdgeSegment*>& Side) -> const TArray<const TArray<FEdgeSegment*>*>
+	{
+		TArray<int32> OppositeSideIndexes;
+		OppositeSideIndexes.Reserve(5);
+		int32 LastIndex = -1;
+		for (const FEdgeSegment* Seg : Side)
+		{
+			const FEdgeSegment* Opposite = Seg->GetCloseSegment();
+			if (Opposite != nullptr)
+			{
+				const int32 Index = Opposite->GetChainIndex();
+				if (Index != LastIndex)
+				{
+					OppositeSideIndexes.AddUnique(Index);
+					LastIndex = Index;
+				}
+			}
+		}
+
+		TArray<const TArray<FEdgeSegment*>*> OppositeSides;
+		OppositeSides.Reserve(OppositeSideIndexes.Num());
+		for (int32 Index : OppositeSideIndexes)
+		{
+			OppositeSides.Add(&ThinZoneSides[Index]);
+		}
+
+		return MoveTemp(OppositeSides);
+	};
+
+	// Finalize:
+	// Find close segments to the added segments
+	for (TArray<FEdgeSegment*>& Side : ThinZoneSides)
+	{
+		if (!Algo::AllOf(Side, [](const FEdgeSegment* Seg) { return !Seg->HasMarker1(); }))
+		{
+			const TArray<const TArray<FEdgeSegment*>*> OppositeSides = FindOppositeSides(Side);
+			FindCloseSegments(Side, OppositeSides);
+		}
+	}
+
+#ifdef DEBUG_THIN_ZONES_IMPROVE
+	DisplayThinZoneSides(ThinZoneSides);
+	Wait();
+#endif
+}
+
+void FThinZone2DFinder::SplitThinSide()
+{
+	TArray<TArray<FEdgeSegment*>> NewThinZoneSides;
+	FIdent NewSideIndex = ThinZoneSides.Num();
+
+	for (TArray<FEdgeSegment*>& ThinZoneSide : ThinZoneSides)
+	{
+		if (ThinZoneSide.IsEmpty())
+		{
+			continue;
+		}
+
+		FIdent CloseSideIndex = ThinZoneSide[0]->GetCloseSegment()->GetChainIndex();
+		if (!Algo::AllOf(ThinZoneSide, [&](const FEdgeSegment* EdgeSegment) { return EdgeSegment->GetCloseSegment()->GetChainIndex() == CloseSideIndex; }))
+		{
+			int32 Index = 0;
+			for (; Index < ThinZoneSide.Num(); ++Index)
+			{
+				if (ThinZoneSide[Index]->GetCloseSegment()->GetChainIndex() != CloseSideIndex)
+				{
+					break;
+				}
+			}
+
+			const int32 LastIndexFirstSide = Index;
+			while (Index < ThinZoneSide.Num())
+			{
+				CloseSideIndex = ThinZoneSide[Index]->GetCloseSegment()->GetChainIndex();
+				TArray<FEdgeSegment*>& NewThinZoneSide = NewThinZoneSides.Emplace_GetRef();
+				for (; Index < ThinZoneSide.Num(); ++Index)
+				{
+					FEdgeSegment* Segment = ThinZoneSide[Index];
+					if (Segment->GetCloseSegment()->GetChainIndex() != CloseSideIndex)
+					{
+						break;
+					}
+					NewThinZoneSide.Add(Segment);
+					Segment->SetChainIndex(NewSideIndex);
+				}
+				++NewSideIndex;
+			}
+			ThinZoneSide.SetNum(LastIndexFirstSide);
+		}
+	}
+
+	for (TArray<FEdgeSegment*>& ThinZoneSide : NewThinZoneSides)
+	{
+		TArray<FEdgeSegment*>& NewThinZoneSide = ThinZoneSides.Emplace_GetRef();
+		NewThinZoneSide = MoveTemp(ThinZoneSide);
+	}
+
+}
+
+void FThinZone2DFinder::BuildThinZone()
+{
+#ifdef DEBUG_BUILD_THIN_ZONES
+	ThinZone::DisplayThinZoneSides(ThinZoneSides);
+	Wait();
+#endif
+
+	// The number of ThinZone should be less than the number of ThinZoneSides, 
+	ThinZones.Reserve(ThinZoneSides.Num());
+
+	for (TArray<FEdgeSegment*>& FirstSide : ThinZoneSides)
+	{
+		if (FirstSide.IsEmpty())
+		{
+			continue;
+		}
+
+		const FEdgeSegment* FirstSideSegment = FirstSide[0];
+		TArray<FEdgeSegment*>& SecondSide = ThinZoneSides[FirstSideSegment->GetCloseSegment()->GetChainIndex()];
+		if (SecondSide.IsEmpty())
+		{
+			continue;
+		}
+
+		BuildThinZone(FirstSide, SecondSide);
+
+		FirstSide.Empty(); // to avoid to rebuild a second thin zone with SecondSide & FirstSide
+		SecondSide.Empty();
+	}
+}
+
+void FThinZone2DFinder::GetThinZoneSideConnectionsLength(const TArray<FEdgeSegment*>& FirstSide, const TArray<FEdgeSegment*>& SecondSide, double InMaxLength, double* OutLengthBetweenExtremities, TArray<const FTopologicalEdge*>* OutPeakEdges)
+{
+	TFunction<double(const FEdgeSegment*, const FEdgeSegment*, TFunction<const FEdgeSegment* (const FEdgeSegment*)>, TArray<const FTopologicalEdge*>&)> ComputeLengthBetweenExtremities = 
+		[&](const FEdgeSegment* Start, const FEdgeSegment* End, TFunction<const FEdgeSegment* (const FEdgeSegment*)> GetNext, TArray<const FTopologicalEdge*>& PeakEdges) -> double
+	{
+		PeakEdges.Reserve(10);
+
+		double LengthBetweenExtremities = 0;
+		const FEdgeSegment* Segment = GetNext(Start);
+		const FTopologicalEdge* Edge = nullptr;
+
+		while (Segment != End)
+		{
+			if (Edge != Segment->GetEdge())
+			{
+				Edge = Segment->GetEdge();
+				PeakEdges.Add(Edge);
+			}
+
+			LengthBetweenExtremities += Segment->GetLength();
+			if (LengthBetweenExtremities > InMaxLength)
+			{
+				LengthBetweenExtremities = HUGE_VALUE;
+				break;
+			}
+			Segment = GetNext(Segment);
+		}
+		return LengthBetweenExtremities;
+	};
+
+	OutLengthBetweenExtremities[0] = ComputeLengthBetweenExtremities(FirstSide[0], SecondSide.Last(), [](const FEdgeSegment* Segment) { return Segment->GetPrevious(); }, OutPeakEdges[0]);
+	OutLengthBetweenExtremities[1] = ComputeLengthBetweenExtremities(FirstSide.Last(), SecondSide[0], [](const FEdgeSegment* Segment) { return Segment->GetNext(); }, OutPeakEdges[1]);
+}
+
+void FThinZone2DFinder::BuildThinZone(const TArray<FEdgeSegment*>& FirstSide, const TArray<FEdgeSegment*>& SecondSide)
+{
+	FThinZone2D& Zone = ThinZones.Emplace_GetRef(FirstSide, SecondSide);
+
+	if (!FirstSide[0]->IsInner() && !SecondSide[0]->IsInner())
+	{
+		const double MaxThickness = Zone.GetMaxThickness();
+		double LengthBetweenExtremity[2] = { HUGE_VALUE, HUGE_VALUE };
+		TArray<const FTopologicalEdge*> PeakEdges[2];
+		GetThinZoneSideConnectionsLength(FirstSide, SecondSide, 3. * MaxThickness, LengthBetweenExtremity, PeakEdges);
+
+		//                     Side 0 
+		//       #-------------------------------------# 
+		//      /
+		//     /  <- LengthBetweendSidesToBePeak  
+		//    /
+		//   #-----------------------------------------# 
+		//                     Side 1 
+		//
+		// Two EdgeSegments are close if, among other things, the shortest distance between the both EdgeSegments make an angle smallest than 45 deg with the EdgeSegments.
+		// So the MaxLengthBetweendSidesToBePeak is theoretically MaxThickness x Sqrt(2)...
+		// It's simplified to MaxThickness x 2.
+		const double MaxLengthBetweendSidesToBeAPeak = MaxThickness * 2.;
+		const double ThinZoneLength = Zone.Length();
+		const double MinThinZoneLengthToBeGlobal = ExternalLoopLength - 2. * MaxLengthBetweendSidesToBeAPeak;
+
+		if (ThinZoneLength > MinThinZoneLengthToBeGlobal || (LengthBetweenExtremity[0] < MaxLengthBetweendSidesToBeAPeak && LengthBetweenExtremity[1] < MaxLengthBetweendSidesToBeAPeak))
+		{
+			Zone.SetCategory(EThinZone2DType::Global);
+		}
+		else if (LengthBetweenExtremity[0] < MaxLengthBetweendSidesToBeAPeak)
+		{
+			if (ThinZoneLength < MaxThickness * 5.)
+			{
+				// The thin zone is too small, it's deleted
+				ThinZones.Pop();
+				return;
+
+			}
+			else
+			{
+				Zone.SetCategory(EThinZone2DType::PeakStart);
+				Zone.SetPeakEdgesMarker(PeakEdges[0]);
+			}
+		}
+		else if (LengthBetweenExtremity[1] < MaxLengthBetweendSidesToBeAPeak)
+		{
+			if (ThinZoneLength < MaxThickness * 5.)
+			{
+				// The thin zone is too small, it's deleted
+				ThinZones.Pop();
+				return;
+
+			}
+			else
+			{
+				Zone.SetCategory(EThinZone2DType::PeakEnd);
+				Zone.SetPeakEdgesMarker(PeakEdges[1]);
+			}
+		}
+		else
+		{
+			Zone.SetCategory(EThinZone2DType::Butterfly);
+		}
+
+	}
+	else
+	{
+		Zone.SetCategory(EThinZone2DType::BetweenLoops);
+	}
+}
+
+bool FThinZone2DFinder::SearchThinZones(double InTolerance)
+{
+#ifdef DEBUG_THIN_ZONES
+	F3DDebugSession A(bDisplay, FString::Printf(TEXT("ThinZone Finder %d"), Grid.GetFace().GetId()));
+	if (bDisplay)
+	{
+		F3DDebugSession _(TEXT("Thin surface grid"));
+		Grid.DisplayGridPoints(EGridSpace::UniformScaled);
+	}
+#endif
+
+	SetTolerance(InTolerance);
+
+	//FMessage::Printf(DBG, TEXT("Searching thin zones on Surface %d\n", Surface->GetId());
+	FTimePoint StartTime = FChrono::Now();
+	if (!BuildLoopSegments())
+	{
+		Face.SetAsDegenerated();
+		return false;
+	}
+	Chronos.BuildLoopSegmentsTime = FChrono::Elapse(StartTime);
+
+#ifdef DEBUG_THIN_ZONES
+	DisplayLoopSegments();
+#endif
+
+	StartTime = FChrono::Now();
+	FindCloseSegments();
+	Chronos.FindCloseSegmentTime = FChrono::Elapse(StartTime);
+
+#ifdef DEBUG_SEARCH_THIN_ZONES
+	if (bDisplay)
+	{
+		DisplayCloseSegments();
+		Wait();
+	}
+#endif
+
+	StartTime = FChrono::Now();
+	LinkCloseSegments();
+
+#ifdef DEBUG_SEARCH_THIN_ZONES
+	if(bDisplay)
+	{
+		ThinZone::DisplayThinZoneSides(ThinZoneSides);
+	}
+#endif
+
+	ImproveThinSide();
+
+#ifdef DEBUG_SEARCH_THIN_ZONES
+	if (bDisplay)
+	{
+		ThinZone::DisplayThinZoneSides(ThinZoneSides);
+		Wait();
+	}
+#endif
+
+	SplitThinSide();
+
+#ifdef DEBUG_SEARCH_THIN_ZONES
+	if (bDisplay)
+	{
+		ThinZone::DisplayThinZoneSides(ThinZoneSides);
+		Wait();
+	}
+#endif
+	Chronos.LinkCloseSegmentTime = FChrono::Elapse(StartTime);
+
+	StartTime = FChrono::Now();
+	BuildThinZone();
+	Chronos.BuildThinZoneTime = FChrono::Elapse(StartTime);
+
+#ifdef DEBUG_THIN_ZONES
+	if(bDisplay)
+	{
+		ThinZone::DisplayThinZones(ThinZones);
+		Wait(false);
+	}
+#endif
+
+	return (ThinZones.Num() > 0);
+}
+
+bool FThinZone2DFinder::BuildLoopSegments()
+{
+	const double GeometricTolerance = Grid.GetTolerance();
+	const double WishedSegmentLength = FinderTolerance / 5.;
+	const TArray<TSharedPtr<FTopologicalLoop>>& Loops = Face.GetLoops();
+
+	{
+		double Length = 0;
+		ExternalLoopLength = -1.;
+		for (const TSharedPtr<FTopologicalLoop>& Loop : Loops)
+		{
+			const double LoopLength = Loop->Length();
+			Length += LoopLength;
+
+			if (Loop->IsExternal())
+			{
+				ExternalLoopLength = LoopLength;
+			}
+		}
+
+		const int32 SegmentNum = (int32)(1.2 * Length / WishedSegmentLength);
+		LoopSegments.Empty(SegmentNum);
+	}
+
+	const TSharedPtr<FTopologicalLoop> OuterLoop = Loops[0];
+	for (const TSharedPtr<FTopologicalLoop>& Loop : Loops)
+	{
+		const bool bIsInnerLoop = (Loop != OuterLoop);
+		const TArray<FOrientedEdge>& Edges = Loop->GetEdges();
+
+		FEdgeSegment* FirstSegment = nullptr;
+		FEdgeSegment* PrecedingSegment = nullptr;
+
+		for (const FOrientedEdge& Edge : Edges)
+		{
+			TArray<double> Coordinates;
+			Edge.Entity->Sample(WishedSegmentLength, Coordinates);
+
+			TArray<FPoint2D> ScaledPoints;
+			{
+				TArray<FPoint2D> Points;
+				Edge.Entity->Approximate2DPoints(Coordinates, Points);
+				Grid.TransformPoints(EGridSpace::UniformScaled, Points, ScaledPoints);
+			}
+
+			// Remove duplicated points (vs GeometricTolerance) of the end
+			{
+				int32 EndPointIndex = ScaledPoints.Num() - 1;
+				while (EndPointIndex > 0 && ScaledPoints[EndPointIndex].Distance(ScaledPoints[EndPointIndex - 1]) < GeometricTolerance)
+				{
+					EndPointIndex--;
+					ScaledPoints.RemoveAt(EndPointIndex);
+					Coordinates.RemoveAt(EndPointIndex);
+				}
+
+				if (ScaledPoints.IsEmpty())
+				{
+					continue;
+				}
+			}
+
+			// Remove duplicated points (vs GeometricTolerance)
+			{
+				for (int32 PointIndex = 1; PointIndex < ScaledPoints.Num(); PointIndex++)
+				{
+					while (ScaledPoints[PointIndex - 1].Distance(ScaledPoints[PointIndex]) < GeometricTolerance)
+					{
+						ScaledPoints.RemoveAt(PointIndex);
+						Coordinates.RemoveAt(PointIndex);
+					}
+				}
+			}
+
+			const int32 LastPointIndex = ScaledPoints.Num() - 1;
+
+			const int32 Increment = (Edge.Direction == EOrientation::Front) ? 1 : -1;
+			int32 ISegment2 = (Edge.Direction == EOrientation::Front) ? 1 : LastPointIndex - 1;
+
+			for (; ISegment2 >= 0 && ISegment2 <= LastPointIndex; ISegment2 += Increment)
+			{
+				const int32 ISegment1 = ISegment2 - Increment;
+				FEdgeSegment& CurrentSeg = SegmentFatory.New();
+				CurrentSeg.SetBoundarySegment(bIsInnerLoop, Edge.Entity.Get(), Coordinates[ISegment1], Coordinates[ISegment2], ScaledPoints[ISegment1], ScaledPoints[ISegment2]);
+
+				LoopSegments.Add(&CurrentSeg);
+				if (!FirstSegment)
+				{
+					FirstSegment = PrecedingSegment = &CurrentSeg;
+				}
+				else
+				{
+					PrecedingSegment->SetNext(&CurrentSeg);
+				}
+				PrecedingSegment = &CurrentSeg;
+			}
+		}
+
+		if(PrecedingSegment)
+		{
+			PrecedingSegment->SetNext(FirstSegment);
+		}
+
+		if (!bIsInnerLoop && LoopSegments.Num() < 2)
+		{
+			return false;
+		}
+	}
+
+	return true;
+}
+
+FThinZoneSide::FThinZoneSide(FThinZoneSide* InFrontSide, const TArray<FEdgeSegment*>& InSegments)
+	: FrontSide(*InFrontSide)
+	, SideLength(-1)
+	, MediumThickness(-1)
+	, MaxThickness(-1)
+{
+	const int32 SegmentCount = InSegments.Num();
+	Segments.Reserve(SegmentCount);
+	for (const FEdgeSegment* Segment : InSegments)
+	{
+		Segments.Emplace(*Segment);
+	}
+}
+
+void FThinZone2D::Finalize()
+{
+	TFunction<void(TArray<FEdgeSegment>&, TMap<int32, FEdgeSegment*>&)> BuildMap = [](TArray<FEdgeSegment>& Segments, TMap<int32, FEdgeSegment*>& Map)
+	{
+		for (FEdgeSegment& Segment : Segments)
+		{
+			Map.Emplace(Segment.GetId(), &Segment);
+		}
+	};
+
+	TFunction<void(TArray<FEdgeSegment>&, TMap<int32, FEdgeSegment*>&)> UpdateReference = [](TArray<FEdgeSegment>& Segments, TMap<int32, FEdgeSegment*>& Map)
+	{
+		for (FEdgeSegment& Segment : Segments)
+		{
+			Segment.UpdateReferences(Map);
+		}
+	};
+
+	TMap<int32, FEdgeSegment*> NewSegmentMap;
+
+	TArray<FEdgeSegment>& FirstSideSegments = FirstSide.GetSegments();
+	TArray<FEdgeSegment>& SecondSideSegments = SecondSide.GetSegments();
+	BuildMap(FirstSideSegments, NewSegmentMap);
+	BuildMap(SecondSideSegments, NewSegmentMap);
+
+	UpdateReference(FirstSideSegments, NewSegmentMap);
+	UpdateReference(SecondSideSegments, NewSegmentMap);
+
+	FirstSide.ComputeThicknessAndLength();
+	SecondSide.ComputeThicknessAndLength();
+
+	Thickness = (FirstSide.GetThickness() + SecondSide.GetThickness()) * 0.5;
+	MaxThickness = FMath::Max(FirstSide.GetMaxThickness(), SecondSide.GetMaxThickness());
+}
+
+void FThinZoneSide::ComputeThicknessAndLength()
+{
+	SideLength = 0;
+	double SquareMediumThickness = 0;
+	double SquareMaxThickness = 0;
+
+	for (const FEdgeSegment& Segment : Segments)
+	{
+		const double SquareThickness = Segment.GetCloseSquareDistance();
+		const double SegmentLength = Segment.GetLength();
+		SideLength += SegmentLength;
+		SquareMediumThickness += SquareThickness * SegmentLength;
+		if (SquareMaxThickness < SquareThickness)
+		{
+			SquareMaxThickness = SquareThickness;
+		}
+	}
+
+	SquareMediumThickness /= SideLength;
+	MediumThickness = FMath::Sqrt(SquareMediumThickness);
+	MaxThickness = FMath::Sqrt(SquareMaxThickness);
+}
+
+void FThinZoneSide::SetEdgesAsThinZone()
+{
+	FTopologicalEdge* Edge = nullptr;
+	for (FEdgeSegment& Segment : Segments)
+	{
+		if (Edge != Segment.GetEdge())
+		{
+			Edge = Segment.GetEdge();
+			Edge->SetThinZoneMarker();
+		}
+	}
+}
+
+void FThinZone2D::SetPeakEdgesMarker(const TArray<const FTopologicalEdge*>& PeakEdges)
+{
+	for (const FTopologicalEdge* Edge : PeakEdges)
+	{
+		Edge->SetThinPeakMarker();
+	}
+}
+
+EMeshingState FThinZoneSide::GetMeshingState() const
+{
+	EMeshingState MeshingState = NotMeshed;
+	int32 EdgeCount = 0;
+	int32 MeshedEdgeCount = 0;
+
+	const FTopologicalEdge* Edge = nullptr;
+	for (const FEdgeSegment& EdgeSegment : Segments)
+	{
+		if (Edge != EdgeSegment.GetEdge())
+		{
+			Edge = EdgeSegment.GetEdge();
+			++EdgeCount;
+			if (Edge->GetLinkActiveEdge()->IsMeshed())
+			{
+				++MeshedEdgeCount;
+			}
+		}
+	}
+
+	if (MeshedEdgeCount == 0)
+	{
+		return NotMeshed;
+	}
+	else if (MeshedEdgeCount != EdgeCount)
+	{
+		return PartiallyMeshed;
+	}
+	return FullyMeshed;
+}
+
+void FThinZone2D::SetEdgesAsThinZone()
+{
+	FirstSide.SetEdgesAsThinZone();
+	SecondSide.SetEdgesAsThinZone();
+}
+
+}
+
diff --git a/Engine/Source/Runtime/Datasmith/CADKernel/Private/CADKernel/Mesh/Structure/ThinZone2DFinder.h b/Engine/Source/Runtime/Datasmith/CADKernel/Private/CADKernel/Mesh/Structure/ThinZone2DFinder.h
new file mode 100644
--- /dev/null
+++ b/Engine/Source/Runtime/Datasmith/CADKernel/Private/CADKernel/Mesh/Structure/ThinZone2DFinder.h
@@ -0,0 +1,159 @@
+// Copyright Epic Games, Inc. All Rights Reserved.
+#pragma once
+
+#include "CADKernel/Core/Chrono.h"
+#include "CADKernel/Core/HaveStates.h"
+#include "CADKernel/Core/Factory.h"
+#include "CADKernel/Core/Types.h"
+#include "CADKernel/Math/Point.h"
+
+#include "CADKernel/Mesh/Structure/EdgeSegment.h"
+#include "CADKernel/Mesh/Structure/GridBase.h"
+#include "CADKernel/Mesh/Structure/ThinZone2D.h"
+
+//#ifdef CADKERNEL_DEV
+//#include "CADKernel/UI/DefineForDebug.h"
+//#endif
+
+namespace UE::CADKernel
+{
+	class FTopologicalLoop;
+	class FTopologicalEdge;
+	class FGrid;
+	class FThinZone2D;
+	class FThinZone2DFinder;
+
+	struct FThinZoneChronos
+	{
+		FDuration BuildLoopSegmentsTime;
+		FDuration DisplayBoundarySegmentTime;
+		FDuration FindCloseSegmentTime;
+		FDuration LinkCloseSegmentTime;
+		FDuration BuildThinZoneTime;
+	};
+
+	class FThinZone2DFinder
+	{
+		friend class FThinZone2D;
+
+	public:
+		FThinZoneChronos Chronos;
+
+	protected:
+		const FGridBase& Grid;
+		const FTopologicalFace& Face;
+
+		double FinderTolerance;
+		double SquareFinderTolerance;
+
+		TArray<FThinZone2D> ThinZones;
+
+		TArray<FEdgeSegment*> LoopSegments;
+		TArray<TArray<FEdgeSegment*>> ThinZoneSides;
+
+		double ExternalLoopLength = 0.0;
+
+		TFactory<FEdgeSegment> SegmentFatory;
+
+	public:
+
+		FThinZone2DFinder(const FGridBase& InGrid, const FTopologicalFace& InFace)
+			: Grid(InGrid)
+			, Face(InFace)
+			, SegmentFatory()
+		{
+#ifdef DEBUG_CADKERNEL
+			bDisplay = Grid.bDisplay;
+#endif
+		}
+
+		bool SearchThinZones(double InFinderTolerance);
+
+		TArray<FThinZone2D>& GetThinZones()
+		{
+			return ThinZones;
+		}
+
+
+	private:
+
+		void SetTolerance(double InFinderTolerance)
+		{
+			FinderTolerance = InFinderTolerance;
+			SquareFinderTolerance = FMath::Square(FinderTolerance);
+		}
+
+		/**
+		 * @return false if the main loop is empty
+		 */
+		bool BuildLoopSegments();
+
+		void FindCloseSegments();
+
+		/**
+		 * For Segment with Marker1 in Segments, find in OppositeSides the closest segment in the respect of the criteria
+		 */
+		void FindCloseSegments(const TArray<FEdgeSegment*> Segments, const TArray< const TArray<FEdgeSegment*>*> OppositeSides);
+
+		/**
+		 * Fill ThinZoneSides (TArray<TArray<FEdgeSegment*>>)
+		 * Each TArray<FEdgeSegment*> is an array of connected segment defining a side of one (or more) thin zone
+		 */
+		void LinkCloseSegments();
+
+		/**
+		 * Two adjacent sides can be separated by a small chain of segments slightly too far to another side be considered as thin side. If this kind of chain is find, the three chains are merged.
+		 * E.g. Side(n-1) and Side(n) are close from Side(k)
+		 *                           _
+		 *                          / \ <-- a small chain of segments slightly too far                          / \ 
+		 *                         /   \												                       /   \                 
+		 *    #-------------------#     #-----------------#              ====>		      #-------------------/     \-----------------#
+		 *          Side(n-1)                  Side(n)									                      Side(n)                  
+		 *    #-------------------------------------------#								  #-------------------------------------------#
+		 *                      Side(k)                                                                       Side(k)  
+		 */
+		void ImproveThinSide();
+
+		/**
+		 * SplitThinSide splits chain in front of many other chain eg.
+		 * Side A is close to Side B then Side C
+		 * Side A has to be split in two
+		 * 
+		 *                      Side A
+		 *     #----------------------------------------#
+		 * 
+		 *     #-----------------# #--------------------# 
+		 *          Side B               Side C
+		 */
+		void SplitThinSide();
+
+		void BuildThinZone();
+		void BuildThinZone(const TArray<FEdgeSegment*>& FirstSide, const TArray<FEdgeSegment*>& SecondSide);
+		static void GetThinZoneSideConnectionsLength(const TArray<FEdgeSegment*>& FirstSide, const TArray<FEdgeSegment*>& SecondSide, double InMaxLength, double* OutLengthBetweenExtremities, TArray<const FTopologicalEdge*>* OutPeakEdges);
+
+	// ======================================================================================================================================================================================================================
+	// Display Methodes   ================================================================================================================================================================================================
+	// ======================================================================================================================================================================================================================
+#ifdef CADKERNEL_DEV
+		bool bDisplay = false;
+		void DisplayCloseSegments();
+		void DisplayLoopSegments();
+		void DisplaySegmentsOfThinZone();
+#endif
+	};
+
+#ifdef CADKERNEL_DEBUG
+namespace ThinZone
+{
+void DisplayEdgeSegmentAndProjection(const FEdgeSegment* Segment, EVisuProperty SegColor, EVisuProperty OppositeColor, EVisuProperty ProjectionColor);
+void DisplayEdgeSegmentAndProjection(const FEdgeSegment* Segment, const FEdgeSegment* Opposite, EVisuProperty SegColor, EVisuProperty OppositeColor, EVisuProperty ProjectionColor);
+void DisplayEdgeSegment(const FEdgeSegment* EdgeSegment, EVisuProperty Color);
+void DisplayEdgeSegment(const FEdgeSegment* EdgeSegment, EVisuProperty Color, int32 Index);
+void DisplayThinZoneSides(const TArray<TArray<FEdgeSegment*>>& ThinZoneSides);
+void DisplayThinZoneSide(const TArray<FEdgeSegment*>& Side, int32 Index, EVisuProperty Color, bool bSplitBySegment = false);
+void DisplayThinZoneSide(const FThinZoneSide& Side, int32 Index, EVisuProperty Color, bool bSplitBySegment = false);
+void DisplayThinZones(const TArray<FThinZone2D>& ThinZones);
+}
+#endif
+
+}
diff --git a/Engine/Source/Runtime/Datasmith/CADKernel/Private/CADKernel/Topo/Model.cpp b/Engine/Source/Runtime/Datasmith/CADKernel/Private/CADKernel/Topo/Model.cpp
--- a/Engine/Source/Runtime/Datasmith/CADKernel/Private/CADKernel/Topo/Model.cpp
+++ b/Engine/Source/Runtime/Datasmith/CADKernel/Private/CADKernel/Topo/Model.cpp
@@ -176,6 +176,8 @@ void FModel::FillTopologyReport(FTopologyReport& Report) const
 	{
 		Body->FillTopologyReport(Report);
 	}
+
+	Report.ResetMarkers();
 }
 #endif
 
diff --git a/Engine/Source/Runtime/Datasmith/CADKernel/Private/CADKernel/Topo/TopologicalEdge.cpp b/Engine/Source/Runtime/Datasmith/CADKernel/Private/CADKernel/Topo/TopologicalEdge.cpp
--- a/Engine/Source/Runtime/Datasmith/CADKernel/Private/CADKernel/Topo/TopologicalEdge.cpp
+++ b/Engine/Source/Runtime/Datasmith/CADKernel/Private/CADKernel/Topo/TopologicalEdge.cpp
@@ -530,70 +530,62 @@ int32 FTopologicalEdge::EvaluateCuttingPointNum()
 	return (int32)Num;
 }
 
-double FTopologicalEdge::TransformLocalCoordinateToActiveEdgeCoordinate(const double InLocalCoordinate)
+double FTopologicalEdge::TransformLocalCoordinateToActiveEdgeCoordinate(const double InLocalCoordinate) const
 {
 	if (IsActiveEntity())
 	{
 		return InLocalCoordinate;
 	}
 
-	TSharedPtr<FTopologicalEdge> ActiveEdge = GetLinkActiveEntity();
+	const FTopologicalEdge& ActiveEdge = *GetLinkActiveEntity();
 	FPoint PointOnEdge = Curve->Approximate3DPoint(InLocalCoordinate);
 	FPoint ProjectedPoint;
-	return ActiveEdge->GetCurve()->GetCoordinateOfProjectedPoint(Boundary, PointOnEdge, ProjectedPoint);
+	return ActiveEdge.GetCurve()->GetCoordinateOfProjectedPoint(Boundary, PointOnEdge, ProjectedPoint);
 }
 
-double FTopologicalEdge::TransformActiveEdgeCoordinateToLocalCoordinate(const double InActiveEdgeCoordinate)
+double FTopologicalEdge::TransformActiveEdgeCoordinateToLocalCoordinate(const double InActiveEdgeCoordinate) const
 {
 	if (IsActiveEntity())
 	{
 		return InActiveEdgeCoordinate;
 	}
 
-	TSharedPtr<FTopologicalEdge> ActiveEdge = GetLinkActiveEntity();
-	FPoint PointOnEdge = ActiveEdge->GetCurve()->Approximate3DPoint(InActiveEdgeCoordinate);
+	const FTopologicalEdge& ActiveEdge = *GetLinkActiveEntity();
+	FPoint PointOnEdge = ActiveEdge.GetCurve()->Approximate3DPoint(InActiveEdgeCoordinate);
 	FPoint ProjectedPoint;
 	return Curve->GetCoordinateOfProjectedPoint(Boundary, PointOnEdge, ProjectedPoint);
 }
 
-void FTopologicalEdge::TransformLocalCoordinatesToActiveEdgeCoordinates(const TArray<double>& InLocalCoordinate, TArray<double>& OutActiveEdgeCoordinates)
+void FTopologicalEdge::TransformLocalCoordinatesToActiveEdgeCoordinates(const TArray<double>& InLocalCoordinate, TArray<double>& OutActiveEdgeCoordinates) const
 {
 	if (IsActiveEntity())
 	{
 		OutActiveEdgeCoordinates = InLocalCoordinate;
 	}
-	TSharedPtr<FTopologicalEdge> ActiveEdge = GetLinkActiveEntity();
+	const FTopologicalEdge& ActiveEdge = *GetLinkActiveEntity();
 	TArray<FPoint> EdgePoints;
 	Curve->Approximate3DPoints(InLocalCoordinate, EdgePoints);
 	TArray<FPoint> ProjectedPoints;
-	ActiveEdge->GetCurve()->ProjectPoints(Boundary, EdgePoints, OutActiveEdgeCoordinates, ProjectedPoints);
+	ActiveEdge.GetCurve()->ProjectPoints(Boundary, EdgePoints, OutActiveEdgeCoordinates, ProjectedPoints);
 }
 
-void FTopologicalEdge::TransformActiveEdgeCoordinatesToLocalCoordinates(const TArray<double>& InActiveEdgeCoordinate, TArray<double>& OutLocalCoordinates)
+void FTopologicalEdge::TransformActiveEdgeCoordinatesToLocalCoordinates(const TArray<double>& InActiveEdgeCoordinate, TArray<double>& OutLocalCoordinates) const
 {
 	if (IsActiveEntity())
 	{
 		OutLocalCoordinates = InActiveEdgeCoordinate;
+		return;
 	}
 
-	TSharedPtr<FTopologicalEdge> ActiveEdge = GetLinkActiveEntity();
+	const FTopologicalEdge& ActiveEdge = *GetLinkActiveEntity();
 	TArray<FPoint> ActiveEdgePoint;
-	ActiveEdge->GetCurve()->Approximate3DPoints(InActiveEdgeCoordinate, ActiveEdgePoint);
+	ActiveEdge.GetCurve()->Approximate3DPoints(InActiveEdgeCoordinate, ActiveEdgePoint);
 	TArray<FPoint> ProjectedPoints;
 	Curve->ProjectPoints(Boundary, ActiveEdgePoint, OutLocalCoordinates, ProjectedPoints);
 }
 
 void FTopologicalEdge::AddImposedCuttingPointU(const double ImposedCuttingPointU, int32 OppositeNodeIndex)
 {
-	if (!IsActiveEntity())
-	{
-		ensureCADKernel(false);
-		FPoint Point = Curve->Approximate3DPoint(ImposedCuttingPointU);
-		FPoint ProjectedPoint;
-		double ActiveEdgeParamU = GetLinkActiveEntity()->ProjectPoint(Point, ProjectedPoint);
-		return GetLinkActiveEntity()->AddImposedCuttingPointU(ActiveEdgeParamU, OppositeNodeIndex);
-	}
-
 	ImposedCuttingPointUs.Emplace(ImposedCuttingPointU, OppositeNodeIndex);
 }
 
diff --git a/Engine/Source/Runtime/Datasmith/CADKernel/Private/CADKernel/Topo/TopologicalFace.cpp b/Engine/Source/Runtime/Datasmith/CADKernel/Private/CADKernel/Topo/TopologicalFace.cpp
--- a/Engine/Source/Runtime/Datasmith/CADKernel/Private/CADKernel/Topo/TopologicalFace.cpp
+++ b/Engine/Source/Runtime/Datasmith/CADKernel/Private/CADKernel/Topo/TopologicalFace.cpp
@@ -433,11 +433,6 @@ void FTopologicalFace::GetEdgeIndex(const FTopologicalEdge& Edge, int32& OutBoun
 	OutBoundaryIndex = INDEX_NONE;
 }
 
-void FTopologicalFace::EvaluateGrid(FGrid& Grid) const
-{
-	CarrierSurface->EvaluateGrid(Grid);
-}
-
 const void FTopologicalFace::Get2DLoopSampling(TArray<TArray<FPoint2D>>& LoopSamplings) const
 {
 	LoopSamplings.Empty(GetLoops().Num());
@@ -527,14 +522,13 @@ void FTopologicalFace::ApplyCriteria(const TArray<TSharedPtr<FCriterion>>& Crite
 	{
 		if (CoordNextMinusCoord < DOUBLE_SMALL_NUMBER)
 		{
-			return Length;
+			return;
 		}
 		const double ElemLength = Length * DeltaMax / CoordNextMinusCoord;
 		if (ElementLengthMin > ElemLength)
 		{
 			ElementLengthMin = ElemLength;
 		}
-		return ElemLength;
 	};
 
 	for (int32 IndexV = 0; IndexV < CrossingCoordinates[EIso::IsoV].Num() - 1; ++IndexV)
diff --git a/Engine/Source/Runtime/Datasmith/CADKernel/Private/CADKernel/Topo/TopologicalLoop.cpp b/Engine/Source/Runtime/Datasmith/CADKernel/Private/CADKernel/Topo/TopologicalLoop.cpp
--- a/Engine/Source/Runtime/Datasmith/CADKernel/Private/CADKernel/Topo/TopologicalLoop.cpp
+++ b/Engine/Source/Runtime/Datasmith/CADKernel/Private/CADKernel/Topo/TopologicalLoop.cpp
@@ -1173,5 +1173,14 @@ bool FTopologicalLoop::IsInside(const FTopologicalLoop& OtherLoop) const
 	return InsidePoint > OutsidePoint;
 }
 
+double FTopologicalLoop::Length() const
+{
+	double Length = 0;
+	for (const FOrientedEdge& Edge : GetEdges())
+	{
+		Length += Edge.Entity->Length();
+	}
+	return Length;
+}
 
 }
\ No newline at end of file
diff --git a/Engine/Source/Runtime/Datasmith/CADKernel/Private/CADKernel/Topo/Topomaker.cpp b/Engine/Source/Runtime/Datasmith/CADKernel/Private/CADKernel/Topo/Topomaker.cpp
--- a/Engine/Source/Runtime/Datasmith/CADKernel/Private/CADKernel/Topo/Topomaker.cpp
+++ b/Engine/Source/Runtime/Datasmith/CADKernel/Private/CADKernel/Topo/Topomaker.cpp
@@ -1136,6 +1136,7 @@ void FTopomaker::SplitIntoConnectedShells()
 		if (FaceSubset.MainShell != nullptr)
 		{
 			FShell* Shell = (FShell*)FaceSubset.MainShell;
+			ensure(Shell->GetFaces().Num() == 0);
 			Shell->Add(FaceSubset.Faces);
 		}
 		else
@@ -1176,6 +1177,16 @@ void FTopomaker::SplitIntoConnectedShells()
 
 void FTopomaker::RemoveEmptyShells()
 {
+#ifdef CADKERNEL_DEV
+	{
+		FModel& Model = Session.GetModel();
+		for (const TSharedPtr<FBody>& Body : Model.GetBodies())
+		{
+			ensureCADKernel(!Body->HasMarker1());
+		}
+	}
+#endif
+
 	TArray<FBody*> Bodies;
 	Bodies.Reserve(Shells.Num());
 
diff --git a/Engine/Source/Runtime/Datasmith/CADKernel/Private/CADKernel/UI/Display.cpp b/Engine/Source/Runtime/Datasmith/CADKernel/Private/CADKernel/UI/Display.cpp
--- a/Engine/Source/Runtime/Datasmith/CADKernel/Private/CADKernel/UI/Display.cpp
+++ b/Engine/Source/Runtime/Datasmith/CADKernel/Private/CADKernel/UI/Display.cpp
@@ -9,6 +9,7 @@
 #include "CADKernel/Geo/Curves/BezierCurve.h"
 #include "CADKernel/Geo/Curves/NURBSCurve.h"
 #include "CADKernel/Geo/Curves/RestrictionCurve.h"
+#include "CADKernel/Geo/Curves/SplineCurve.h"
 #include "CADKernel/Geo/Sampler/SamplerOnChord.h"
 #include "CADKernel/Geo/Sampler/SamplerOnParam.h"
 #include "CADKernel/Geo/Surfaces/BezierSurface.h"
@@ -31,7 +32,7 @@
 namespace UE::CADKernel
 {
 
-#ifdef CADKERNEL_DEV
+#ifdef CADKERNEL_DEBUG
 void Open3DDebugSession(FString name, const TArray<FIdent>& IdArray)
 {
 	FSystem::Get().GetVisu()->Open3DDebugSession(*name, IdArray);
@@ -57,42 +58,42 @@ void Wait(bool bMakeWait)
 
 void Open3DDebugSegment(FIdent Ident)
 {
-#ifdef CADKERNEL_DEV
+#ifdef CADKERNEL_DEBUG
 	FSystem::Get().GetVisu()->Open3DDebugSegment(Ident);
 #endif
 }
 
 void Close3DDebugSegment()
 {
-#ifdef CADKERNEL_DEV
+#ifdef CADKERNEL_DEBUG
 	FSystem::Get().GetVisu()->Close3DDebugSegment();
 #endif
 }
 
 void FlushVisu()
 {
-#ifdef CADKERNEL_DEV
+#ifdef CADKERNEL_DEBUG
 	FSystem::Get().GetVisu()->UpdateViewer();
 #endif
 }
 
 void DrawElement(int32 Dimension, TArray<FPoint>& Points, EVisuProperty Property)
 {
-#ifdef CADKERNEL_DEV
+#ifdef CADKERNEL_DEBUG
 	FSystem::Get().GetVisu()->DrawElement(Dimension, Points, Property);
 #endif
 }
 
 void DrawMesh(const FMesh& Mesh)
 {
-#ifdef CADKERNEL_DEV
+#ifdef CADKERNEL_DEBUG
 	FSystem::Get().GetVisu()->DrawMesh(Mesh.GetId());
 #endif
 }
 
 void Draw(const FLinearBoundary& Boundary, const FRestrictionCurve& Curve, EVisuProperty Property)
 {
-#ifdef CADKERNEL_DEV
+#ifdef CADKERNEL_DEBUG
 	if (Boundary.IsDegenerated())
 	{
 		return;
@@ -126,7 +127,7 @@ void Draw(const FLinearBoundary& Boundary, const FRestrictionCurve& Curve, EVisu
 
 void Draw2D(const FCurve& Curve, const FLinearBoundary& Boundary, EVisuProperty Property)
 {
-#ifdef CADKERNEL_DEV
+#ifdef CADKERNEL_DEBUG
 	if (Boundary.IsDegenerated())
 	{
 		return;
@@ -159,7 +160,7 @@ void Draw2D(const FCurve& Curve, const FLinearBoundary& Boundary, EVisuProperty
 
 void Draw3D(const FCurve& Curve, const FLinearBoundary& Boundary, EVisuProperty Property)
 {
-#ifdef CADKERNEL_DEV
+#ifdef CADKERNEL_DEBUG
 	if (Boundary.IsDegenerated())
 	{
 		return;
@@ -192,7 +193,7 @@ void Draw3D(const FCurve& Curve, const FLinearBoundary& Boundary, EVisuProperty
 
 void Draw(const FCurve& Curve, const FLinearBoundary& Boundary, EVisuProperty Property)
 {
-#ifdef CADKERNEL_DEV
+#ifdef CADKERNEL_DEBUG
 	if (Curve.GetDimension() == 3)
 	{
 		Draw3D(Curve, Boundary, Property);
@@ -207,7 +208,7 @@ void Draw(const FCurve& Curve, const FLinearBoundary& Boundary, EVisuProperty Pr
 
 void Draw(const FCurve& Curve, EVisuProperty Property)
 {
-#ifdef CADKERNEL_DEV
+#ifdef CADKERNEL_DEBUG
 	const FLinearBoundary& Bounds = Curve.GetBoundary();
 	Draw(Curve, Bounds, Property);
 #endif
@@ -215,7 +216,7 @@ void Draw(const FCurve& Curve, EVisuProperty Property)
 
 void DrawQuadripode(double Height, double Base, FPoint& Center, FPoint& InDirection)
 {
-#ifdef CADKERNEL_DEV
+#ifdef CADKERNEL_DEBUG
 	FPoint Direction = InDirection;
 	Direction.Normalize();
 
@@ -254,7 +255,7 @@ void DrawQuadripode(double Height, double Base, FPoint& Center, FPoint& InDirect
 
 void DisplayEntity(const FEntity& Entity)
 {
-#ifdef CADKERNEL_DEV
+#ifdef CADKERNEL_DEBUG
 	if (Entity.IsDeleted())
 	{
 		return;
@@ -312,7 +313,7 @@ void DisplayEntity(const FEntity& Entity)
 
 void DisplayEntity2D(const FEntity& Entity)
 {
-#ifdef CADKERNEL_DEV
+#ifdef CADKERNEL_DEBUG
 	if (Entity.IsDeleted())
 	{
 		return;
@@ -350,7 +351,7 @@ void DisplayEntity2D(const FEntity& Entity)
 
 void Display(const FPlane& Plane, FIdent Ident)
 {
-#ifdef CADKERNEL_DEV
+#ifdef CADKERNEL_DEBUG
 	FVector Normal = Plane.GetNormal();
 	if (Normal.Size() < DOUBLE_SMALL_NUMBER)
 	{
@@ -379,7 +380,7 @@ void Display(const FPlane& Plane, FIdent Ident)
 	TArray<FPoint> Points;
 
 	const FVector Point = Plane.GetOrigin();
-
+	
 	FVector VTemp;
 
 	VTemp = Point + UAxis + VAxis;
@@ -404,7 +405,7 @@ void Display(const FPlane& Plane, FIdent Ident)
 
 void DisplayEdgeCriteriaGrid(int32 EdgeId, const TArray<FPoint>& Points3D)
 {
-#ifdef CADKERNEL_DEV
+#ifdef CADKERNEL_DEBUG
 	FString Name = FString::Printf(TEXT("Edge Grid %d"), EdgeId);
 	F3DDebugSession G(Name);
 	{
@@ -426,7 +427,7 @@ void DisplayEdgeCriteriaGrid(int32 EdgeId, const TArray<FPoint>& Points3D)
 
 void DisplayAABB(const FAABB& Aabb, FIdent Ident)
 {
-#ifdef CADKERNEL_DEV
+#ifdef CADKERNEL_DEBUG
 	TFunction<void(int32, FPoint&, FPoint&)> GetWire = [&](int32 Segment, FPoint& A, FPoint& B)
 	{
 		int32 Corner1;
@@ -477,7 +478,7 @@ void DisplayAABB(const FAABB& Aabb, FIdent Ident)
 
 void DisplayAABB2D(const FAABB2D& aabb, FIdent Ident)
 {
-#ifdef CADKERNEL_DEV
+#ifdef CADKERNEL_DEBUG
 	FPoint A(aabb.GetMin().U, aabb.GetMin().V, 0.);
 	FPoint B(aabb.GetMin().U, aabb.GetMax().V, 0.);
 	FPoint C(aabb.GetMax().U, aabb.GetMax().V, 0.);
@@ -495,7 +496,7 @@ void DisplayAABB2D(const FAABB2D& aabb, FIdent Ident)
 
 void Display(const FTopologicalVertex& Vertex)
 {
-#ifdef CADKERNEL_DEV
+#ifdef CADKERNEL_DEBUG
 	F3DDebugSegment G(Vertex.GetId());
 	DrawPoint(Vertex.GetCoordinates());
 #endif
@@ -503,7 +504,7 @@ void Display(const FTopologicalVertex& Vertex)
 
 void Display(const FCurve& Curve)
 {
-#ifdef CADKERNEL_DEV
+#ifdef CADKERNEL_DEBUG
 	F3DDebugSegment G(Curve.GetId());
 	Draw(Curve, Curve.GetBoundary());
 #endif
@@ -511,7 +512,7 @@ void Display(const FCurve& Curve)
 
 void Display(const FShell& Shell)
 {
-#ifdef CADKERNEL_DEV
+#ifdef CADKERNEL_DEBUG
 	FProgress Progress(Shell.GetFaces().Num(), TEXT("Display Shell"));
 	for (const FOrientedFace& Face : Shell.GetFaces())
 	{
@@ -522,7 +523,7 @@ void Display(const FShell& Shell)
 
 void Draw(const FShell& Shell)
 {
-#ifdef CADKERNEL_DEV
+#ifdef CADKERNEL_DEBUG
 	FProgress Progress(Shell.GetFaces().Num(), TEXT("Display Shell"));
 	for (const FOrientedFace& Face : Shell.GetFaces())
 	{
@@ -534,7 +535,7 @@ void Draw(const FShell& Shell)
 
 void Display(const FBody& Body)
 {
-#ifdef CADKERNEL_DEV
+#ifdef CADKERNEL_DEBUG
 	FProgress Progress(Body.GetShells().Num(), TEXT("Display Body"));
 
 	F3DDebugSegment GraphicSegment(Body.GetId());
@@ -551,7 +552,7 @@ void Display(const FBody& Body)
 
 void Display(const FSurface& Surface)
 {
-#ifdef CADKERNEL_DEV
+#ifdef CADKERNEL_DEBUG
 	FProgress Progress(TEXT("Display Surface"));
 
 	F3DDebugSegment GraphicSegment(Surface.GetId());
@@ -591,7 +592,7 @@ void Display(const FSurface& Surface)
 
 void Display2D(const FSurface& Surface)
 {
-#ifdef CADKERNEL_DEV
+#ifdef CADKERNEL_DEBUG
 	double StepU, StepV;
 
 	F3DDebugSegment G(Surface.GetId());
@@ -621,7 +622,7 @@ void Display2D(const FSurface& Surface)
 
 void DisplayIsoCurve(const FSurface& Surface, double Coordinate, EIso IsoType)
 {
-#ifdef CADKERNEL_DEV
+#ifdef CADKERNEL_DEBUG
 	const double VisuSag = FSystem::Get().GetVisu()->GetParameters()->ChordError;
 
 	F3DDebugSegment G(Surface.GetId());
@@ -639,7 +640,7 @@ void DisplayIsoCurve(const FSurface& Surface, double Coordinate, EIso IsoType)
 
 void DisplayControlPolygon(const FCurve& Curve)
 {
-#ifdef CADKERNEL_DEV
+#ifdef CADKERNEL_DEBUG
 	TFunction<void(const TArray<FPoint>&)> DisplayHull = [](const TArray<FPoint>& Poles)
 	{
 		for (const FPoint& Pole : Poles)
@@ -670,12 +671,32 @@ void DisplayControlPolygon(const FCurve& Curve)
 		DisplayHull(Poles);
 		return;
 	}
+	case ECurve::Spline:
+	{
+		const FSplineCurve& Spline = (const FSplineCurve&)Curve;
+		const FInterpCurveFPoint& Poles = Spline.GetSplinePointsPosition();
+
+		for (const auto& Pole : Poles.Points)
+		{
+			DisplayPoint(Pole.OutVal, EVisuProperty::BluePoint);
+		}
+
+		for (const auto& Pole : Poles.Points)
+		{
+			DisplaySegment(Pole.OutVal - Pole.ArriveTangent / 2., Pole.OutVal, 0, EVisuProperty::GreenCurve);
+			DisplaySegment(Pole.OutVal, Pole.OutVal + Pole.LeaveTangent / 2., 0, EVisuProperty::GreenCurve);
+		}
+		return;
+	}
+
+	}
+
 #endif
 }
 
 void DisplayControlPolygon(const FSurface& Surface)
 {
-#ifdef CADKERNEL_DEV
+#ifdef CADKERNEL_DEBUG
 	TFunction<void(const TArray<FPoint>&, int32, int32)> DisplayHull = [](const TArray<FPoint>& Poles, int32 PoleUNum, int32 PoleVNum)
 	{
 		for (int32 Index = 0; Index < Poles.Num(); Index++)
@@ -721,7 +742,7 @@ void DisplayControlPolygon(const FSurface& Surface)
 
 void Display(const FTopologicalFace& Face)
 {
-#ifdef CADKERNEL_DEV
+#ifdef CADKERNEL_DEBUG
 	F3DDebugSegment GraphicSegment(Face.GetId());
 	Draw(Face);
 #endif
@@ -729,7 +750,7 @@ void Display(const FTopologicalFace& Face)
 
 void Display2D(const FTopologicalFace& Face)
 {
-#ifdef CADKERNEL_DEV
+#ifdef CADKERNEL_DEBUG
 	F3DDebugSegment GraphicSegment(Face.GetId());
 	Draw2D(Face);
 #endif
@@ -737,7 +758,7 @@ void Display2D(const FTopologicalFace& Face)
 
 void Display2DWithScale(const FTopologicalFace& Face)
 {
-#ifdef CADKERNEL_DEV
+#ifdef CADKERNEL_DEBUG
 	F3DDebugSegment GraphicSegment(Face.GetId());
 	Draw2D(Face);
 #endif
@@ -745,7 +766,7 @@ void Display2DWithScale(const FTopologicalFace& Face)
 
 void Draw(const FTopologicalFace& Face)
 {
-#ifdef CADKERNEL_DEV
+#ifdef CADKERNEL_DEBUG
 	{
 		for (const TSharedPtr<FTopologicalLoop>& Loop : Face.GetLoops())
 		{
@@ -755,7 +776,7 @@ void Draw(const FTopologicalFace& Face)
 				switch (Edge.Entity->GetTwinEntityCount())
 				{
 				case 1:
-					Property = Edge.Entity->IsDegenerated() ? EVisuProperty::OrangeCurve : EVisuProperty::BorderEdge;
+					Property =  Edge.Entity->IsDegenerated() ? EVisuProperty::OrangeCurve : EVisuProperty::BorderEdge;
 					break;
 				case 2:
 					Property = EVisuProperty::BlueCurve;
@@ -774,7 +795,7 @@ void Draw(const FTopologicalFace& Face)
 
 void Draw2D(const FTopologicalFace& Face)
 {
-#ifdef CADKERNEL_DEV
+#ifdef CADKERNEL_DEBUG
 	TArray<TArray<FPoint2D>> BoundaryApproximation;
 	Face.Get2DLoopSampling(BoundaryApproximation);
 
@@ -834,7 +855,7 @@ void Draw2D(const FTopologicalFace& Face)
 
 void DrawIsoCurves(const FTopologicalFace& Face)
 {
-#ifdef CADKERNEL_DEV
+#ifdef CADKERNEL_DEBUG
 	TArray<TArray<FPoint2D>> BoundaryApproximation;
 	Face.Get2DLoopSampling(BoundaryApproximation);
 
@@ -902,7 +923,7 @@ void DrawIsoCurves(const FTopologicalFace& Face)
 
 void Display(const FTopologicalEdge& Edge, EVisuProperty Property)
 {
-#ifdef CADKERNEL_DEV
+#ifdef CADKERNEL_DEBUG
 	F3DDebugSegment GraphicSegment(Edge.GetId());
 	Draw(Edge, Property);
 #endif
@@ -910,7 +931,7 @@ void Display(const FTopologicalEdge& Edge, EVisuProperty Property)
 
 void Display2D(const FTopologicalEdge& Edge, EVisuProperty Property)
 {
-#ifdef CADKERNEL_DEV
+#ifdef CADKERNEL_DEBUG
 	F3DDebugSegment GraphicSegment(Edge.GetId());
 	TArray<FPoint2D> Polyline;
 	Edge.GetCurve()->GetDiscretizationPoints(Edge.GetBoundary(), EOrientation::Front, Polyline);
@@ -920,7 +941,7 @@ void Display2D(const FTopologicalEdge& Edge, EVisuProperty Property)
 
 void Display2DWithScale(const FTopologicalEdge& Edge, EVisuProperty Property)
 {
-#ifdef CADKERNEL_DEV
+#ifdef CADKERNEL_DEBUG
 	F3DDebugSegment GraphicSegment(Edge.GetId());
 	TArray<FPoint2D> Polyline;
 	Edge.GetCurve()->GetDiscretizationPoints(Edge.GetBoundary(), EOrientation::Front, Polyline);
@@ -928,14 +949,14 @@ void Display2DWithScale(const FTopologicalEdge& Edge, EVisuProperty Property)
 
 	for (const FPoint2D& Point : Polyline)
 	{
-		DisplayPoint2DWithScale(Point, EVisuProperty::RedPoint);
+		DisplayPoint2DWithScale(Point, EVisuProperty ::RedPoint);
 	}
 #endif
 }
 
 void Draw(const FTopologicalEdge& Edge, EVisuProperty Property)
 {
-#ifdef CADKERNEL_DEV
+#ifdef CADKERNEL_DEBUG
 	const FLinearBoundary& Boundary = Edge.GetBoundary();
 	Draw(Boundary, *Edge.GetCurve(), Property);
 #endif
@@ -943,7 +964,7 @@ void Draw(const FTopologicalEdge& Edge, EVisuProperty Property)
 
 void Display(const FModel& Model)
 {
-#ifdef CADKERNEL_DEV
+#ifdef CADKERNEL_DEBUG
 	FProgress MainProgress(2, TEXT("Display model"));
 	{
 		const TArray<TSharedPtr<FBody>>& Bodies = Model.GetBodies();
@@ -959,7 +980,7 @@ void Display(const FModel& Model)
 
 void DisplayProductTree(const FModel& Model)
 {
-#ifdef CADKERNEL_DEV
+#ifdef CADKERNEL_DEBUG
 	F3DDebugSession GraphicSession(FString::Printf(TEXT("%s %d"), Model.GetTypeName(), Model.GetId()), { Model.GetId() });
 
 	FProgress MainProgress(2, TEXT("Display Model"));
@@ -977,7 +998,7 @@ void DisplayProductTree(const FModel& Model)
 
 void DisplayProductTree(const FBody& Body)
 {
-#ifdef CADKERNEL_DEV
+#ifdef CADKERNEL_DEBUG
 #ifdef CORETECHBRIDGE_DEBUG
 	F3DDebugSession GraphicSession(FString::Printf(TEXT("%s %s KioId: %d Id: %d"), Body.GetTypeName(), Body.GetName(), Body.GetKioId(), Body.GetId()), { Body.GetId() });
 #else
@@ -998,7 +1019,7 @@ void DisplayProductTree(const FBody& Body)
 
 void DisplayProductTree(const FShell& Shell)
 {
-#ifdef CADKERNEL_DEV
+#ifdef CADKERNEL_DEBUG
 #ifdef CORETECHBRIDGE_DEBUG
 	F3DDebugSession GraphicSession(FString::Printf(TEXT("%s %s KioId: %d Id: %d"), Shell.GetTypeName(), Shell.GetName(), Shell.GetKioId(), Shell.GetId()), { Shell.GetId() });
 #else
@@ -1010,7 +1031,7 @@ void DisplayProductTree(const FShell& Shell)
 
 void DisplayProductTree(const FEntity& Entity)
 {
-#ifdef CADKERNEL_DEV
+#ifdef CADKERNEL_DEBUG
 	FTimePoint StartTime = FChrono::Now();
 
 	switch (Entity.GetEntityType())
@@ -1038,7 +1059,7 @@ void DisplayProductTree(const FEntity& Entity)
 
 void Display(const FGroup& Group)
 {
-#ifdef CADKERNEL_DEV
+#ifdef CADKERNEL_DEBUG
 	TArray<TSharedPtr<FEntity>> Entities;
 
 	Group.GetValidEntities(Entities);
@@ -1053,7 +1074,7 @@ void Display(const FGroup& Group)
 
 void DisplayMesh(const FFaceMesh& Mesh)
 {
-#ifdef CADKERNEL_DEV
+#ifdef CADKERNEL_DEBUG
 	TMap<int32, const FPoint*> NodeIdToCoordinates;
 	Mesh.GetNodeIdToCoordinates(NodeIdToCoordinates);
 
@@ -1122,7 +1143,7 @@ void DisplayMesh(const FFaceMesh& Mesh)
 
 void DisplayMesh(const FEdgeMesh& Mesh)
 {
-#ifdef CADKERNEL_DEV
+#ifdef CADKERNEL_DEBUG
 	const FModelMesh& MeshModel = Mesh.GetMeshModel();
 
 	const TArray<int32>& NodeIds = Mesh.EdgeVerticesIndex;
@@ -1174,14 +1195,14 @@ void DisplayMesh(const FEdgeMesh& Mesh)
 
 void DisplayMesh(const FVertexMesh& Mesh)
 {
-#ifdef CADKERNEL_DEV
+#ifdef CADKERNEL_DEBUG
 	DisplayPoint(Mesh.GetNodeCoordinates()[0], EVisuProperty::NodeMesh, Mesh.GetId());
 #endif
 }
 
 void Display(const FModelMesh& MeshModel)
 {
-#ifdef CADKERNEL_DEV
+#ifdef CADKERNEL_DEBUG
 	for (const TSharedPtr<FFaceMesh>& Mesh : MeshModel.GetFaceMeshes())
 	{
 		if (Mesh.IsValid())
@@ -1194,7 +1215,7 @@ void Display(const FModelMesh& MeshModel)
 
 void DisplaySegment(const FPoint& Point1, const FPoint& Point2, FIdent Ident, EVisuProperty Property)
 {
-#ifdef CADKERNEL_DEV
+#ifdef CADKERNEL_DEBUG
 	F3DDebugSegment G(Ident);
 	DrawSegment(Point1, Point2, Property);
 #endif
@@ -1202,7 +1223,7 @@ void DisplaySegment(const FPoint& Point1, const FPoint& Point2, FIdent Ident, EV
 
 void DisplaySegment(const FPoint2D& Point1, const FPoint2D& Point2, FIdent Ident, EVisuProperty Property, bool bWithOrientation)
 {
-#ifdef CADKERNEL_DEV
+#ifdef CADKERNEL_DEBUG
 	F3DDebugSegment G(Ident);
 	if (bWithOrientation)
 	{
@@ -1214,7 +1235,7 @@ void DisplaySegment(const FPoint2D& Point1, const FPoint2D& Point2, FIdent Ident
 
 void DisplaySegment(const FPointH& Point1, const FPointH& Point2, FIdent Ident, EVisuProperty Property)
 {
-#ifdef CADKERNEL_DEV
+#ifdef CADKERNEL_DEBUG
 	F3DDebugSegment G(Ident);
 	DrawSegment(Point1, Point2, Property);
 #endif
@@ -1222,7 +1243,7 @@ void DisplaySegment(const FPointH& Point1, const FPointH& Point2, FIdent Ident,
 
 void DisplayLoop(const FTopologicalFace& Surface)
 {
-#ifdef CADKERNEL_DEV
+#ifdef CADKERNEL_DEBUG
 	for (const TSharedPtr<FTopologicalLoop>& Loop : Surface.GetLoops())
 	{
 		Display(*Loop);
@@ -1233,7 +1254,7 @@ void DisplayLoop(const FTopologicalFace& Surface)
 
 void Display(const FTopologicalLoop& Loop)
 {
-#ifdef CADKERNEL_DEV
+#ifdef CADKERNEL_DEBUG
 	for (const FOrientedEdge& Edge : Loop.GetEdges())
 	{
 		Display(*Edge.Entity);
@@ -1243,7 +1264,7 @@ void Display(const FTopologicalLoop& Loop)
 
 void Display2D(const FTopologicalLoop& Loop)
 {
-#ifdef CADKERNEL_DEV
+#ifdef CADKERNEL_DEBUG
 	for (const FOrientedEdge& Edge : Loop.GetEdges())
 	{
 		Display2D(*Edge.Entity);
@@ -1253,7 +1274,7 @@ void Display2D(const FTopologicalLoop& Loop)
 
 void Display2DWithScale(const FTopologicalLoop& Loop)
 {
-#ifdef CADKERNEL_DEV
+#ifdef CADKERNEL_DEBUG
 	for (const FOrientedEdge& Edge : Loop.GetEdges())
 	{
 		Display2DWithScale(*Edge.Entity);
diff --git a/Engine/Source/Runtime/Datasmith/CADKernel/Public/CADKernel/Core/HaveStates.h b/Engine/Source/Runtime/Datasmith/CADKernel/Public/CADKernel/Core/HaveStates.h
--- a/Engine/Source/Runtime/Datasmith/CADKernel/Public/CADKernel/Core/HaveStates.h
+++ b/Engine/Source/Runtime/Datasmith/CADKernel/Public/CADKernel/Core/HaveStates.h
@@ -126,6 +126,12 @@ public:
 		States &= ~EHaveStates::IsDeleted;
 	}
 
+	virtual bool IsDeletedOrDegenerated() const
+	{
+		constexpr EHaveStates DeletedOrDegenerated = EHaveStates::Degenerated | EHaveStates::IsDeleted;
+		return ((States & DeletedOrDegenerated) != EHaveStates::None);
+	}
+
 	virtual bool IsDegenerated() const
 	{
 		return ((States & EHaveStates::Degenerated) == EHaveStates::Degenerated);
diff --git a/Engine/Source/Runtime/Datasmith/CADKernel/Public/CADKernel/Geo/Curves/Curve.h b/Engine/Source/Runtime/Datasmith/CADKernel/Public/CADKernel/Geo/Curves/Curve.h
--- a/Engine/Source/Runtime/Datasmith/CADKernel/Public/CADKernel/Geo/Curves/Curve.h
+++ b/Engine/Source/Runtime/Datasmith/CADKernel/Public/CADKernel/Geo/Curves/Curve.h
@@ -46,8 +46,12 @@ protected:
 public:
 
 	static TSharedPtr<FCurve> MakeNurbsCurve(FNurbsCurveData& InNurbsData);
+	static TSharedPtr<FCurve> MakeBezierCurve(const TArray<FPoint>& InPoles);
+	static TSharedPtr<FCurve> MakeSplineCurve(const TArray<FPoint>& InPoles);
+	static TSharedPtr<FCurve> MakeSplineCurve(const TArray<FPoint>& InPoles, const TArray<FPoint>& Tangents);
+	static TSharedPtr<FCurve> MakeSplineCurve(const TArray<FPoint>& InPoles, const TArray<FPoint>& ArriveTangents, const TArray<FPoint>& LeaveTangents);
+
 	// TODO
-	//static TSharedPtr<FCurve> MakeBezierCurve(const double InToleranceGeometric, int32 InUDegre, int32 InVDegre, const TArray<FPoint>& InPoles);
 	//static TSharedPtr<FCurve> MakeConeCurve(const double InToleranceGeometric, const FMatrixH& InMatrix, double InStartRadius, double InConeAngle, const FSurfacicBoundary& InBoundary);
 	//static TSharedPtr<FCurve> MakeCylinderCurve(const double InToleranceGeometric, const FMatrixH& InMatrix, const double InRadius, const FSurfacicBoundary& InBoundary);
 	//static TSharedPtr<FCurve> MakePlaneCurve(const double InToleranceGeometric, const FMatrixH& InMatrix, const FSurfacicBoundary& InBoundary);
@@ -106,7 +110,7 @@ public:
 	virtual TSharedPtr<FEntityGeom> ApplyMatrix(const FMatrixH& InMatrix) const override = 0;
 	virtual void Offset(const FPoint& OffsetDirection) = 0;
 
-	double GetLength(double Tolerance) const;
+	virtual double GetLength(double Tolerance) const;
 
 	/**
 	 * Evaluate exact 3D point of the curve at the input Coordinate
diff --git a/Engine/Source/Runtime/Datasmith/CADKernel/Public/CADKernel/Geo/GeoPoint.h b/Engine/Source/Runtime/Datasmith/CADKernel/Public/CADKernel/Geo/GeoPoint.h
--- a/Engine/Source/Runtime/Datasmith/CADKernel/Public/CADKernel/Geo/GeoPoint.h
+++ b/Engine/Source/Runtime/Datasmith/CADKernel/Public/CADKernel/Geo/GeoPoint.h
@@ -121,7 +121,7 @@ struct FCoordinateGrid
 
 	int32 Count() const
 	{
-		return Coordinates[0].Num() * Coordinates[1].Num();
+		return Coordinates[EIso::IsoU].Num() * Coordinates[EIso::IsoV].Num();
 	}
 
 	int32 IsoCount(EIso Iso) const
diff --git a/Engine/Source/Runtime/Datasmith/CADKernel/Public/CADKernel/Geo/Sampling/PolylineTools.h b/Engine/Source/Runtime/Datasmith/CADKernel/Public/CADKernel/Geo/Sampling/PolylineTools.h
--- a/Engine/Source/Runtime/Datasmith/CADKernel/Public/CADKernel/Geo/Sampling/PolylineTools.h
+++ b/Engine/Source/Runtime/Datasmith/CADKernel/Public/CADKernel/Geo/Sampling/PolylineTools.h
@@ -407,6 +407,15 @@ public:
 		TArray<double> CurvilinearCoordinates;
 		double CurveLength = ComputeCurvilinearCoordinatesOfPolyline(InBoundary, CurvilinearCoordinates, BoundaryIndices);
 
+		if (CurveLength < 2. * DesiredSegmentLength)
+		{
+			OutCoordinates.SetNum(3);
+			OutCoordinates[0] = InBoundary.GetMin();
+			OutCoordinates[1] = InBoundary.GetMiddle();
+			OutCoordinates[2] = InBoundary.GetMax();
+			return;
+		}
+
 		int32 SegmentNum = (int32)FMath::Max(CurveLength / DesiredSegmentLength + 0.5, 1.0);
 
 		double SectionLength = CurveLength / (double)(SegmentNum);
diff --git a/Engine/Source/Runtime/Datasmith/CADKernel/Public/CADKernel/Math/Geometry.h b/Engine/Source/Runtime/Datasmith/CADKernel/Public/CADKernel/Math/Geometry.h
--- a/Engine/Source/Runtime/Datasmith/CADKernel/Public/CADKernel/Math/Geometry.h
+++ b/Engine/Source/Runtime/Datasmith/CADKernel/Public/CADKernel/Math/Geometry.h
@@ -29,7 +29,7 @@ CADKERNEL_API inline FPoint2D ComputeCircumCircleCenter(const FPoint2D& InPoint0
 	FPoint2D Segment_P0_P2 = InPoint2 - InPoint0;
 
 	// D = 2(BuCv - BvCu)
-	double D = 2 * Segment_P0_P1 ^ Segment_P0_P2;
+	double D = 2. * Segment_P0_P1 ^ Segment_P0_P2;
 	if (FMath::IsNearlyZero(D, SMALL_NUMBER_SQUARE))
 	{
 		return FPoint2D::ZeroPoint;
@@ -61,7 +61,7 @@ CADKERNEL_API inline FPoint ComputeCircumCircleCenter(const FPoint& Point0, cons
 	FPoint2D D2Point1 = MatrixInverse * Point1;
 	FPoint2D D2Point2 = MatrixInverse * Point2;
 
-	double D = 2 * D2Point1 ^ D2Point2;
+	double D = 2. * D2Point1 ^ D2Point2;
 	if (FMath::IsNearlyZero(D, SMALL_NUMBER_SQUARE))
 	{
 		return FPoint::ZeroPoint;
@@ -83,7 +83,7 @@ CADKERNEL_API inline FPoint2D ComputeCircumCircleCenterAndSquareRadius(const FPo
 	FPoint2D Segment_P0_P1 = InPoint1 - InPoint0;
 	FPoint2D Segment_P0_P2 = InPoint2 - InPoint0;
 
-	double D = 2 * Segment_P0_P1 ^ Segment_P0_P2;
+	double D = 2. * Segment_P0_P1 ^ Segment_P0_P2;
 	if (FMath::IsNearlyZero(D, SMALL_NUMBER_SQUARE))
 	{
 		OutSquareRadius = 0;
@@ -132,8 +132,15 @@ struct CADKERNEL_API TSegment
 		return Point0.SquareDistance(Point1);
 	}
 
+	PointType GetVector() const
+	{
+		return Point1 - Point0;
+	}
 };
 
+using FSegment2D = TSegment<FPoint2D>;
+using FSegment3D = TSegment<FPoint>;
+
 template<class PointType>
 struct CADKERNEL_API TTriangle
 {
diff --git a/Engine/Source/Runtime/Datasmith/CADKernel/Public/CADKernel/Math/Point.h b/Engine/Source/Runtime/Datasmith/CADKernel/Public/CADKernel/Math/Point.h
--- a/Engine/Source/Runtime/Datasmith/CADKernel/Public/CADKernel/Math/Point.h
+++ b/Engine/Source/Runtime/Datasmith/CADKernel/Public/CADKernel/Math/Point.h
@@ -787,9 +787,21 @@ inline FPoint operator*(double Scale, const FPoint& Point)
 	return FPoint(Point.X * Scale, Point.Y * Scale, Point.Z * Scale);
 }
 
-inline FPoint2D operator*(double s, const FPoint2D& Point)
+inline FPoint operator*(float Scale, const FPoint& Point)
 {
-	return FPoint2D(Point.U * s, Point.V * s);
+	double ScaleD = static_cast<double>(Scale);
+	return FPoint(Point.X * ScaleD, Point.Y * ScaleD, Point.Z * ScaleD);
+}
+
+inline FPoint2D operator*(double Scale, const FPoint2D& Point)
+{
+	return FPoint2D(Point.U * Scale, Point.V * Scale);
+}
+
+inline FPoint2D operator*(float Scale, const FPoint2D& Point)
+{
+	double ScaleD = static_cast<double>(Scale);
+	return FPoint2D(Point.U * ScaleD, Point.V * ScaleD);
 }
 
 }
diff --git a/Engine/Source/Runtime/Datasmith/CADKernel/Public/CADKernel/Mesh/Meshers/ParametricMesher.h b/Engine/Source/Runtime/Datasmith/CADKernel/Public/CADKernel/Mesh/Meshers/ParametricMesher.h
--- a/Engine/Source/Runtime/Datasmith/CADKernel/Public/CADKernel/Mesh/Meshers/ParametricMesher.h
+++ b/Engine/Source/Runtime/Datasmith/CADKernel/Public/CADKernel/Mesh/Meshers/ParametricMesher.h
@@ -44,6 +44,8 @@ protected:
 
 	TArray<FTopologicalFace*> Faces;
 
+	bool bThinZoneMeshing = false;
+
 #ifdef CADKERNEL_DEV
 	FMesherReport MesherReport;
 	bool bDisplay = false;
@@ -78,12 +80,14 @@ public:
 
 	void MeshFaceLoops(FGrid& Grid);
 
-	void MeshThinZoneEdges(FGrid&);
-	void MeshThinZoneSide(const FThinZoneSide& Side);
-	void GetThinZoneBoundary(const FThinZoneSide& Side);
-
-	void GenerateCloud(FGrid& Grid);
+	void MeshThinZoneEdges(FTopologicalFace& Face);
+	void MeshThinZoneSide(FThinZoneSide& Side);
+	void FindThinZoneBoundary(FThinZoneSide& Side);
 
+	/**
+	 * @return false if the process fails i.e. the grid is degenerated or else. 
+	 */
+	bool GenerateCloud(FGrid& Grid);
 
 #ifdef CADKERNEL_DEV
 	void SetMeshReport(const FMesherReport& InMesherReport)
@@ -112,7 +116,7 @@ protected:
 	void MeshSurfaceByFront(TArray<FCostToFace>& QuadTrimmedSurfaceSet);
 
 	void ApplyEdgeCriteria(FTopologicalEdge& Edge);
-	static void ApplyFaceCriteria(FTopologicalFace& Face, const TArray<TSharedPtr<FCriterion>>& Criteria);
+	static void ApplyFaceCriteria(FTopologicalFace& Face, const TArray<TSharedPtr<FCriterion>>& Criteria, bool);
 
 	/**
 	 * Generate Edge Elements on active edge from Edge cutting points
diff --git a/Engine/Source/Runtime/Datasmith/CADKernel/Public/CADKernel/Mesh/Structure/ThinZone2D.h b/Engine/Source/Runtime/Datasmith/CADKernel/Public/CADKernel/Mesh/Structure/ThinZone2D.h
new file mode 100644
--- /dev/null
+++ b/Engine/Source/Runtime/Datasmith/CADKernel/Public/CADKernel/Mesh/Structure/ThinZone2D.h
@@ -0,0 +1,215 @@
+// Copyright Epic Games, Inc. All Rights Reserved.
+#pragma once
+
+#include "CADKernel/Core/Chrono.h"
+#include "CADKernel/Core/HaveStates.h"
+//#include "CADKernel/Core/Factory.h"
+#include "CADKernel/Core/Types.h"
+#include "CADKernel/Math/Point.h"
+
+#include "CADKernel/Mesh/MeshEnum.h"
+#include "CADKernel/Mesh/Structure/EdgeSegment.h"
+
+namespace UE::CADKernel
+{
+	enum class EThinZone2DType : uint8
+	{
+		Undefined = 0,
+		Global,			// a Surface globally thin
+		PeakStart,		// an extremity of a Surface that is fine
+		PeakEnd,		// an extremity of a Surface that is fine
+		Butterfly,		// the outer loop that is like a bow tie (Butterfly)
+		BetweenLoops,   // a bow tie between two different loops
+		TooSmall        // -> to delete
+	};
+
+	class FThinZone2D;
+
+	class FThinZoneSide : public FHaveStates
+	{
+		friend FThinZone2D;
+
+	private:
+		TArray<FEdgeSegment> Segments;
+		FThinZoneSide& FrontSide;
+
+		double SideLength;
+		double MediumThickness;
+		double MaxThickness;
+
+	public:
+
+		FThinZoneSide(FThinZoneSide* InFrontSide, const TArray<FEdgeSegment*>& InSegments);
+		virtual ~FThinZoneSide() = default;
+
+		void Empty()
+		{
+			Segments.Empty();
+		}
+
+		const FEdgeSegment& GetFirst() const
+		{
+			return Segments[0];
+		}
+
+		const FEdgeSegment& GetLast() const
+		{
+			return Segments.Last();
+		}
+
+		void SetEdgesAsThinZone();
+
+		const TArray<FEdgeSegment>& GetSegments() const
+		{
+			return Segments;
+		}
+
+		TArray<FEdgeSegment>& GetSegments()
+		{
+			return Segments;
+		}
+
+		EMeshingState GetMeshingState() const;
+
+		double Length() const
+		{
+			return SideLength;
+		}
+
+		double GetThickness() const
+		{
+			return MediumThickness;
+		}
+
+		double GetMaxThickness() const
+		{
+			return MaxThickness;
+		}
+
+		bool IsInner() const
+		{
+			if (Segments.Num() == 0)
+			{
+				return true;
+			}
+			return Segments[0].IsInner();
+		}
+
+	private:
+		void ComputeThicknessAndLength();
+
+	};
+
+	class FThinZone2D : public FHaveStates
+	{
+	private:
+
+		FThinZoneSide FirstSide;
+		FThinZoneSide SecondSide;
+
+		EThinZone2DType Category;
+
+		double Thickness;
+		double MaxThickness;
+
+	public:
+
+		/**
+		 * FThinZone2D::FThinZoneSides are maded with a copy of TArray<FEdgeSegment*> into TArray<FEdgeSegment> to break the link with TFactory<FEdgeSegment> of FThinZone2DFinder
+		 * So FThinZone2D can be transfered into the FTopologicalFace 
+		 */
+		FThinZone2D(const TArray<FEdgeSegment*>& InFirstSideSegments, const TArray<FEdgeSegment*>& InSecondSideSegments)
+			: FirstSide(&SecondSide, InFirstSideSegments)
+			, SecondSide(&FirstSide, InSecondSideSegments)
+			, Category(EThinZone2DType::Undefined)
+		{
+			Finalize();
+ 		}
+
+		virtual ~FThinZone2D() = default;
+
+		void Empty()
+		{
+			FirstSide.Empty();
+			SecondSide.Empty();
+			Thickness = -1;
+			SetRemoved();
+		}
+
+		double GetThickness() const
+		{
+			return Thickness;
+		};
+
+		double GetMaxThickness() const
+		{
+			return MaxThickness;
+		};
+
+		FThinZoneSide& GetFirstSide()
+		{
+			return FirstSide;
+		}
+
+		FThinZoneSide& GetSecondSide()
+		{
+			return SecondSide;
+		}
+
+		const FThinZoneSide& GetFirstSide() const
+		{
+			return FirstSide;
+		}
+
+		const FThinZoneSide& GetSecondSide() const 
+		{
+			return SecondSide;
+		}
+
+		EThinZone2DType GetCategory() const
+		{
+			return Category;
+		}
+
+		void SetEdgesAsThinZone();
+		static void SetPeakEdgesMarker(const TArray<const FTopologicalEdge*>&);
+
+		double Length() const
+		{
+			return FirstSide.Length() + SecondSide.Length();
+		}
+
+		double GetMaxSideLength() const
+		{
+			return FMath::Max(FirstSide.Length(), SecondSide.Length());
+		}
+
+		bool IsRemoved() const
+		{
+			return ((States & EHaveStates::IsRemoved) == EHaveStates::IsRemoved);
+		}
+
+		void SetRemoved() const
+		{
+			States |= EHaveStates::IsRemoved;
+		}
+
+		void ResetRemoved() const
+		{
+			States &= ~EHaveStates::IsRemoved;
+		}
+
+		void SetCategory(EThinZone2DType InType)
+		{
+			Category = InType;
+		}
+
+#ifdef CADKERNEL_DEV
+		void Display(const FString& Title, EVisuProperty VisuProperty) const;
+#endif
+
+	private:
+		void Finalize();
+	};
+
+}
diff --git a/Engine/Source/Runtime/Datasmith/CADKernel/Public/CADKernel/Topo/Linkable.h b/Engine/Source/Runtime/Datasmith/CADKernel/Public/CADKernel/Topo/Linkable.h
--- a/Engine/Source/Runtime/Datasmith/CADKernel/Public/CADKernel/Topo/Linkable.h
+++ b/Engine/Source/Runtime/Datasmith/CADKernel/Public/CADKernel/Topo/Linkable.h
@@ -150,7 +150,7 @@ public:
 		return ((States & EHaveStates::ThinZone) == EHaveStates::ThinZone);
 	}
 
-	virtual void SetThinZone()
+	virtual void SetThinZoneMarker()
 	{
 		States |= EHaveStates::ThinZone;
 	}
diff --git a/Engine/Source/Runtime/Datasmith/CADKernel/Public/CADKernel/Topo/Shell.h b/Engine/Source/Runtime/Datasmith/CADKernel/Public/CADKernel/Topo/Shell.h
--- a/Engine/Source/Runtime/Datasmith/CADKernel/Public/CADKernel/Topo/Shell.h
+++ b/Engine/Source/Runtime/Datasmith/CADKernel/Public/CADKernel/Topo/Shell.h
@@ -67,7 +67,7 @@ private:
 
 public:
 
-	virtual ~FShell()
+	virtual ~FShell() override
 	{
 		FShell::Empty();
 	}
diff --git a/Engine/Source/Runtime/Datasmith/CADKernel/Public/CADKernel/Topo/TopologicalEdge.h b/Engine/Source/Runtime/Datasmith/CADKernel/Public/CADKernel/Topo/TopologicalEdge.h
--- a/Engine/Source/Runtime/Datasmith/CADKernel/Public/CADKernel/Topo/TopologicalEdge.h
+++ b/Engine/Source/Runtime/Datasmith/CADKernel/Public/CADKernel/Topo/TopologicalEdge.h
@@ -27,12 +27,33 @@ enum ECoordinateType : uint8
 	OtherCoordinate
 };
 
+/**
+ * Cutting point used for thin zone purpose
+ */
+struct CADKERNEL_API FImposedCuttingPoint
+{
+	/**
+	 * coordinate of the edge's mesh nodes
+	 */
+	double Coordinate = 0;
+	int32 OppositNodeIndex = -1;
+
+	FImposedCuttingPoint()
+	{
+	}
+
+	FImposedCuttingPoint(double InCoordinate, int32 NodeIndex1)
+		: Coordinate(InCoordinate)
+		, OppositNodeIndex(NodeIndex1)
+	{
+	};
+};
+
 template<typename FCuttingPointType>
 void GetCuttingPointCoordinates(const TArray<FCuttingPointType>& CuttingPoints, TArray<double>& CuttingPointCoordinates);
 
 struct FEdge2DProperties;
 struct FCuttingPoint;
-struct FImposedCuttingPoint;
 
 class FModelMesh;
 class FEdgeMesh;
@@ -486,6 +507,26 @@ public:
 		return CrossingPointDeltaUMaxs;
 	}
 
+	double GetDeltaUFor(double Coordinate, int32& Index) const 
+	{
+		for (; Index < CrossingPointUs.Num() - 1; ++Index)
+		{
+			if (Coordinate < CrossingPointUs[Index + 1])
+			{
+				if (Coordinate > CrossingPointUs[Index] - DOUBLE_SMALL_NUMBER)
+				{
+					break;
+				}
+				Index = 0;
+				if(Coordinate < CrossingPointUs[1])
+				{
+					break;
+				}
+			}
+		}
+		return CrossingPointDeltaUMaxs[Index];
+	}
+
 	TArray<FCuttingPoint>& GetCuttingPoints()
 	{
 		return CuttingPointUs;
@@ -497,8 +538,12 @@ public:
 	}
 
 	// For thin zone purpose
+	void SortImposedCuttingPoints()
+	{
+		Algo::Sort(ImposedCuttingPointUs, [](const FImposedCuttingPoint& C1, const FImposedCuttingPoint& C2) { return C1.Coordinate < C2.Coordinate; });
+	}
 
-	const TArray<FImposedCuttingPoint>& GetImposedCuttingPoints()
+	const TArray<FImposedCuttingPoint>& GetImposedCuttingPoints() const
 	{
 		return ImposedCuttingPointUs;
 	}
@@ -645,11 +690,11 @@ public:
 		Curve->GetDiscretizationPoints(Boundary, Orientation, OutPoints);
 	}
 
-	double TransformLocalCoordinateToActiveEdgeCoordinate(const double LocalCoordinate);
-	double TransformActiveEdgeCoordinateToLocalCoordinate(const double ActiveEdgeCoordinate);
+	double TransformLocalCoordinateToActiveEdgeCoordinate(const double LocalCoordinate) const;
+	double TransformActiveEdgeCoordinateToLocalCoordinate(const double ActiveEdgeCoordinate) const;
 
-	void TransformActiveEdgeCoordinatesToLocalCoordinates(const TArray<double>& InActiveEdgeCoordinate, TArray<double>& OutLocalCoordinate);
-	void TransformLocalCoordinatesToActiveEdgeCoordinates(const TArray<double>& InLocalCoordinate, TArray<double>& OutActiveEdgeCoordinate);
+	void TransformActiveEdgeCoordinatesToLocalCoordinates(const TArray<double>& InActiveEdgeCoordinate, TArray<double>& OutLocalCoordinate) const;
+	void TransformLocalCoordinatesToActiveEdgeCoordinates(const TArray<double>& InLocalCoordinate, TArray<double>& OutActiveEdgeCoordinate) const;
 
 	/**
 	 * Compute the edge 2D properties i.e. the mean and standard deviation of the slop of the edge in the parametric space of the carrier surface
@@ -696,17 +741,30 @@ public:
 		return FHaveStates::IsDegenerated();
 	}
 
+	/**
+	 * An edge is a thin peak means that this edge is a small edge at the extremity of a peak thin zone
+	 * So this edge must not be meshed (except at its extremities)
+	 *
+	 *                         ThinSide 0
+	 *           #-------------------------------------#
+	 *          /
+	 *         /  <- Thin peak edge
+	 *        /
+	 *       #-----------------------------------------#
+	 *                         ThinSide 1
+	 *
+	 */
 	bool IsThinPeak() const
 	{
 		return ((States & EHaveStates::ThinPeak) == EHaveStates::ThinPeak);
 	}
 
-	virtual void SetThinPeak() const
+	virtual void SetThinPeakMarker() const
 	{
 		States |= EHaveStates::ThinPeak;
 	}
 
-	virtual void ResetThinPeak() const
+	virtual void ResetThinPeakMarker() const
 	{
 		States &= ~EHaveStates::ThinPeak;
 	}
@@ -716,12 +774,12 @@ public:
 		return ((States & EHaveStates::IsVirtuallyMeshed) == EHaveStates::IsVirtuallyMeshed);
 	}
 
-	virtual void SetAsVirtuallyMeshed() const
+	virtual void SetVirtuallyMeshedMarker() const
 	{
 		States |= EHaveStates::IsVirtuallyMeshed;
 	}
 
-	virtual void ResetVirtuallyMeshed() const
+	virtual void ResetVirtuallyMeshedMarker() const
 	{
 		States &= ~EHaveStates::IsVirtuallyMeshed;
 	}
@@ -827,28 +885,6 @@ struct CADKERNEL_API FEdge2DProperties
 
 };
 
-/**
- * Cutting point used for thin zone purpose
- */
-struct CADKERNEL_API FImposedCuttingPoint
-{
-	/**
-	 * coordinate of the edge's mesh nodes
-	 */
-	double Coordinate = 0;
-	int32 OppositNodeIndex = -1;
-
-	FImposedCuttingPoint()
-	{
-	}
-
-	FImposedCuttingPoint(double InCoordinate, int32 NodeIndex1)
-		: Coordinate(InCoordinate)
-		, OppositNodeIndex(NodeIndex1)
-	{
-	};
-};
-
 /**
  * Cutting point used for meshing purpose
  */
diff --git a/Engine/Source/Runtime/Datasmith/CADKernel/Public/CADKernel/Topo/TopologicalEntity.h b/Engine/Source/Runtime/Datasmith/CADKernel/Public/CADKernel/Topo/TopologicalEntity.h
--- a/Engine/Source/Runtime/Datasmith/CADKernel/Public/CADKernel/Topo/TopologicalEntity.h
+++ b/Engine/Source/Runtime/Datasmith/CADKernel/Public/CADKernel/Topo/TopologicalEntity.h
@@ -51,6 +51,18 @@ public:
 		States &= ~EHaveStates::IsApplyCriteria;
 	}
 
+	bool IsNotMeshable() const
+	{
+		constexpr EHaveStates MeshedOrDeletedOrDegenerated = EHaveStates::IsMeshed | EHaveStates::Degenerated | EHaveStates::IsDeleted;
+		return ((States & MeshedOrDeletedOrDegenerated) != EHaveStates::None);
+	}
+
+	bool IsMeshable() const
+	{
+		constexpr EHaveStates MeshedOrDeletedOrDegenerated = EHaveStates::IsMeshed | EHaveStates::Degenerated | EHaveStates::IsDeleted;
+		return ((States & MeshedOrDeletedOrDegenerated) == EHaveStates::None);
+	}
+
 	bool IsMeshed() const
 	{
 		return ((States & EHaveStates::IsMeshed) == EHaveStates::IsMeshed);
diff --git a/Engine/Source/Runtime/Datasmith/CADKernel/Public/CADKernel/Topo/TopologicalFace.h b/Engine/Source/Runtime/Datasmith/CADKernel/Public/CADKernel/Topo/TopologicalFace.h
--- a/Engine/Source/Runtime/Datasmith/CADKernel/Public/CADKernel/Topo/TopologicalFace.h
+++ b/Engine/Source/Runtime/Datasmith/CADKernel/Public/CADKernel/Topo/TopologicalFace.h
@@ -16,6 +16,13 @@
 #include "CADKernel/Topo/TopologicalLoop.h"
 #include "CADKernel/Topo/TopologicalShapeEntity.h"
 
+#include "CADKernel/Mesh/Structure/EdgeSegment.h"
+#include "CADKernel/Mesh/Structure/ThinZone2D.h"
+
+#ifdef CADKERNEL_DEV
+#include "CADKernel/UI/DefineForDebug.h"
+#endif
+
 namespace UE::CADKernel
 {
 
@@ -78,6 +85,8 @@ protected:
 	 */
 	FCoordinateGrid CrossingPointDeltaMaxs;
 
+	double EstimatedMinimalElementLength = DOUBLE_BIG_NUMBER;
+
 	/**
 	 * Build a non-trimmed trimmed surface
 	 * This constructor has to be completed with one of the three "AddBoundaries" methods to be finalized.
@@ -295,7 +304,7 @@ public:
 
 	virtual void Empty() override
 	{
-#ifndef CADKERNEL_DEV
+#ifndef DONT_DELETE_FACE_FOR_DEBUG
 		CarrierSurface.Reset();
 		for (TSharedPtr<FTopologicalLoop>& Loop : Loops)
 		{
@@ -472,19 +481,27 @@ public:
 	{
 	}
 
-	// =========================================================================================================================================================================================================
-	// =========================================================================================================================================================================================================
-	// =========================================================================================================================================================================================================
-	//
-	//
-	//                                                                            NOT YET REVIEWED
-	//
-	//
-	// =========================================================================================================================================================================================================
-	// =========================================================================================================================================================================================================
-	// =========================================================================================================================================================================================================
+	// ======================================================================================================================================================================================================================
+	// Thin zone properties ===============================================================================================================================================================================
+	// ======================================================================================================================================================================================================================
+private:
+	TArray<FThinZone2D> ThinZones;
 
+public:
+	void MoveThinZones(TArray<FThinZone2D>& InThinZones)
+	{
+		ThinZones = MoveTemp(InThinZones);
+	}
 
+	const TArray<FThinZone2D>& GetThinZones() const
+	{
+		return ThinZones;
+	}
+
+	TArray<FThinZone2D>& GetThinZones()
+	{
+		return ThinZones;
+	}
 
 	// ======================================================================================================================================================================================================================
 	// Quad properties for meshing scheduling ===============================================================================================================================================================================
@@ -500,8 +517,6 @@ private:
 	FSurfaceCurvature Curvatures;
 	EQuadType QuadType = EQuadType::Unset;
 
-	double EstimatedMinimalElementLength = DOUBLE_BIG_NUMBER;
-
 public:
 	void ComputeQuadCriteria();
 	double GetQuadCriteria();
diff --git a/Engine/Source/Runtime/Datasmith/CADKernel/Public/CADKernel/Topo/TopologicalLoop.h b/Engine/Source/Runtime/Datasmith/CADKernel/Public/CADKernel/Topo/TopologicalLoop.h
--- a/Engine/Source/Runtime/Datasmith/CADKernel/Public/CADKernel/Topo/TopologicalLoop.h
+++ b/Engine/Source/Runtime/Datasmith/CADKernel/Public/CADKernel/Topo/TopologicalLoop.h
@@ -129,6 +129,8 @@ public:
 		return EEntity::TopologicalLoop;
 	}
 
+	double Length() const;
+	
 	const int32 EdgeCount() const
 	{
 		return Edges.Num();
diff --git a/Engine/Source/Runtime/Datasmith/CADKernel/Public/CADKernel/Topo/TopologicalShapeEntity.h b/Engine/Source/Runtime/Datasmith/CADKernel/Public/CADKernel/Topo/TopologicalShapeEntity.h
--- a/Engine/Source/Runtime/Datasmith/CADKernel/Public/CADKernel/Topo/TopologicalShapeEntity.h
+++ b/Engine/Source/Runtime/Datasmith/CADKernel/Public/CADKernel/Topo/TopologicalShapeEntity.h
@@ -18,6 +18,11 @@ private:
 
 public:
 
+	virtual ~FTopologicalShapeEntity() override 
+	{
+		FTopologicalShapeEntity::Empty();
+	}
+
 	virtual void Serialize(FCADKernelArchive& Ar) override
 	{
 		FTopologicalEntity::Serialize(Ar);
diff --git a/Engine/Source/Runtime/Datasmith/CADKernel/Public/CADKernel/Topo/Topomaker.h b/Engine/Source/Runtime/Datasmith/CADKernel/Public/CADKernel/Topo/Topomaker.h
--- a/Engine/Source/Runtime/Datasmith/CADKernel/Public/CADKernel/Topo/Topomaker.h
+++ b/Engine/Source/Runtime/Datasmith/CADKernel/Public/CADKernel/Topo/Topomaker.h
@@ -115,11 +115,6 @@ public:
 
 	void Sew();
 
-	/**
-	 * Check topology of each body
-	 */
-	void CheckTopology();
-
 	/**
 	 * Split into connected shell and put each shell into the appropriate body
 	 */
@@ -156,6 +151,9 @@ private:
 
 	void EmptyShells();
 
+	/**
+	 * Remove a subset of faces of their shell and keep the other shell's faces in the shell
+	 */
 	void RemoveFacesFromShell();
 
 	void RemoveEmptyShells();
diff --git a/Engine/Source/Runtime/Datasmith/CADKernel/Public/CADKernel/UI/Display.h b/Engine/Source/Runtime/Datasmith/CADKernel/Public/CADKernel/UI/Display.h
--- a/Engine/Source/Runtime/Datasmith/CADKernel/Public/CADKernel/UI/Display.h
+++ b/Engine/Source/Runtime/Datasmith/CADKernel/Public/CADKernel/UI/Display.h
@@ -8,9 +8,7 @@
 #include "CADKernel/UI/Visu.h"
 
 #ifdef CADKERNEL_DEV
-#include "CADKernel/Mesh/Meshers/IsoTriangulator/DefineForDebug.h"
-#else
-static const double DisplayScale = 1.;
+#include "CADKernel/UI/DefineForDebug.h"
 #endif
 
 class FString;
@@ -42,7 +40,7 @@ class FVertexMesh;
 class FVertexMesh;
 struct FLinearBoundary;
 
-#ifdef CADKERNEL_DEV
+#ifdef CADKERNEL_DEBUG
 void Wait(bool bMakeWait = true);
 void Open3DDebugSession(FString name, const TArray<FIdent>& idList = TArray<FIdent>());
 inline void Open3DDebugSession(bool bIsDisplayed, FString name, const TArray<FIdent>& idList = TArray<FIdent>())
@@ -68,24 +66,30 @@ private:
 public:
 	F3DDebugSession(FString Name, const TArray<FIdent>& Idents = TArray<FIdent>())
 	{
+#ifdef CADKERNEL_DEBUG
 		Open3DDebugSession(Name, Idents);
+#endif
 	}
 
 	F3DDebugSession(bool bInDisplay, FString Name, const TArray<FIdent>& Idents = TArray<FIdent>())
 		: bDisplay(bInDisplay)
 	{
+#ifdef CADKERNEL_DEBUG
 		if (bDisplay)
 		{
 			Open3DDebugSession(Name, Idents);
 		}
+#endif
 	}
 
 	~F3DDebugSession()
 	{
+#ifdef CADKERNEL_DEBUG
 		if (bDisplay)
 		{
 			Close3DDebugSession();
 		}
+#endif
 	}
 };
 
@@ -97,12 +101,16 @@ class CADKERNEL_API F3DDebugSegment
 public:
 	F3DDebugSegment(FIdent Ident)
 	{
+#ifdef CADKERNEL_DEBUG
 		Open3DDebugSegment(Ident);
+#endif
 	}
 
 	~F3DDebugSegment()
 	{
+#ifdef CADKERNEL_DEBUG
 		Close3DDebugSegment();
+#endif
 	}
 };
 
@@ -111,7 +119,9 @@ CADKERNEL_API void FlushVisu();
 template<typename TPoint>
 void DrawPoint(const TPoint& InPoint, EVisuProperty Property = EVisuProperty::BluePoint)
 {
+#ifdef CADKERNEL_DEBUG
 	FSystem::Get().GetVisu()->DrawPoint(InPoint, Property);
+#endif
 }
 
 /**
@@ -122,7 +132,9 @@ CADKERNEL_API void DrawElement(int32 Dimension, TArray<FPoint>& Points, EVisuPro
 template<typename TPoint>
 void Draw(const TArray<TPoint>& Points, EVisuProperty Property = EVisuProperty::BlueCurve)
 {
+#ifdef CADKERNEL_DEBUG
 	FSystem::Get().GetVisu()->DrawPolyline(Points, Property);
+#endif
 }
 
 CADKERNEL_API void DrawMesh(const TSharedPtr<FMesh>& mesh);
@@ -132,37 +144,48 @@ CADKERNEL_API void DisplayEdgeCriteriaGrid(int32 EdgeId, const TArray<FPoint>& P
 template<typename TPoint>
 void DisplayPoint(const TPoint& Point, FIdent Ident)
 {
+#ifdef CADKERNEL_DEBUG
 	F3DDebugSegment G(Ident);
 	DrawPoint(Point);
+#endif
 }
 
 template<typename TPoint>
 void DisplayPoint(const TPoint& Point, EVisuProperty Property = EVisuProperty::BluePoint)
 {
+#ifdef CADKERNEL_DEBUG
 	DrawPoint(Point, Property);
+#endif
 }
 
 template<typename TPoint>
 void DisplayPoint(const TPoint& Point, EVisuProperty Property, FIdent Ident)
 {
+#ifdef CADKERNEL_DEBUG
 	F3DDebugSegment G(Ident);
 	DrawPoint(Point, Property);
+#endif
 }
 
 static void DisplayPoint2DWithScale(const FPoint2D& Point, EVisuProperty Property = EVisuProperty::BluePoint)
 {
+#ifdef CADKERNEL_DEBUG
 	DrawPoint(Point * DisplayScale, Property);
+#endif
 }
 
 static void DisplayPoint2DWithScale(const FPoint2D& Point, EVisuProperty Property, FIdent Ident)
 {
+#ifdef CADKERNEL_DEBUG
 	F3DDebugSegment G(Ident);
 	DrawPoint(Point * DisplayScale, Property);
+#endif
 }
 
 template<typename TPoint>
 void DisplayPoints(FString Message, const TArray<TPoint>& Points, EVisuProperty Property = EVisuProperty::BluePoint, bool bDisplay = true)
 {
+#ifdef CADKERNEL_DEBUG
 	if (!bDisplay)
 	{
 		return;
@@ -174,6 +197,7 @@ void DisplayPoints(FString Message, const TArray<TPoint>& Points, EVisuProperty
 		DisplayPoint(Points[Index], Property);
 	}
 	Close3DDebugSession();
+#endif
 }
 
 CADKERNEL_API void DisplayProductTree(const FEntity& RootId);
@@ -226,35 +250,42 @@ CADKERNEL_API void DisplayControlPolygon(const FSurface& Entity);
 template<typename TPoint>
 void DisplaySegment(const TPoint& Point1, const TPoint& Point2, FIdent Ident = 0, EVisuProperty Property = EVisuProperty::BlueCurve, bool bWithOrientation = false)
 {
+#ifdef CADKERNEL_DEBUG
 	F3DDebugSegment G(Ident);
 	if (bWithOrientation)
 	{
 		DrawSegmentOrientation(Point1, Point2, Property);
 	}
 	DrawSegment(Point1, Point2, Property);
+ #endif
 };
 
 template<typename TPoint>
 void DisplaySegmentWithScale(const TPoint& Point1, const TPoint& Point2, FIdent Ident = 0, EVisuProperty Property = EVisuProperty::BlueCurve, bool bWithOrientation = false)
 {
+#ifdef CADKERNEL_DEBUG
 	F3DDebugSegment G(Ident);
 	if (bWithOrientation)
 	{
 		DrawSegmentOrientation(Point1 * DisplayScale, Point2 * DisplayScale, Property);
 	}
 	DrawSegment(Point1 * DisplayScale, Point2 * DisplayScale, Property);
+#endif
 };
 
 template<typename TPoint>
 void DisplayPolyline(const TArray<TPoint>& Points, EVisuProperty Property)
 {
+#ifdef CADKERNEL_DEBUG
 	Open3DDebugSegment(0);
 	Draw(Points, Property);
 	Close3DDebugSegment();
+#endif
 }
 
 static void DisplayPolylineWithScale(const TArray<FPoint2D>& Points, EVisuProperty Property)
 {
+#ifdef CADKERNEL_DEBUG
 	TArray<FPoint2D> PointsWithScale = Points;
 	for (FPoint2D& Point : PointsWithScale)
 	{
@@ -264,17 +295,20 @@ static void DisplayPolylineWithScale(const TArray<FPoint2D>& Points, EVisuProper
 	Open3DDebugSegment(0);
 	Draw(PointsWithScale, Property);
 	Close3DDebugSegment();
+#endif
 }
 
 template<typename TPoint>
 void DisplayOrientedPolyline(const TArray<TPoint>& Points, EVisuProperty Property)
 {
+#ifdef CADKERNEL_DEBUG
 	Open3DDebugSegment(0);
 	for (int32 Index = 1; Index < Points.Num(); ++Index)
 	{
 		DisplaySegment(Points[Index - 1], Points[Index], Index, Property, true);
 	}
 	Close3DDebugSegment();
+#endif
 }
 
 CADKERNEL_API void DrawQuadripode(double Height, double Base, FPoint& Centre, FPoint& Direction);
@@ -292,15 +326,18 @@ CADKERNEL_API void Draw(const FLinearBoundary& Boundary, const FRestrictionCurve
 template<typename TPoint>
 void DrawSegment(const TPoint& Point1, const TPoint& Point2, EVisuProperty Property = EVisuProperty::Element)
 {
+#ifdef CADKERNEL_DEBUG
 	TArray<FPoint> Points;
 	Points.Add(Point1);
 	Points.Add(Point2);
 	Draw(Points, Property);
+#endif
 }
 
 template<typename TPoint>
 void DrawSegmentOrientation(const TPoint& Point1, const TPoint& Point2, EVisuProperty Property = EVisuProperty::Element)
 {
+#ifdef CADKERNEL_DEBUG
 	double Length = Point1.Distance(Point2);
 	double Height = Length / 10.0;
 	double Base = Height / 2;
@@ -308,6 +345,7 @@ void DrawSegmentOrientation(const TPoint& Point1, const TPoint& Point2, EVisuPro
 	FPoint Middle = (Point1 + Point2) / 2.;
 	FPoint Tangent = Point2 - Point1;
 	DrawQuadripode(Height, Base, Middle, Tangent);
+#endif
 }
 
 CADKERNEL_API void DrawIsoCurves(const FTopologicalFace& Face);
