Fixed mismatched copy constructors/assignments in UE::CADKernel::FParameterValue and UE::CADKernel::FParameter.

diff --git a/Engine/Source/Runtime/Datasmith/CADKernel/Public/CADKernel/Core/Parameter.h b/Engine/Source/Runtime/Datasmith/CADKernel/Public/CADKernel/Core/Parameter.h
--- a/Engine/Source/Runtime/Datasmith/CADKernel/Public/CADKernel/Core/Parameter.h
+++ b/Engine/Source/Runtime/Datasmith/CADKernel/Public/CADKernel/Core/Parameter.h
@@ -55,13 +55,6 @@ public:
 		*this = 0;
 	}
 
-	FParameterValue(const FParameterValue& Other)
-	{
-		Type = Other.Type;
-		Value = Other.Value;
-		StringValue = Other.StringValue;
-	}
-
 	FString ToString() const;
 
 	inline EValue GetType() const
@@ -154,6 +147,7 @@ private:
 	void AddToParameterMap(FParameter& Parameter, FParameters& Parameters);
 
 public:
+	FParameter(const FParameter&) = default;
 
 	template<typename ValueType>
 	FParameter(const FString& ParameterName, const ValueType& FirstValue, FParameters& Parameters)
