Defaulted the special member functions in TOrientedEntity to avoid Clang 16 warnings.

diff --git a/Engine/Source/Runtime/Datasmith/CADKernel/Public/CADKernel/Core/OrientedEntity.h b/Engine/Source/Runtime/Datasmith/CADKernel/Public/CADKernel/Core/OrientedEntity.h
--- a/Engine/Source/Runtime/Datasmith/CADKernel/Public/CADKernel/Core/OrientedEntity.h
+++ b/Engine/Source/Runtime/Datasmith/CADKernel/Public/CADKernel/Core/OrientedEntity.h
@@ -19,13 +19,11 @@ public:
 	{
 	}
 
-	TOrientedEntity(const TOrientedEntity<KeyType>& OrientiredEntity)
-		: Entity(OrientiredEntity.Entity)
-		, Direction(OrientiredEntity.Direction)
-	{
-	}
-
 	TOrientedEntity() = default;
+	TOrientedEntity(TOrientedEntity&&) = default;
+	TOrientedEntity(const TOrientedEntity&) = default;
+	TOrientedEntity& operator=(TOrientedEntity&&) = default;
+	TOrientedEntity& operator=(const TOrientedEntity&) = default;
 };
 }
 
