Remove unused class: UE::CADKernel::FParameters

diff --git a/Engine/Plugins/Enterprise/DatasmithCADImporter/Source/DatasmithCADTranslator/Private/DatasmithCADTranslator.cpp b/Engine/Plugins/Enterprise/DatasmithCADImporter/Source/DatasmithCADTranslator/Private/DatasmithCADTranslator.cpp
--- a/Engine/Plugins/Enterprise/DatasmithCADImporter/Source/DatasmithCADTranslator/Private/DatasmithCADTranslator.cpp
+++ b/Engine/Plugins/Enterprise/DatasmithCADImporter/Source/DatasmithCADTranslator/Private/DatasmithCADTranslator.cpp
@@ -137,7 +137,6 @@ void FDatasmithCADTranslator::Initialize(FDatasmithTranslatorCapabilities& OutCa
 		OutCapabilities.SupportedFileFormats.Add(FFileFormatInfo{ TEXT("vda"), TEXT("VDA-FS") });
 		OutCapabilities.SupportedFileFormats.Add(FFileFormatInfo{ TEXT("vrml"), TEXT("VRML") });
 		OutCapabilities.SupportedFileFormats.Add(FFileFormatInfo{ TEXT("wrl"), TEXT("VRML") });
-		OutCapabilities.SupportedFileFormats.Add(FFileFormatInfo{ TEXT("obj"), TEXT("Wavefront OBJ") });
 	}
 }
 
diff --git a/Engine/Source/Runtime/Datasmith/CADKernel/Private/CADKernel/Core/KernelParameters.h b/Engine/Source/Runtime/Datasmith/CADKernel/Private/CADKernel/Core/KernelParameters.h
deleted file mode 100644
diff --git a/Engine/Source/Runtime/Datasmith/CADKernel/Private/CADKernel/Core/Parameters.cpp b/Engine/Source/Runtime/Datasmith/CADKernel/Private/CADKernel/Core/Parameters.cpp
deleted file mode 100644
diff --git a/Engine/Source/Runtime/Datasmith/CADKernel/Private/CADKernel/Core/System.cpp b/Engine/Source/Runtime/Datasmith/CADKernel/Private/CADKernel/Core/System.cpp
--- a/Engine/Source/Runtime/Datasmith/CADKernel/Private/CADKernel/Core/System.cpp
+++ b/Engine/Source/Runtime/Datasmith/CADKernel/Private/CADKernel/Core/System.cpp
@@ -3,7 +3,6 @@
 
 #include "CADKernel/Core/Database.h"
 #include "CADKernel/Core/Entity.h"
-#include "CADKernel/Core/KernelParameters.h"
 #include "CADKernel/UI/Visu.h"
 #include "CADKernel/Utils/Util.h"
 
@@ -23,12 +22,13 @@ namespace UE::CADKernel
 TUniquePtr<FSystem> FSystem::Instance = nullptr;
 
 FSystem::FSystem()
-	: Parameters(MakeShared<FKernelParameters>())
-	, DefaultVisu()
+	: DefaultVisu()
 	, Viewer(&DefaultVisu)
 	, Console(&DefaultConsole)
 	, ProgressManager(&DefaultProgressManager)
-
+#ifdef CADKERNEL_DEV
+	, Parameters(MakeShared<FKernelParameters>())
+#endif
 {
 	LogLevel = Log;
 	VerboseLevel = Log;
diff --git a/Engine/Source/Runtime/Datasmith/CADKernel/Private/CADKernel/Geo/Curves/CompositeCurve.cpp b/Engine/Source/Runtime/Datasmith/CADKernel/Private/CADKernel/Geo/Curves/CompositeCurve.cpp
--- a/Engine/Source/Runtime/Datasmith/CADKernel/Private/CADKernel/Geo/Curves/CompositeCurve.cpp
+++ b/Engine/Source/Runtime/Datasmith/CADKernel/Private/CADKernel/Geo/Curves/CompositeCurve.cpp
@@ -2,7 +2,6 @@
 
 #include "CADKernel/Geo/Curves/CompositeCurve.h"
 
-#include "CADKernel/Core/KernelParameters.h"
 #include "CADKernel/Core/System.h"
 #include "CADKernel/Math/Boundary.h"
 #include "CADKernel/Geo/GeoPoint.h"
diff --git a/Engine/Source/Runtime/Datasmith/CADKernel/Private/CADKernel/Geo/Curves/Curve.cpp b/Engine/Source/Runtime/Datasmith/CADKernel/Private/CADKernel/Geo/Curves/Curve.cpp
--- a/Engine/Source/Runtime/Datasmith/CADKernel/Private/CADKernel/Geo/Curves/Curve.cpp
+++ b/Engine/Source/Runtime/Datasmith/CADKernel/Private/CADKernel/Geo/Curves/Curve.cpp
@@ -1,7 +1,6 @@
 // Copyright Epic Games, Inc. All Rights Reserved.
 #include "CADKernel/Geo/Curves/Curve.h"
 
-#include "CADKernel/Core/KernelParameters.h"
 #include "CADKernel/Core/System.h"
 #include "CADKernel/Geo/Curves/BoundedCurve.h"
 #include "CADKernel/Geo/Sampling/Polyline.h"
diff --git a/Engine/Source/Runtime/Datasmith/CADKernel/Private/CADKernel/Geo/Curves/SurfacicCurve.cpp b/Engine/Source/Runtime/Datasmith/CADKernel/Private/CADKernel/Geo/Curves/SurfacicCurve.cpp
--- a/Engine/Source/Runtime/Datasmith/CADKernel/Private/CADKernel/Geo/Curves/SurfacicCurve.cpp
+++ b/Engine/Source/Runtime/Datasmith/CADKernel/Private/CADKernel/Geo/Curves/SurfacicCurve.cpp
@@ -2,7 +2,6 @@
 
 #include "CADKernel/Geo/Curves/SurfacicCurve.h"
 
-#include "CADKernel/Core/KernelParameters.h"
 #include "CADKernel/Core/System.h"
 #include "CADKernel/Geo/GeoPoint.h"
 #include "CADKernel/Geo/Sampling/Polyline.h"
diff --git a/Engine/Source/Runtime/Datasmith/CADKernel/Private/CADKernel/Geo/Sampling/SurfacicPolyline.cpp b/Engine/Source/Runtime/Datasmith/CADKernel/Private/CADKernel/Geo/Sampling/SurfacicPolyline.cpp
--- a/Engine/Source/Runtime/Datasmith/CADKernel/Private/CADKernel/Geo/Sampling/SurfacicPolyline.cpp
+++ b/Engine/Source/Runtime/Datasmith/CADKernel/Private/CADKernel/Geo/Sampling/SurfacicPolyline.cpp
@@ -2,7 +2,6 @@
 
 #include "CADKernel/Geo/Sampling/SurfacicPolyline.h"
 
-#include "CADKernel/Core/KernelParameters.h"
 #include "CADKernel/Core/System.h"
 #include "CADKernel/Geo/Curves/Curve.h"
 #include "CADKernel/Geo/Surfaces/Surface.h"
diff --git a/Engine/Source/Runtime/Datasmith/CADKernel/Private/CADKernel/Geo/Surfaces/OffsetSurface.cpp b/Engine/Source/Runtime/Datasmith/CADKernel/Private/CADKernel/Geo/Surfaces/OffsetSurface.cpp
--- a/Engine/Source/Runtime/Datasmith/CADKernel/Private/CADKernel/Geo/Surfaces/OffsetSurface.cpp
+++ b/Engine/Source/Runtime/Datasmith/CADKernel/Private/CADKernel/Geo/Surfaces/OffsetSurface.cpp
@@ -2,7 +2,6 @@
 
 #include "CADKernel/Geo/Surfaces/OffsetSurface.h"
 
-#include "CADKernel/Core/KernelParameters.h"
 #include "CADKernel/Core/System.h"
 #include "CADKernel/Geo/GeoPoint.h"
 #include "CADKernel/Geo/Sampling/SurfacicSampling.h"
diff --git a/Engine/Source/Runtime/Datasmith/CADKernel/Private/CADKernel/Geo/Surfaces/PlaneSurface.cpp b/Engine/Source/Runtime/Datasmith/CADKernel/Private/CADKernel/Geo/Surfaces/PlaneSurface.cpp
--- a/Engine/Source/Runtime/Datasmith/CADKernel/Private/CADKernel/Geo/Surfaces/PlaneSurface.cpp
+++ b/Engine/Source/Runtime/Datasmith/CADKernel/Private/CADKernel/Geo/Surfaces/PlaneSurface.cpp
@@ -2,7 +2,6 @@
 
 #include "CADKernel/Geo/Surfaces/PlaneSurface.h"
 
-#include "CADKernel/Core/KernelParameters.h"
 #include "CADKernel/Core/System.h"
 #include "CADKernel/Geo/GeoPoint.h"
 #include "CADKernel/Geo/Sampling/SurfacicSampling.h"
diff --git a/Engine/Source/Runtime/Datasmith/CADKernel/Private/CADKernel/Geo/Surfaces/Surface.cpp b/Engine/Source/Runtime/Datasmith/CADKernel/Private/CADKernel/Geo/Surfaces/Surface.cpp
--- a/Engine/Source/Runtime/Datasmith/CADKernel/Private/CADKernel/Geo/Surfaces/Surface.cpp
+++ b/Engine/Source/Runtime/Datasmith/CADKernel/Private/CADKernel/Geo/Surfaces/Surface.cpp
@@ -2,7 +2,6 @@
 
 #include "CADKernel/Geo/Surfaces/Surface.h"
 
-#include "CADKernel/Core/KernelParameters.h"
 #include "CADKernel/Core/System.h"
 
 #include "CADKernel/Geo/GeoPoint.h"
diff --git a/Engine/Source/Runtime/Datasmith/CADKernel/Private/CADKernel/Math/Geometry.cpp b/Engine/Source/Runtime/Datasmith/CADKernel/Private/CADKernel/Math/Geometry.cpp
--- a/Engine/Source/Runtime/Datasmith/CADKernel/Private/CADKernel/Math/Geometry.cpp
+++ b/Engine/Source/Runtime/Datasmith/CADKernel/Private/CADKernel/Math/Geometry.cpp
@@ -2,7 +2,6 @@
 
 #include "CADKernel/Math/Geometry.h"
 
-#include "CADKernel/Core/KernelParameters.h" 
 #include "CADKernel/Core/System.h"
 #include "CADKernel/Math/Geometry.h"
 #include "CADKernel/Math/Point.h"
diff --git a/Engine/Source/Runtime/Datasmith/CADKernel/Private/CADKernel/Mesh/Criteria/Criterion.cpp b/Engine/Source/Runtime/Datasmith/CADKernel/Private/CADKernel/Mesh/Criteria/Criterion.cpp
--- a/Engine/Source/Runtime/Datasmith/CADKernel/Private/CADKernel/Mesh/Criteria/Criterion.cpp
+++ b/Engine/Source/Runtime/Datasmith/CADKernel/Private/CADKernel/Mesh/Criteria/Criterion.cpp
@@ -7,7 +7,6 @@
 #include "CADKernel/Mesh/Criteria/SagCriterion.h"
 #include "CADKernel/Mesh/Criteria/SizeCriterion.h"
 
-#include "CADKernel/Core/KernelParameters.h"
 #include "CADKernel/Geo/GeoPoint.h"
 #include "CADKernel/Topo/TopologicalEdge.h"
 
diff --git a/Engine/Source/Runtime/Datasmith/CADKernel/Private/CADKernel/Mesh/Criteria/SizeCriterion.cpp b/Engine/Source/Runtime/Datasmith/CADKernel/Private/CADKernel/Mesh/Criteria/SizeCriterion.cpp
--- a/Engine/Source/Runtime/Datasmith/CADKernel/Private/CADKernel/Mesh/Criteria/SizeCriterion.cpp
+++ b/Engine/Source/Runtime/Datasmith/CADKernel/Private/CADKernel/Mesh/Criteria/SizeCriterion.cpp
@@ -1,7 +1,6 @@
 // Copyright Epic Games, Inc. All Rights Reserved.
 #include "CADKernel/Mesh/Criteria/SizeCriterion.h"
 
-#include "CADKernel/Core/KernelParameters.h"
 #include "CADKernel/Core/System.h"
 #include "CADKernel/Geo/GeoPoint.h"
 #include "CADKernel/Topo/TopologicalEdge.h"
diff --git a/Engine/Source/Runtime/Datasmith/CADKernel/Private/CADKernel/Mesh/Structure/ThinZone2DDisplay.cpp b/Engine/Source/Runtime/Datasmith/CADKernel/Private/CADKernel/Mesh/Structure/ThinZone2DDisplay.cpp
deleted file mode 100644
diff --git a/Engine/Source/Runtime/Datasmith/CADKernel/Private/CADKernel/Topo/TopologicalEdge.cpp b/Engine/Source/Runtime/Datasmith/CADKernel/Private/CADKernel/Topo/TopologicalEdge.cpp
--- a/Engine/Source/Runtime/Datasmith/CADKernel/Private/CADKernel/Topo/TopologicalEdge.cpp
+++ b/Engine/Source/Runtime/Datasmith/CADKernel/Private/CADKernel/Topo/TopologicalEdge.cpp
@@ -2,7 +2,6 @@
 
 #include "CADKernel/Topo/TopologicalEdge.h"
 
-#include "CADKernel/Core/KernelParameters.h"
 #include "CADKernel/Geo/Curves/NURBSCurve.h"
 #include "CADKernel/Geo/Curves/SegmentCurve.h"
 #include "CADKernel/Geo/GeoEnum.h"
diff --git a/Engine/Source/Runtime/Datasmith/CADKernel/Private/CADKernel/Topo/TopologicalFace.cpp b/Engine/Source/Runtime/Datasmith/CADKernel/Private/CADKernel/Topo/TopologicalFace.cpp
--- a/Engine/Source/Runtime/Datasmith/CADKernel/Private/CADKernel/Topo/TopologicalFace.cpp
+++ b/Engine/Source/Runtime/Datasmith/CADKernel/Private/CADKernel/Topo/TopologicalFace.cpp
@@ -2,7 +2,6 @@
 
 #include "CADKernel/Topo/TopologicalFace.h"
 
-#include "CADKernel/Core/KernelParameters.h"
 #include "CADKernel/Core/System.h"
 #include "CADKernel/Geo/Curves/RestrictionCurve.h"
 #include "CADKernel/Geo/Curves/SegmentCurve.h"
diff --git a/Engine/Source/Runtime/Datasmith/CADKernel/Private/CADKernel/Topo/TopologicalLoop.cpp b/Engine/Source/Runtime/Datasmith/CADKernel/Private/CADKernel/Topo/TopologicalLoop.cpp
--- a/Engine/Source/Runtime/Datasmith/CADKernel/Private/CADKernel/Topo/TopologicalLoop.cpp
+++ b/Engine/Source/Runtime/Datasmith/CADKernel/Private/CADKernel/Topo/TopologicalLoop.cpp
@@ -3,7 +3,6 @@
 #include "CADKernel/Topo/TopologicalLoop.h"
 
 #include "Algo/AllOf.h"
-#include "CADKernel/Core/KernelParameters.h"
 #include "CADKernel/Geo/Curves/RestrictionCurve.h"
 #include "CADKernel/Geo/GeoEnum.h"
 #include "CADKernel/Math/SlopeUtils.h"
diff --git a/Engine/Source/Runtime/Datasmith/CADKernel/Public/CADKernel/Core/Parameter.h b/Engine/Source/Runtime/Datasmith/CADKernel/Public/CADKernel/Core/Parameter.h
deleted file mode 100644
diff --git a/Engine/Source/Runtime/Datasmith/CADKernel/Public/CADKernel/Core/Parameters.h b/Engine/Source/Runtime/Datasmith/CADKernel/Public/CADKernel/Core/Parameters.h
deleted file mode 100644
diff --git a/Engine/Source/Runtime/Datasmith/CADKernel/Public/CADKernel/Geo/Curves/Curve.h b/Engine/Source/Runtime/Datasmith/CADKernel/Public/CADKernel/Geo/Curves/Curve.h
--- a/Engine/Source/Runtime/Datasmith/CADKernel/Public/CADKernel/Geo/Curves/Curve.h
+++ b/Engine/Source/Runtime/Datasmith/CADKernel/Public/CADKernel/Geo/Curves/Curve.h
@@ -4,7 +4,6 @@
 #include "CADKernel/Core/CADKernelArchive.h"
 #include "CADKernel/Core/Entity.h"
 #include "CADKernel/Core/EntityGeom.h"
-#include "CADKernel/Core/Parameters.h"
 #include "CADKernel/Core/Types.h"
 #include "CADKernel/Geo/GeoEnum.h"
 #include "CADKernel/Geo/GeoPoint.h"
diff --git a/Engine/Source/Runtime/Datasmith/CADKernel/Public/CADKernel/UI/Visu.h b/Engine/Source/Runtime/Datasmith/CADKernel/Public/CADKernel/UI/Visu.h
--- a/Engine/Source/Runtime/Datasmith/CADKernel/Public/CADKernel/UI/Visu.h
+++ b/Engine/Source/Runtime/Datasmith/CADKernel/Public/CADKernel/UI/Visu.h
@@ -1,8 +1,11 @@
 // Copyright Epic Games, Inc. All Rights Reserved.
 #pragma once
 
+#ifdef CADKERNEL_DEV
 #include "CADKernel/Core/Parameter.h"
 #include "CADKernel/Core/Parameters.h"
+#endif
+
 #include "CADKernel/Core/Types.h"
 #include "CADKernel/Math/Point.h"
 
@@ -43,6 +46,8 @@ extern const TCHAR* VisuPropertyNames[];
 
 class FMesh;
 
+#ifdef CADKERNEL_DEV
+
 class CADKERNEL_API FVisuParameters : public FParameters
 {
 public:
@@ -78,11 +83,14 @@ public:
 		, NormalLength(TEXT("NormalLength"), 10., *this)
 	{}
 };
+#endif
 
 class CADKERNEL_API FVisu
 {
 protected:
+#ifdef CADKERNEL_DEV
 	FVisuParameters Parameters;
+#endif
 
 	/**
 	 * Number of graphic session opened
@@ -118,10 +126,12 @@ public:
 		return EVisuProperty::Last;
 	}
 
+#ifdef CADKERNEL_DEV
 	FVisuParameters* GetParameters()
 	{
 		return &Parameters;
 	}
+#endif
 
 	virtual void NewDB(const TCHAR* InName)
 	{}
