Fix Static Analysis Win64 (MSVC)

diff --git a/Engine/Source/Runtime/Datasmith/CADKernel/Private/CADKernel/Mesh/Meshers/ParametricMesher.cpp b/Engine/Source/Runtime/Datasmith/CADKernel/Private/CADKernel/Mesh/Meshers/ParametricMesher.cpp
--- a/Engine/Source/Runtime/Datasmith/CADKernel/Private/CADKernel/Mesh/Meshers/ParametricMesher.cpp
+++ b/Engine/Source/Runtime/Datasmith/CADKernel/Private/CADKernel/Mesh/Meshers/ParametricMesher.cpp
@@ -1447,7 +1447,6 @@ void FParametricMesher::FindThinZoneBoundary(FThinZoneSide& Side)
 	Edge->AddThinZone(SideEdgeCoordinate);
 }
 
-#define DEBUG_MESHTHINSURF
 void FParametricMesher::MeshThinZoneSide(FThinZoneSide& Side)
 {
 	typedef TFunction<bool(double, double)> CompareMethode;
@@ -1553,6 +1552,10 @@ void FParametricMesher::MeshThinZoneSide(FThinZoneSide& Side)
 			Open3DDebugSession(TEXT("Projection of mesh"));
 #endif
 			Edge = EdgeSegment.GetEdge();
+			if (Edge == nullptr)
+			{
+				continue;
+			}
 
 			if (!Edge->IsMeshed())
 			{
