Fixed mismatched copy constructor/assignment operator in FOrientedEdge.

diff --git a/Engine/Source/Runtime/Datasmith/CADKernel/Public/CADKernel/Topo/TopologicalLoop.h b/Engine/Source/Runtime/Datasmith/CADKernel/Public/CADKernel/Topo/TopologicalLoop.h
--- a/Engine/Source/Runtime/Datasmith/CADKernel/Public/CADKernel/Topo/TopologicalLoop.h
+++ b/Engine/Source/Runtime/Datasmith/CADKernel/Public/CADKernel/Topo/TopologicalLoop.h
@@ -26,11 +26,6 @@ public:
 	{
 	}
 
-	FOrientedEdge(const FOrientedEdge& OrientiredEntity)
-		: TOrientedEntity(OrientiredEntity)
-	{
-	}
-
 	FOrientedEdge()
 		: TOrientedEntity()
 	{
