Fixed mismatched copy constructor/assignment operator in UE::CADKernel::FLinearBoundary.

diff --git a/Engine/Source/Runtime/Datasmith/CADKernel/Public/CADKernel/Math/Boundary.h b/Engine/Source/Runtime/Datasmith/CADKernel/Public/CADKernel/Math/Boundary.h
--- a/Engine/Source/Runtime/Datasmith/CADKernel/Public/CADKernel/Math/Boundary.h
+++ b/Engine/Source/Runtime/Datasmith/CADKernel/Public/CADKernel/Math/Boundary.h
@@ -407,13 +407,6 @@ public:
 		UVBoundaries[EIso::IsoV].MoveInsideIfNot(Point.V, Tolerance);
 	}
 
-	FSurfacicBoundary& operator=(const FSurfacicBoundary& InBounds)
-	{
-		UVBoundaries[EIso::IsoU] = InBounds.UVBoundaries[EIso::IsoU];
-		UVBoundaries[EIso::IsoV] = InBounds.UVBoundaries[EIso::IsoV];
-		return *this;
-	}
-
 	double Length(const EIso& Iso) const
 	{
 		return UVBoundaries[Iso].Length();
