Fixed mismatched copy constructor/assignment operator in UE::CADKernel::FOrientedFace.

diff --git a/Engine/Source/Runtime/Datasmith/CADKernel/Public/CADKernel/Topo/Shell.h b/Engine/Source/Runtime/Datasmith/CADKernel/Public/CADKernel/Topo/Shell.h
--- a/Engine/Source/Runtime/Datasmith/CADKernel/Public/CADKernel/Topo/Shell.h
+++ b/Engine/Source/Runtime/Datasmith/CADKernel/Public/CADKernel/Topo/Shell.h
@@ -29,11 +29,6 @@ public:
 	{
 	}
 
-	FOrientedFace(const FOrientedFace& OrientiredEntity)
-		: TOrientedEntity(OrientiredEntity)
-	{
-	}
-
 	FOrientedFace()
 		: TOrientedEntity()
 	{
