Fixed mismatched copy constructor/assignment operator in UE::CADKernel::FOrientedCurve.

diff --git a/Engine/Source/Runtime/Datasmith/CADKernel/Private/CADKernel/Geo/Curves/CompositeCurve.h b/Engine/Source/Runtime/Datasmith/CADKernel/Private/CADKernel/Geo/Curves/CompositeCurve.h
--- a/Engine/Source/Runtime/Datasmith/CADKernel/Private/CADKernel/Geo/Curves/CompositeCurve.h
+++ b/Engine/Source/Runtime/Datasmith/CADKernel/Private/CADKernel/Geo/Curves/CompositeCurve.h
@@ -16,11 +16,6 @@ public:
 	{
 	}
 
-	FOrientedCurve(const FOrientedCurve& OrientiredEntity)
-		: TOrientedEntity(OrientiredEntity)
-	{
-	}
-
 	FOrientedCurve()
 		: TOrientedEntity()
 	{
