Fixed mismatched copy constructor/assignment operator in FNURBSCurve.

diff --git a/Engine/Source/Runtime/Datasmith/CADKernel/Private/CADKernel/Geo/Curves/NURBSCurve.cpp b/Engine/Source/Runtime/Datasmith/CADKernel/Private/CADKernel/Geo/Curves/NURBSCurve.cpp
--- a/Engine/Source/Runtime/Datasmith/CADKernel/Private/CADKernel/Geo/Curves/NURBSCurve.cpp
+++ b/Engine/Source/Runtime/Datasmith/CADKernel/Private/CADKernel/Geo/Curves/NURBSCurve.cpp
@@ -35,17 +35,6 @@ FNURBSCurve::FNURBSCurve(FNurbsCurveData& NurbsCurveData)
 	Finalize();
 }
 
-FNURBSCurve::FNURBSCurve(const FNURBSCurve& Nurbs)
-	: FCurve(Nurbs.Dimension)
-	, Degree(Nurbs.Degree)
-	, NodalVector(Nurbs.NodalVector)
-	, Weights(Nurbs.Weights)
-	, Poles(Nurbs.Poles)
-	, bIsRational(Nurbs.bIsRational)
-{
-	Finalize();
-}
-
 TSharedPtr<FEntityGeom> FNURBSCurve::ApplyMatrix(const FMatrixH& InMatrix) const
 {
 	TArray<FPoint> TransformedPoles;
diff --git a/Engine/Source/Runtime/Datasmith/CADKernel/Private/CADKernel/Geo/Curves/NURBSCurve.h b/Engine/Source/Runtime/Datasmith/CADKernel/Private/CADKernel/Geo/Curves/NURBSCurve.h
--- a/Engine/Source/Runtime/Datasmith/CADKernel/Private/CADKernel/Geo/Curves/NURBSCurve.h
+++ b/Engine/Source/Runtime/Datasmith/CADKernel/Private/CADKernel/Geo/Curves/NURBSCurve.h
@@ -41,7 +41,6 @@ protected:
 
 	FNURBSCurve(int32 InDegree, const TArray<double>& InNodalVector, const TArray<FPoint>& InPoles, const TArray<double>& InWeights, int8 InDimension = 3);
 	FNURBSCurve(FNurbsCurveData& NurbsCurveData);
-	FNURBSCurve(const FNURBSCurve& Nurbs);
 
 	FNURBSCurve() = default;
 
