CADKernel update: Meshing of thin faces

diff --git a/Engine/Source/Runtime/Datasmith/CADKernel/Private/CADKernel/Mesh/Meshers/IsoTriangulator.cpp b/Engine/Source/Runtime/Datasmith/CADKernel/Private/CADKernel/Mesh/Meshers/IsoTriangulator.cpp
--- a/Engine/Source/Runtime/Datasmith/CADKernel/Private/CADKernel/Mesh/Meshers/IsoTriangulator.cpp
+++ b/Engine/Source/Runtime/Datasmith/CADKernel/Private/CADKernel/Mesh/Meshers/IsoTriangulator.cpp
@@ -6,6 +6,7 @@
 #include "CADKernel/Math/MathConst.h"
 #include "CADKernel/Math/Point.h"
 #include "CADKernel/Mesh/Meshers/BowyerWatsonTriangulator.h"
+#include "CADKernel/Mesh/Meshers/CycleTriangulator.h"
 #include "CADKernel/Mesh/Meshers/IsoTriangulator/IsoCell.h"
 #include "CADKernel/Mesh/Meshers/IsoTriangulator/IsoNode.h"
 #include "CADKernel/Mesh/Meshers/IsoTriangulator/IsoSegment.h"
@@ -13,33 +14,58 @@
 #include "CADKernel/Mesh/Structure/EdgeMesh.h"
 #include "CADKernel/Mesh/Structure/FaceMesh.h"
 #include "CADKernel/Mesh/Structure/LoopCleaner.h"
+#include "CADKernel/Mesh/Structure/ThinZone2D.h"
 #include "CADKernel/Topo/TopologicalEdge.h"
 #include "CADKernel/Topo/TopologicalFace.h"
-#include "CADKernel/Topo/TopologicalFace.h"
 #include "CADKernel/UI/Display.h"
 #include "CADKernel/Utils/ArrayUtils.h"
 
 namespace UE::CADKernel
 {
 
-const double FIsoTriangulator::GeometricToMeshingToleranceFactor = 10.;
-
 #ifdef DEBUG_BOWYERWATSON
 bool FBowyerWatsonTriangulator::bDisplay = false;
 #endif
 
-FIsoTriangulator::FIsoTriangulator(FGrid& InGrid, TSharedRef<FFaceMesh> EntityMesh)
+namespace IsoTriangulatorImpl
+{
+const double MaxSlopeToBeIso = 0.125;
+
+const double LimitValueMin(double Slope)
+{
+	return Slope - MaxSlopeToBeIso;
+}
+
+const double LimitValueMax(double Slope)
+{
+	return Slope + MaxSlopeToBeIso;
+}
+
+struct FCandidateSegment
+{
+	FLoopNode& StartNode;
+	FLoopNode& EndNode;
+	double Length;
+	
+	FCandidateSegment(const FGrid& Grid, FLoopNode& Node1, FLoopNode& Node2)
+		: StartNode(Node1)
+		, EndNode(Node2)
+	{
+		Length = Node1.Get2DPoint(EGridSpace::UniformScaled, Grid).Distance(Node2.Get2DPoint(EGridSpace::UniformScaled, Grid));
+	}
+};
+
+}
+
+FIsoTriangulator::FIsoTriangulator(FGrid& InGrid, FFaceMesh& OutMesh, const FMeshingTolerances& InTolerances)
 	: Grid(InGrid)
-	, Mesh(EntityMesh)
+	, Mesh(OutMesh)
 	, LoopSegmentsIntersectionTool(InGrid)
 	, InnerSegmentsIntersectionTool(InGrid)
 	, InnerToLoopSegmentsIntersectionTool(InGrid)
 	, InnerToOuterSegmentsIntersectionTool(InGrid)
-	, GeometricTolerance(InGrid.GetFace().GetCarrierSurface()->Get3DTolerance())
-	, SquareGeometricTolerance(FMath::Square(GeometricTolerance))
-	, SquareGeometricTolerance2(2. * SquareGeometricTolerance)
-	, MeshingTolerance(GeometricTolerance* GeometricToMeshingToleranceFactor)
-	, SquareMeshingTolerance(FMath::Square(MeshingTolerance))
+	, ThinZoneIntersectionTool(InGrid)
+	, Tolerances(InTolerances)
 {
 	FinalInnerSegments.Reserve(3 * Grid.InnerNodesCount());
 	IndexOfLowerLeftInnerNodeSurroundingALoop.Reserve(Grid.GetLoopCount());
@@ -80,13 +106,14 @@ bool FIsoTriangulator::Triangulate()
 
 #ifdef DEBUG_ISOTRIANGULATOR
 	DisplayLoops(TEXT("FIsoTrianguler::LoopSegments"), true, false);
+	Wait(false);
 #endif
 
 	FLoopCleaner LoopCleaner(*this);
 	if (!LoopCleaner.Run())
 	{
 #ifdef CADKERNEL_DEV
-		MesherReport->Logs.AddDegeneratedLoop();
+		FMesherReport::Get().Logs.AddDegeneratedLoop();
 #endif
 		FMessage::Printf(EVerboseLevel::Log, TEXT("The meshing of the surface %d failed due to a degenerated loop\n"), Grid.GetFace().GetId());
 		return false;
@@ -99,14 +126,18 @@ bool FIsoTriangulator::Triangulate()
 	Wait(false);
 #endif
 
-	BuildThinZoneSegments();
-
 	//Fill Intersection tool
 	LoopSegmentsIntersectionTool.Empty(LoopSegments.Num());
 	LoopSegmentsIntersectionTool.AddSegments(LoopSegments);
+	LoopSegmentsIntersectionTool.Sort();
+
+	GetThinZonesMesh();
+
 	LoopSegmentsIntersectionTool.AddSegments(ThinZoneSegments);
 	LoopSegmentsIntersectionTool.Sort();
 
+	FinalToLoops.Append(ThinZoneSegments);
+
 #ifdef DEBUG_ISOTRIANGULATOR
 	Grid.DisplayIsoSegments(TEXT("FIsoTrianguler::ThinZoneSegments"), DisplaySpace, ThinZoneSegments, false);
 	LoopSegmentsIntersectionTool.Display(Grid.bDisplay, TEXT("FIsoTrianguler::IntersectionTool Loop"));
@@ -115,12 +146,13 @@ bool FIsoTriangulator::Triangulate()
 	BuildInnerSegments();
 
 #ifdef DEBUG_ISOTRIANGULATOR
+	if (bDisplay)
 	{
-		F3DDebugSession A(bDisplay, TEXT("FIsoTrianguler::FinalInnerSegments"));
+		F3DDebugSession A(TEXT("FIsoTrianguler::FinalInnerSegments"));
 		Grid.DisplayIsoSegments(DisplaySpace, FinalInnerSegments, true, false, false, EVisuProperty::BlueCurve);
+		//Wait();
 	}
 	InnerToOuterSegmentsIntersectionTool.Display(bDisplay, TEXT("FIsoTrianguler::IntersectionTool InnerToOuter"), EVisuProperty::RedCurve);
-	Chronos.TriangulateDuration1 = FChrono::Elapse(StartTime);
 #endif
 
 	// =============================================================================================================
@@ -131,10 +163,6 @@ bool FIsoTriangulator::Triangulate()
 	InnerSegmentsIntersectionTool.Display(Grid.bDisplay, TEXT("FIsoTrianguler::IntersectionTool Inner"));
 #endif
 
-#ifdef DEBUG_ISOTRIANGULATOR
-	Chronos.TriangulateDuration2 = FChrono::Elapse(StartTime);
-#endif
-
 	// =============================================================================================================
 	// 	   For each cell
 	// 	      - Connect loops together and to cell vertices
@@ -150,8 +178,8 @@ bool FIsoTriangulator::Triangulate()
 	ConnectCellLoops();
 
 #ifdef DEBUG_ISOTRIANGULATOR
-	Grid.DisplayIsoSegments(TEXT("FIsoTrianguler::Final Iso ToLink Inner To Loops"), DisplaySpace, FinalToLoops, false, false, true, EVisuProperty::YellowCurve);
-	Wait(bDisplay);
+	Grid.DisplayIsoSegments(TEXT("FIsoTrianguler::Final Iso ToLink Inner To Loops"), DisplaySpace, FinalToLoops, false, false, false, EVisuProperty::YellowCurve);
+	//Wait(bDisplay);
 #endif
 
 	// =============================================================================================================
@@ -168,7 +196,7 @@ bool FIsoTriangulator::Triangulate()
 	if (bDisplay)
 	{
 		F3DDebugSession A(TEXT("Mesh 3D"));
-		DisplayMesh(*Mesh);
+		DisplayMesh(Mesh);
 	}
 	//Wait(bDisplay);
 #endif
@@ -236,9 +264,9 @@ void FIsoTriangulator::BuildNodes()
 void FIsoTriangulator::FillMeshNodes()
 {
 	int32 TriangleNum = 50 + (int32)((2 * InnerNodeCount + LoopNodeCount) * 1.1);
-	Mesh->Init(TriangleNum, InnerNodeCount + LoopNodeCount);
+	Mesh.Init(TriangleNum, InnerNodeCount + LoopNodeCount);
 
-	TArray<FPoint>& InnerNodeCoordinates = Mesh->GetNodeCoordinates();
+	TArray<FPoint>& InnerNodeCoordinates = Mesh.GetNodeCoordinates();
 	InnerNodeCoordinates.Reserve(InnerNodeCount);
 	for (int32 Index = 0; Index < (int32)Grid.GetInner3DPoints().Num(); ++Index)
 	{
@@ -248,62 +276,47 @@ void FIsoTriangulator::FillMeshNodes()
 		}
 	}
 
-	int32 StartId = Mesh->RegisterCoordinates();
+	int32 StartId = Mesh.RegisterCoordinates();
 	for (FIsoInnerNode& Node : InnerNodes)
 	{
 		Node.OffsetId(StartId);
 	}
 
-	Mesh->VerticesGlobalIndex.SetNum(InnerNodeCount + LoopNodeCount);
+	Mesh.VerticesGlobalIndex.SetNum(InnerNodeCount + LoopNodeCount);
 	int32 Index = 0;
 	for (FLoopNode& Node : LoopNodes)
 	{
-		Mesh->VerticesGlobalIndex[Index++] = Node.GetId();
+		Mesh.VerticesGlobalIndex[Index++] = Node.GetNodeId();
 	}
 
 	for (FIsoInnerNode& Node : InnerNodes)
 	{
-		Mesh->VerticesGlobalIndex[Index++] = Node.GetId();
+		Mesh.VerticesGlobalIndex[Index++] = Node.GetNodeId();
 	}
 
 	for (FLoopNode& Node : LoopNodes)
 	{
-		Mesh->Normals.Emplace(Node.GetNormal(Grid));
+		Mesh.Normals.Emplace(Node.GetNormal(Grid));
 	}
 
 	for (FIsoInnerNode& Node : InnerNodes)
 	{
-		Mesh->Normals.Emplace(Node.GetNormal(Grid));
+		Mesh.Normals.Emplace(Node.GetNormal(Grid));
 	}
 
 	for (FLoopNode& Node : LoopNodes)
 	{
 		const FPoint2D& UVCoordinate = Node.Get2DPoint(EGridSpace::Scaled, Grid);
-		Mesh->UVMap.Emplace(UVCoordinate.U, UVCoordinate.V);
+		Mesh.UVMap.Emplace(UVCoordinate.U, UVCoordinate.V);
 	}
 
 	for (FIsoInnerNode& Node : InnerNodes)
 	{
 		const FPoint2D& UVCoordinate = Node.Get2DPoint(EGridSpace::Scaled, Grid);
-		Mesh->UVMap.Emplace(UVCoordinate.U, UVCoordinate.V);
+		Mesh.UVMap.Emplace(UVCoordinate.U, UVCoordinate.V);
 	}
 }
 
-namespace IsoTriangulatorImpl
-{
-const double MaxSlopeToBeIso = 0.125;
-
-const double LimitValueMin(double Slope)
-{
-	return Slope - MaxSlopeToBeIso;
-}
-
-const double LimitValueMax(double Slope)
-{
-	return Slope + MaxSlopeToBeIso;
-}
-}
-
 void FIsoTriangulator::BuildLoopSegments()
 {
 	FTimePoint StartTime = FChrono::Now();
@@ -371,10 +384,8 @@ void FIsoTriangulator::BuildLoopSegments()
 #endif
 }
 
-void FIsoTriangulator::BuildThinZoneSegments()
+void FIsoTriangulator::GetThinZonesMesh()
 {
-	ThinZoneSegments.Reserve((int32)(0.6 * (double)LoopNodeCount));
-
 	TMap<int32, FLoopNode*> IndexToNode;
 	for (FLoopNode& Node : LoopNodes)
 	{
@@ -383,10 +394,49 @@ void FIsoTriangulator::BuildThinZoneSegments()
 			continue;
 		}
 
-		IndexToNode.Add(Node.GetFaceIndex(), &Node);
+		IndexToNode.Add(Node.GetNodeId(), &Node);
+	}
+
+	{
+		for (const FThinZone2D& ThinZone : Grid.GetFace().GetThinZones())
+		{
+			GetThinZoneMesh(IndexToNode, ThinZone);
+		}
 	}
 
-	TFunction<void(FIsoNode*, FIsoNode*)> AddSegment = [this](FIsoNode* NodeA, FIsoNode* NodeB)
+	ThinZoneIntersectionTool.Empty(0);
+}
+
+void FIsoTriangulator::GetThinZoneMesh(const TMap<int32, FLoopNode*>& IndexToNode, const FThinZone2D& ThinZone)
+{
+	using namespace IsoTriangulatorImpl;
+
+	TArray<TPair<int32, FPairOfIndex>> CrossZoneElements;
+	TArray<FCandidateSegment> MeshOfThinZones;
+
+	FAddMeshNodeFunc AddElement = [&CrossZoneElements](const int32 NodeIndice, const FPoint2D& MeshNode2D, double MeshingTolerance3D, const FEdgeSegment& EdgeSegment, const FPairOfIndex& OppositeNodeIndices)
+	{
+		if (CrossZoneElements.Num() && CrossZoneElements.Last().Key == NodeIndice)
+		{
+			CrossZoneElements.Last().Value.Add(OppositeNodeIndices);
+		}
+		else
+		{
+			CrossZoneElements.Emplace(NodeIndice, OppositeNodeIndices);
+		}
+	};
+
+	FReserveContainerFunc Reserve = [&CrossZoneElements](int32 MeshVertexCount)
+	{
+		CrossZoneElements.Reserve(CrossZoneElements.Num() + MeshVertexCount);
+	};
+
+	ThinZone.GetFirstSide().GetExistingMeshNodes(Grid.GetFace(), Mesh.GetMeshModel(), Reserve, AddElement, /*bWithTolerance*/ false);
+	ThinZone.GetSecondSide().GetExistingMeshNodes(Grid.GetFace(), Mesh.GetMeshModel(), Reserve, AddElement, /*bWithTolerance*/ false);
+
+	MeshOfThinZones.Reserve(CrossZoneElements.Num() * 2);
+
+	TFunction<void(FLoopNode*, FLoopNode*)> AddSegmentFromNode = [&MeshOfThinZones, this](FLoopNode* NodeA, FLoopNode* NodeB)
 	{
 		if (!NodeA)
 		{
@@ -398,56 +448,101 @@ void FIsoTriangulator::BuildThinZoneSegments()
 			return;
 		}
 
+		if (&NodeA->GetPreviousNode() == NodeB || &NodeB->GetNextNode() == NodeA)
+		{
+			return;
+		}
+
 		if (NodeA->GetSegmentConnectedTo(NodeB))
 		{
 			return;
 		}
 
-		FIsoSegment& Segment = IsoSegmentFactory.New();
-		Segment.Init(*NodeA, *NodeB, ESegmentType::ThinZone);
-		Segment.ConnectToNode();
-		ThinZoneSegments.Add(&Segment);
+		const FPoint2D& CoordinateA = NodeA->Get2DPoint(EGridSpace::UniformScaled, Grid);
+		const FPoint2D& CoordinateB = NodeB->Get2DPoint(EGridSpace::UniformScaled, Grid);
+
+		// Is Outside and not too flat at Node1
+		const double FlatAngle = 0.1;
+		if (NodeA->IsSegmentBeInsideFace(CoordinateB, Grid, FlatAngle))
+		{
+			return;
+		}
+
+		// Is Outside and not too flat at Node2
+		if (NodeB->IsSegmentBeInsideFace(CoordinateA, Grid, FlatAngle))
+		{
+			return;
+		}
+
+		MeshOfThinZones.Emplace(Grid, *NodeA, *NodeB);
 	};
 
-	for (const TSharedPtr<FTopologicalLoop>& Loop : Grid.GetFace().GetLoops())
+	TFunction<void(int32, int32)> AddSegment = [&IndexToNode, AddSegmentFromNode](int32 IndexNodeA, int32 IndexNodeB)
 	{
-		for (const FOrientedEdge& OrientedEdge : Loop->GetEdges())
+		if (IndexNodeA < 0 || IndexNodeB < 0)
 		{
-			const TSharedPtr<FTopologicalEdge>& Edge = OrientedEdge.Entity;
-			if (!Edge->IsThinZone())
-			{
-				continue;
-			}
+			return;
+		}
 
-			const TArray<FCuttingPoint>& CuttingPoints = Edge->GetLinkActiveEdge()->GetCuttingPoints();
-			const TArray<int32>& NodeIds = Edge->GetMesh()->EdgeVerticesIndex;
-			for (int32 Index = 0; Index < NodeIds.Num(); ++Index)
-			{
-				if (CuttingPoints[Index].OppositNodeIndex > 0)
-				{
-					AddSegment(IndexToNode[NodeIds[Index]], IndexToNode[CuttingPoints[Index].OppositNodeIndex]);
-				}
-				if (CuttingPoints[Index].OppositNodeIndex2 > 0)
-				{
-					AddSegment(IndexToNode[NodeIds[Index]], IndexToNode[CuttingPoints[Index].OppositNodeIndex2]);
-				}
-			}
+		if (IndexNodeA == IndexNodeB)
+		{
+			return;
+		}
+
+		FLoopNode* const* NodeA = IndexToNode.Find(IndexNodeA);
+		FLoopNode* const* NodeB = IndexToNode.Find(IndexNodeB);
+		if (NodeA && NodeB)
+		{
+			AddSegmentFromNode(*NodeA, *NodeB);
 		}
+	};
+
+	for (const TPair<int32, FPairOfIndex>& CrossZoneElement : CrossZoneElements)
+	{
+		AddSegment(CrossZoneElement.Key, CrossZoneElement.Value[0]);
+		AddSegment(CrossZoneElement.Key, CrossZoneElement.Value[1]);
 	}
 
-	LoopSegmentsIntersectionTool.AddSegments(ThinZoneSegments);
-	LoopSegmentsIntersectionTool.Sort();
+	Algo::Sort(MeshOfThinZones, [](const FCandidateSegment& SegmentA, const FCandidateSegment& SegmentB) { return SegmentA.Length < SegmentB.Length; });
+
+	ThinZoneIntersectionTool.Reserve(ThinZoneIntersectionTool.Count() + MeshOfThinZones.Num());
+
+	for (FCandidateSegment& CandidateSegment : MeshOfThinZones)
+	{
+		if (FIsoSegment::IsItAlreadyDefined(&CandidateSegment.StartNode, &CandidateSegment.EndNode))
+		{
+			continue;
+		}
+
+		if (LoopSegmentsIntersectionTool.DoesIntersect(CandidateSegment.StartNode, CandidateSegment.EndNode))
+		{
+			continue;
+		}
+
+		if (ThinZoneIntersectionTool.DoesIntersect(CandidateSegment.StartNode, CandidateSegment.EndNode))
+		{
+			continue;
+		}
+
+		FIsoSegment& Segment = IsoSegmentFactory.New();
+		Segment.Init(CandidateSegment.StartNode, CandidateSegment.EndNode, ESegmentType::ThinZone);
+
+		CandidateSegment.StartNode.SetThinZoneNodeMarker();
+		CandidateSegment.EndNode.SetThinZoneNodeMarker();
+
+		Segment.ConnectToNode();
+		ThinZoneSegments.Add(&Segment);
+		ThinZoneIntersectionTool.AddSegment(Segment);
+	}
 }
 
 void FIsoTriangulator::BuildInnerSegments()
 {
-
 #ifdef DEBUG_BUILDINNERSEGMENTS
 	if (bDisplay)
 	{
 		Grid.DisplayInnerPoints(TEXT("BuildInnerSegments::Points"), EGridSpace::UniformScaled);
 	}
-
 	F3DDebugSession _(bDisplay, TEXT("BuildInnerSegments"));
 #endif
 
@@ -460,18 +555,34 @@ void FIsoTriangulator::BuildInnerSegments()
 
 	LoopSegmentsIntersectionTool.Reserve(InnerSegmentsIntersectionTool.Count());
 
-	TFunction<bool(const FPoint2D&, const FPoint2D&, const ESegmentType)> AlmostHitsLoop = [&](const FPoint2D& Node1, const FPoint2D& Node2, const ESegmentType InType) -> bool
+
+	// Check if the loop tangents the grid between Node1 and Node 2
+	//                            
+	//                       \   /  Loop		                       \     /  Loop
+	//                        \./ 				                        \   / 
+	//        Node1 *------------------* Node2 	        Node1 *----------\./-------* Node2 
+	//                                                                       
+	//
+	TFunction<bool(const FPoint2D&, const FPoint2D&, const ESegmentType, const double)> AlmostHitsLoop = [&](const FPoint2D& Node1, const FPoint2D& Node2, const ESegmentType InType, const double Tolerance) -> bool
 	{
-		if (InType == ESegmentType::IsoU)
+		if (InType == ESegmentType::IsoV)
 		{
 			for (const TArray<FPoint2D>& Loop : Grid.GetLoops2D(EGridSpace::UniformScaled))
 			{
 				for (const FPoint2D& LoopPoint : Loop)
 				{
-					if (FMath::IsNearlyEqual(LoopPoint.V, Node1.V, GeometricTolerance))
+					if (FMath::IsNearlyEqual(LoopPoint.V, Node1.V, Tolerance))
 					{
 						if (Node1.U - DOUBLE_SMALL_NUMBER < LoopPoint.U && LoopPoint.U < Node2.U + DOUBLE_SMALL_NUMBER)
 						{
+#ifdef DEBUG_BUILDINNERSEGMENTS
+							if (bDisplay)
+							{
+								F3DDebugSession _(bDisplay, FString::Printf(TEXT("Point")));
+								DisplayPoint2DWithScale(LoopPoint, EVisuProperty::RedPoint);
+								Wait();
+							}
+#endif
 							return true;
 						}
 					}
@@ -484,10 +595,18 @@ void FIsoTriangulator::BuildInnerSegments()
 			{
 				for (const FPoint2D& LoopPoint : Loop)
 				{
-					if (FMath::IsNearlyEqual(LoopPoint.U, Node1.U, GeometricTolerance))
+					if (FMath::IsNearlyEqual(LoopPoint.U, Node1.U, Tolerance))
 					{
 						if (Node1.V - DOUBLE_SMALL_NUMBER < LoopPoint.V && LoopPoint.V < Node2.V + DOUBLE_SMALL_NUMBER)
 						{
+#ifdef DEBUG_BUILDINNERSEGMENTS
+							if (bDisplay)
+							{
+								//F3DDebugSession _(bDisplay, FString::Printf(TEXT("Point")));
+								DisplayPoint2DWithScale(LoopPoint, EVisuProperty::RedPoint);
+								Wait();
+							}
+#endif
 							return true;
 						}
 					}
@@ -497,7 +616,7 @@ void FIsoTriangulator::BuildInnerSegments()
 		return false;
 	};
 
-	TFunction<void(const int32, const int32, const ESegmentType)> AddToInnerToOuterSegmentsIntersectionTool = [&](const int32 IndexNode1, const int32 IndexNode2, const ESegmentType InType)
+	TFunction<void(const int32, const int32)> AddToInnerToOuterSegmentsIntersectionTool = [&](const int32 IndexNode1, const int32 IndexNode2)
 	{
 		const FPoint2D& Point1 = Grid.GetInner2DPoint(EGridSpace::UniformScaled, IndexNode1);
 		const FPoint2D& Point2 = Grid.GetInner2DPoint(EGridSpace::UniformScaled, IndexNode2);
@@ -518,7 +637,7 @@ void FIsoTriangulator::BuildInnerSegments()
 		FinalInnerSegments.Add(&Segment);
 	};
 
-	TFunction<void(const int32, const int32, const ESegmentType)> BuildSegmentIfValid = [&](const int32 IndexNode1, const int32 IndexNode2, const ESegmentType InType)
+	TFunction<void(const int32, const int32, const ESegmentType, const double)> BuildSegmentIfValid = [&](const int32 IndexNode1, const int32 IndexNode2, const ESegmentType InType, const double Tolerance)
 	{
 		if (Grid.IsNodeOusideFaceButClose(IndexNode1) && Grid.IsNodeOusideFaceButClose(IndexNode2))
 		{
@@ -528,7 +647,7 @@ void FIsoTriangulator::BuildInnerSegments()
 				DisplaySegmentWithScale(Grid.GetInner2DPoint(EGridSpace::UniformScaled, IndexNode1), Grid.GetInner2DPoint(EGridSpace::UniformScaled, IndexNode2), IndexNode1, EVisuProperty::PinkCurve);
 			}
 #endif
-			AddToInnerToOuterSegmentsIntersectionTool(IndexNode1, IndexNode2, InType);
+			AddToInnerToOuterSegmentsIntersectionTool(IndexNode1, IndexNode2);
 			return;
 		}
 
@@ -537,7 +656,6 @@ void FIsoTriangulator::BuildInnerSegments()
 #ifdef DEBUG_BUILDINNERSEGMENTS
 			if (bDisplay)
 			{
-				F3DDebugSession _(bDisplay, TEXT("Segment"));
 				DisplaySegmentWithScale(Grid.GetInner2DPoint(EGridSpace::UniformScaled, IndexNode1), Grid.GetInner2DPoint(EGridSpace::UniformScaled, IndexNode2), IndexNode1, EVisuProperty::RedCurve);
 			}
 #endif
@@ -546,19 +664,26 @@ void FIsoTriangulator::BuildInnerSegments()
 
 		if (Grid.IsNodeInsideAndCloseToLoop(IndexNode1) && Grid.IsNodeInsideAndCloseToLoop(IndexNode2))
 		{
-#ifdef DEBUG_BUILDINNERSEGMENTS
-			if (bDisplay)
-			{
-				DisplaySegmentWithScale(Grid.GetInner2DPoint(EGridSpace::UniformScaled, IndexNode1), Grid.GetInner2DPoint(EGridSpace::UniformScaled, IndexNode2), IndexNode1, EVisuProperty::PinkCurve);
-			}
-#endif
 			if (LoopSegmentsIntersectionTool.DoesIntersect(Grid.GetInner2DPoint(EGridSpace::UniformScaled, IndexNode1), Grid.GetInner2DPoint(EGridSpace::UniformScaled, IndexNode2))
-				|| AlmostHitsLoop(Grid.GetInner2DPoint(EGridSpace::UniformScaled, IndexNode1), Grid.GetInner2DPoint(EGridSpace::UniformScaled, IndexNode2), InType))
+				|| AlmostHitsLoop(Grid.GetInner2DPoint(EGridSpace::UniformScaled, IndexNode1), Grid.GetInner2DPoint(EGridSpace::UniformScaled, IndexNode2), InType, Tolerance))
 			{
-				AddToInnerToOuterSegmentsIntersectionTool(IndexNode1, IndexNode2, InType);
+#ifdef DEBUG_BUILDINNERSEGMENTS
+				if (bDisplay)
+				{
+					DisplaySegmentWithScale(Grid.GetInner2DPoint(EGridSpace::UniformScaled, IndexNode1), Grid.GetInner2DPoint(EGridSpace::UniformScaled, IndexNode2), IndexNode1, EVisuProperty::PurpleCurve);
+					Wait();
+				}
+#endif
+				AddToInnerToOuterSegmentsIntersectionTool(IndexNode1, IndexNode2);
 			}
 			else
 			{
+#ifdef DEBUG_BUILDINNERSEGMENTS
+				if (bDisplay)
+				{
+					DisplaySegmentWithScale(Grid.GetInner2DPoint(EGridSpace::UniformScaled, IndexNode1), Grid.GetInner2DPoint(EGridSpace::UniformScaled, IndexNode2), IndexNode1, EVisuProperty::OrangeCurve);
+				}
+#endif
 				AddToInnerSegments(IndexNode1, IndexNode2, InType);
 			}
 
@@ -594,28 +719,55 @@ void FIsoTriangulator::BuildInnerSegments()
 			DisplaySegmentWithScale(Grid.GetInner2DPoint(EGridSpace::UniformScaled, IndexNode1), Grid.GetInner2DPoint(EGridSpace::UniformScaled, IndexNode2), IndexNode1, EVisuProperty::YellowCurve);
 		}
 #endif
-		AddToInnerToOuterSegmentsIntersectionTool(IndexNode1, IndexNode2, InType);
+		AddToInnerToOuterSegmentsIntersectionTool(IndexNode1, IndexNode2);
+	};
+
+	TFunction<const TArray<double>(const TArray<double>&)> ComputeLocalTolerance = [](const TArray<double>& UniformCutting) -> const TArray<double>
+	{
+		const int32 Num = UniformCutting.Num();
+		TArray<double> TolerancesAlongU;
+		TArray<double> Temp;
+		TolerancesAlongU.Reserve(Num);
+		Temp.Reserve(Num);
+		for (int32 Index = 1; Index < Num; Index++)
+		{
+			Temp.Add((UniformCutting[Index] - UniformCutting[Index - 1]));
+		}
+		TolerancesAlongU.Add(Temp[0] * 0.1);
+		for (int32 Index = 1; Index < Num-1; Index++)
+		{
+			TolerancesAlongU.Add((Temp[Index] + Temp[Index - 1]) * 0.05);
+		}
+		TolerancesAlongU.Add(Temp.Last() * 0.1);
+		return MoveTemp(TolerancesAlongU);
 	};
 
-	for (int32 uIndex = 0; uIndex < NumU; uIndex++)
+	// Process along V
 	{
-		for (int32 vIndex = 0; vIndex < NumV - 1; vIndex++)
+		TArray<double> TolerancesAlongU = ComputeLocalTolerance(Grid.GetUniformCuttingCoordinatesAlongIso(EIso::IsoU));
+		for (int32 UIndex = 0; UIndex < NumU; UIndex++)
 		{
-			BuildSegmentIfValid(Grid.GobalIndex(uIndex, vIndex), Grid.GobalIndex(uIndex, vIndex + 1), ESegmentType::IsoV);
+			for (int32 VIndex = 0; VIndex < NumV - 1; VIndex++)
+			{
+				BuildSegmentIfValid(Grid.GobalIndex(UIndex, VIndex), Grid.GobalIndex(UIndex, VIndex + 1), ESegmentType::IsoU, TolerancesAlongU[UIndex]);
+			}
 		}
 	}
 
-	for (int32 vIndex = 0; vIndex < NumV; vIndex++)
+	// Process along U
 	{
-		for (int32 uIndex = 0; uIndex < NumU - 1; uIndex++)
+		TArray<double> TolerancesAlongV = ComputeLocalTolerance(Grid.GetUniformCuttingCoordinatesAlongIso(EIso::IsoV));
+		for (int32 VIndex = 0; VIndex < NumV; VIndex++)
 		{
-			BuildSegmentIfValid(Grid.GobalIndex(uIndex, vIndex), Grid.GobalIndex(uIndex + 1, vIndex), ESegmentType::IsoU);
+			for (int32 UIndex = 0; UIndex < NumU - 1; UIndex++)
+			{
+				BuildSegmentIfValid(Grid.GobalIndex(UIndex, VIndex), Grid.GobalIndex(UIndex + 1, VIndex), ESegmentType::IsoV, TolerancesAlongV[VIndex]);
+			}
 		}
 	}
 
 #ifdef DEBUG_BUILDINNERSEGMENTS
 	Wait(bDisplay);
-	InnerToOuterSegmentsIntersectionTool.Display(bDisplay, TEXT("InnerToOuterSegmentsIntersectionTool End"), EVisuProperty::RedCurve);
 #endif
 
 	InnerToOuterSegmentsIntersectionTool.Sort();
@@ -1044,6 +1196,10 @@ void FIsoTriangulator::FindSegmentToLinkOuterToInnerLoopNodes(FCell& Cell)
 	{
 		FLoopNode& StartNode = (FLoopNode&)Segment->GetFirstNode();
 		FLoopNode& EndNode = (FLoopNode&)Segment->GetSecondNode();
+		if (StartNode.IsThinZoneNode() || EndNode.IsThinZoneNode())
+		{
+			return;
+		}
 
 		if ((StartNode.GetConnectedSegments().Num() > 2) || (EndNode.GetConnectedSegments().Num() > 2))
 		{
@@ -1099,6 +1255,11 @@ void FIsoTriangulator::FindSegmentToLinkOuterToInnerLoopNodes(FCell& Cell)
 		{
 			for (FLoopNode* Node : Cell.SubLoops[Index])
 			{
+				if (Node->IsThinZoneNode())
+				{
+					continue;
+				}
+
 				// Check if the node is inside the sector (X) or outside (Z)
 				const FPoint2D& NodePoint2D = Node->Get2DPoint(EGridSpace::UniformScaled, Grid);
 				double SlopeAtStartNode = GetSlopeAtStartNode(StartPoint2D, NodePoint2D, StartReferenceSlope);
@@ -1114,7 +1275,7 @@ void FIsoTriangulator::FindSegmentToLinkOuterToInnerLoopNodes(FCell& Cell)
 					continue;
 				}
 
-				double PointCriteria = IsoTriangulatorImpl::IsoscelesCriteria(StartPoint2D, EndPoint2D, NodePoint2D);
+				double PointCriteria = IsoTriangulatorImpl::IsoscelesCriteriaMax(StartPoint2D, EndPoint2D, NodePoint2D);
 
 				if (
 					// the candidate triangle is inside the current candidate triangle
@@ -1167,7 +1328,6 @@ void FIsoTriangulator::FindSegmentToLinkOuterToInnerLoopNodes(FCell& Cell)
 	Cell.SelectSegmentInCandidateSegments(IsoSegmentFactory);
 }
 
-//#define DEBUG_FIND_ISO_SEGMENT_TO_LINK_OUTER_LOOP_NODES
 void FIsoTriangulator::FindIsoSegmentToLinkOuterLoopNodes(FCell& Cell)
 {
 	TArray<FLoopNode*>& OuterLoop = Cell.SubLoops[0];
@@ -1209,7 +1369,6 @@ void FIsoTriangulator::FindIsoSegmentToLinkOuterLoopNodes(FCell& Cell)
 // 	           - build Delaunay connection
 // 	           - find the shortest segment to connect each connected loop by Delaunay
 // =============================================================================================================
-//#define DEBUG_CONNECT_CELL_LOOPS
 void FIsoTriangulator::ConnectCellLoops()
 {
 	TArray<FCell> Cells;
@@ -1266,7 +1425,6 @@ void FIsoTriangulator::ConnectCellLoops()
 	}
 }
 
-//#define DEBUG_BUILD_CELLS
 void FIsoTriangulator::FindCellContainingBoundaryNodes(TArray<FCell>& Cells)
 {
 	FTimePoint StartTime = FChrono::Now();
@@ -1638,695 +1796,93 @@ bool FIsoTriangulator::CanCycleBeMeshed(const TArray<FIsoSegment*>& Cycle, FInte
 	return true;
 }
 
-void FIsoTriangulator::MeshCycle(const EGridSpace Space, const TArray<FIsoSegment*>& Cycle, const TArray<bool>& CycleOrientation)
+void FIsoTriangulator::MeshCycle(const TArray<FIsoSegment*>& Cycle, const TArray<bool>& CycleOrientation)
 {
-	int32 NodeCycleNum = (int32)Cycle.Num();
-
-	if (NodeCycleNum == 4)
+	switch (Cycle.Num())
 	{
-		FIsoNode* Nodes[4];
-		if (CycleOrientation[0])
-		{
-			Nodes[0] = &Cycle[0]->GetFirstNode();
-			Nodes[1] = &Cycle[0]->GetSecondNode();
-		}
-		else
-		{
-			Nodes[0] = &Cycle[0]->GetSecondNode();
-			Nodes[1] = &Cycle[0]->GetFirstNode();
-		}
+	case 2:
+		return;
+	case 3:
+		return MeshCycleOf<3>(Cycle, CycleOrientation, Polygon::MeshTriangle);
+	case 4:
+		return MeshCycleOf<4>(Cycle, CycleOrientation, Polygon::MeshQuadrilateral);
+	case 5:
+		return MeshCycleOf<5>(Cycle, CycleOrientation, Polygon::MeshPentagon);
+	default:
+		MeshLargeCycle(Cycle, CycleOrientation);
+	}
+}
 
-		if (CycleOrientation[2])
-		{
-			Nodes[2] = &Cycle[2]->GetFirstNode();
-			Nodes[3] = &Cycle[2]->GetSecondNode();
-		}
-		else
-		{
-			Nodes[2] = &Cycle[2]->GetSecondNode();
-			Nodes[3] = &Cycle[2]->GetFirstNode();
-		}
+void FIsoTriangulator::MeshLargeCycle(const TArray<FIsoSegment*>& Cycle, const TArray<bool>& CycleOrientation)
+{
+	FCycleTriangulator CycleTriangulator(*this, Cycle, CycleOrientation);
+	CycleTriangulator.MeshCycle();
+}
 
-		const FPoint2D* NodeCoordinates[4];
-		for (int32 Index = 0; Index < 4; ++Index)
-		{
-			NodeCoordinates[Index] = &Nodes[Index]->Get2DPoint(Space, Grid);
-		}
+/**
+ * The purpose is to add surrounding segments to the small loop to intersection tool to prevent traversing inner segments
+ * A loop is inside inner segments
+ *									|			 |
+ *								   -----------------
+ *									|	 XXX	 |
+ *									|	XXXXX	 |
+ *									|	 XXX	 |
+ *								   -----------------
+ *									|			 |
+ *
+ */
+void FIsoTriangulator::FindInnerGridCellSurroundingSmallLoop()
+{
+	FTimePoint StartTime = FChrono::Now();
+
+	if (GlobalIndexToIsoInnerNodes.Num() == 0)
+	{
+		// No inner node
+		return;
+	}
 
-		double SegmentSlopes[4];
-		SegmentSlopes[0] = ComputeSlope(*NodeCoordinates[0], *NodeCoordinates[1]);
-		SegmentSlopes[1] = ComputeSlope(*NodeCoordinates[1], *NodeCoordinates[2]);
-		SegmentSlopes[2] = ComputeSlope(*NodeCoordinates[2], *NodeCoordinates[3]);
-		SegmentSlopes[3] = ComputeSlope(*NodeCoordinates[3], *NodeCoordinates[0]);
+	// when an internal loop is inside inner UV cell
+	int32 NumU = Grid.GetCuttingCount(EIso::IsoU);
+	int32 NumV = Grid.GetCuttingCount(EIso::IsoV);
+	const TArray<double>& UCoordinates = Grid.GetCuttingCoordinatesAlongIso(EIso::IsoU);
+	const TArray<double>& VCoordinates = Grid.GetCuttingCoordinatesAlongIso(EIso::IsoV);
 
-		double RelativeSlopes[4];
-		RelativeSlopes[0] = TransformIntoOrientedSlope(SegmentSlopes[1] - SegmentSlopes[0]);
-		RelativeSlopes[1] = TransformIntoOrientedSlope(SegmentSlopes[2] - SegmentSlopes[1]);
-		RelativeSlopes[2] = TransformIntoOrientedSlope(SegmentSlopes[3] - SegmentSlopes[2]);
-		RelativeSlopes[3] = TransformIntoOrientedSlope(SegmentSlopes[0] - SegmentSlopes[3]);
+	const TArray<TArray<FPoint2D>>& Loops = Grid.GetLoops2D(EGridSpace::Default2D);
+	for (int32 LoopIndex = 1; LoopIndex < Loops.Num(); ++LoopIndex)
+	{
+		FPoint2D FirstPoint = Loops[LoopIndex][0];
 
-		int32 FlattenNodeIndex = 0;
-		for (int32 IndexAngle = 0; IndexAngle < 4; ++IndexAngle)
+		int32 IndexU = 0;
+		for (; IndexU < NumU - 1; ++IndexU)
 		{
-			if (RelativeSlopes[IndexAngle] < RelativeSlopes[FlattenNodeIndex])
+			if ((FirstPoint.U > UCoordinates[IndexU]) && (FirstPoint.U < UCoordinates[IndexU + 1] + DOUBLE_SMALL_NUMBER))
 			{
-				FlattenNodeIndex = IndexAngle;
+				break;
 			}
 		}
 
-		int32 NodeIndices[4];
-		NodeIndices[0] = FlattenNodeIndex;
-		for (int32 IndexN = 1; IndexN < 4; ++IndexN)
+		int32 IndexV = 0;
+		for (; IndexV < NumV - 1; ++IndexV)
 		{
-			NodeIndices[IndexN] = NodeIndices[IndexN - 1] == 3 ? 0 : NodeIndices[IndexN - 1] + 1;
+			if ((FirstPoint.V > VCoordinates[IndexV]) && (FirstPoint.V < VCoordinates[IndexV + 1] + DOUBLE_SMALL_NUMBER))
+			{
+				break;
+			}
 		}
 
-#ifdef ADD_TRIANGLE_2D
-		if (bDisplay)
+		double UMin = UCoordinates[IndexU];
+		double UMax = UCoordinates[IndexU + 1] + DOUBLE_SMALL_NUMBER;
+		double VMin = VCoordinates[IndexV];
+		double VMax = VCoordinates[IndexV + 1] + DOUBLE_SMALL_NUMBER;
+
+		bool bBoudardyIsSurrounded = true;
+		for (const FPoint2D& LoopPoint : Loops[LoopIndex])
 		{
-			F3DDebugSession G(TEXT("Mesh cycle"));
-			Grid.DisplayTriangle(EGridSpace::UniformScaled, *Nodes[NodeIndices[1]], *Nodes[NodeIndices[3]], *Nodes[NodeIndices[0]]);
-			Grid.DisplayTriangle(EGridSpace::UniformScaled, *Nodes[NodeIndices[1]], *Nodes[NodeIndices[2]], *Nodes[NodeIndices[3]]);
-		}
-#endif 
-		Mesh->AddTriangle(Nodes[NodeIndices[1]]->GetFaceIndex(), Nodes[NodeIndices[3]]->GetFaceIndex(), Nodes[NodeIndices[0]]->GetFaceIndex());
-		Mesh->AddTriangle(Nodes[NodeIndices[1]]->GetFaceIndex(), Nodes[NodeIndices[2]]->GetFaceIndex(), Nodes[NodeIndices[3]]->GetFaceIndex());
-
-		return;
-	}
-	else if (NodeCycleNum == 3)
-	{
-		if (CycleOrientation[0])
-		{
-#ifdef ADD_TRIANGLE_2D
-			if (bDisplay)
-			{
-				F3DDebugSession G(TEXT("Mesh cycle"));
-				Grid.DisplayTriangle(EGridSpace::UniformScaled, Cycle[0]->GetFirstNode(), Cycle[0]->GetSecondNode(), CycleOrientation[1] ? Cycle[1]->GetSecondNode() : Cycle[1]->GetFirstNode());
-			}
-#endif 
-			Mesh->AddTriangle(Cycle[0]->GetFirstNode().GetFaceIndex(), Cycle[0]->GetSecondNode().GetFaceIndex(), CycleOrientation[1] ? Cycle[1]->GetSecondNode().GetFaceIndex() : Cycle[1]->GetFirstNode().GetFaceIndex());
-		}
-		else
-		{
-#ifdef ADD_TRIANGLE_2D
-			if (bDisplay)
-			{
-				F3DDebugSession G(TEXT("Mesh cycle"));
-				Grid.DisplayTriangle(EGridSpace::UniformScaled, Cycle[0]->GetSecondNode(), Cycle[0]->GetFirstNode(), CycleOrientation[1] ? Cycle[1]->GetSecondNode() : Cycle[1]->GetFirstNode());
-			}
-#endif 
-			Mesh->AddTriangle(Cycle[0]->GetSecondNode().GetFaceIndex(), Cycle[0]->GetFirstNode().GetFaceIndex(), CycleOrientation[1] ? Cycle[1]->GetSecondNode().GetFaceIndex() : Cycle[1]->GetFirstNode().GetFaceIndex());
-		}
-		return;
-	}
-
-	const double SquareMinSize = FMath::Square(Grid.GetMinElementSize());
-
-	FIntersectionSegmentTool CycleIntersectionTool(Grid);
-	CycleIntersectionTool.Reserve(5 * NodeCycleNum);
-	CycleIntersectionTool.AddSegments(Cycle);
-	CycleIntersectionTool.Sort();
-
-#ifdef ADD_TRIANGLE_2D
-	static int32 CycleId = 0;
-	F3DDebugSession _(bDisplay, FString::Printf(TEXT("Mesh cycle %d"), CycleId++));
-	Grid.DisplayIsoSegments(TEXT("Cycle"), EGridSpace::UniformScaled, Cycle);
-	InnerToOuterSegmentsIntersectionTool.Display(Grid.bDisplay, TEXT("InnerToOuterSegmentsIntersectionTool"), EVisuProperty::RedCurve);
-	Wait(bDisplay);
-#endif
-
-	// Check if the cycle is in self intersecting and fix it. 
-	if (!CanCycleBeMeshed(Cycle, CycleIntersectionTool))
-	{
-#ifdef CADKERNEL_DEV
-		MesherReport->Logs.AddCycleMeshingFailure();
-#endif
-		return;
-	}
-
-	TArray<FIsoNode*> CycleNodes;
-	CycleNodes.Reserve(NodeCycleNum);
-
-	TArray<FIsoSegment*> SegmentStack;
-	SegmentStack.Reserve(5 * NodeCycleNum);
-
-	{
-		using BoolIter = TArray<bool>::RangedForConstIteratorType;
-		using SegmentIter = TArray<FIsoSegment*>::RangedForConstIteratorType;
-
-		TFunction<bool(FIsoSegment*, EIsoSegmentStates)> HasTriangleOn = [](FIsoSegment* Segment, EIsoSegmentStates Side) -> bool
-		{
-			if (Segment->HasTriangleOn(Side))
-			{
-#ifdef CADKERNEL_DEV
-				Wait();
-				ensureCADKernel(false);
-#endif
-				return true;
-			}
-			Segment->SetHasTriangleOn(Side);
-			return false;
-		};
-
-
-		// Get cycle's nodes and set segments as they have a triangle outside the cycle (to don't try to mesh outside the cycle)
-		BoolIter SegmentOrientation = CycleOrientation.begin();
-		for (SegmentIter Segment = Cycle.begin(); Segment != Cycle.end(); ++Segment, ++SegmentOrientation)
-		{
-			if (*SegmentOrientation)
-			{
-				CycleNodes.Add(&(*Segment)->GetFirstNode());
-				if (HasTriangleOn(*Segment, EIsoSegmentStates::RightTriangle))
-				{
-					return;
-				}
-			}
-			else
-			{
-				CycleNodes.Add(&(*Segment)->GetSecondNode());
-				if (HasTriangleOn(*Segment, EIsoSegmentStates::LeftTriangle))
-				{
-					return;
-				}
-			}
-		}
-
-		// If the Segment has 2 adjacent triangles, the segment is a inner cycle segment
-		// It will have triangle in both side
-		//
-		//    X---------------X----------------X      X---------------X----------------X
-		//    |                                |      |                                |  
-		//    |         X--------------------X |      |                                |  
-		//    |         |                    | |      |                                |  
-		//    X---------X  <- inner segment  | |  or  X---------X  <- inner segment    |
-		//    |         |                    | |      |                                |  
-		//    |         X--------------------X |      |                                |  
-		//    |                                |      |                                |
-		//    X---------------X----------------X      X---------------X----------------X
-		//
-		for (FIsoSegment* Segment : Cycle)
-		{
-			if (Segment->HasTriangleOnRightAndLeft())
-			{
-				Segment->ResetHasTriangle();
-			}
-
-			if (Segment->GetFirstNode().GetConnectedSegments().Num() == 1 ||
-				Segment->GetSecondNode().GetConnectedSegments().Num() == 1)
-			{
-				Segment->ResetHasTriangle();
-			}
-		}
-
-		NodeCycleNum = (int32)Cycle.Num();
-
-		TArray<int32> NodeIndex;
-		NodeIndex.Reserve(NodeCycleNum);
-		for (int32 Index = 0; Index < NodeCycleNum; ++Index)
-		{
-			NodeIndex.Add(Index);
-		}
-
-		TArray<double> SegmentLengths;
-		SegmentLengths.Reserve(NodeCycleNum);
-		for (int32 Index = 0, NextIndex = 1; Index < NodeCycleNum; ++Index, ++NextIndex)
-		{
-			if (NextIndex == NodeCycleNum)
-			{
-				NextIndex = 0;
-			}
-
-			double Length = CycleNodes[Index]->Get3DPoint(Grid).SquareDistance(CycleNodes[NextIndex]->Get3DPoint(Grid));
-			if (Length < SquareMinSize)
-			{
-				Cycle[Index]->SetAsDegenerated();
-			}
-			SegmentLengths.Add(Length);
-		}
-
-		// Sort the nodes to process segments from the longest to the shortest 
-		NodeIndex.Sort([SegmentLengths](const int32& Index1, const int32& Index2) { return SegmentLengths[Index1] > SegmentLengths[Index2]; });
-
-		for (int32 Index = 0; Index < NodeCycleNum; ++Index)
-		{
-			SegmentStack.Add(Cycle[NodeIndex[Index]]);
-		}
-	}
-
-	// Function used in FindBestTriangle
-	TFunction<bool(FIsoNode*, FIsoNode*, FIsoSegment*)> BuildSegmentIfNeeded = [&](FIsoNode* NodeA, FIsoNode* NodeB, FIsoSegment* ABSegment) -> bool
-	{
-		if (ABSegment)
-		{
-			if (&ABSegment->GetFirstNode() == NodeA)
-			{
-				if (ABSegment->HasTriangleOnLeft())
-				{
-#ifdef ADD_TRIANGLE_2D
-					F3DDebugSession _(FString::Printf(TEXT("Segment")));
-					Grid.DisplayIsoSegment(EGridSpace::UniformScaled, *NodeA, *NodeB, 0, EVisuProperty::YellowCurve);
-					Wait();
-#endif
-					return false;
-				}
-				ABSegment->SetHasTriangleOnLeft();
-			}
-			else
-			{
-				if (ABSegment->HasTriangleOnRight())
-				{
-#ifdef ADD_TRIANGLE_2D
-					F3DDebugSession _(FString::Printf(TEXT("Segment")));
-					Grid.DisplayIsoSegment(EGridSpace::UniformScaled, *NodeA, *NodeB, 0, EVisuProperty::YellowCurve);
-					Wait();
-#endif
-					return false;
-				}
-				ABSegment->SetHasTriangleOnRight();
-			}
-		}
-		else
-		{
-			FIsoSegment& NewSegment = IsoSegmentFactory.New();
-			NewSegment.Init(*NodeA, *NodeB, ESegmentType::Unknown);
-			NewSegment.ConnectToNode();
-			CycleIntersectionTool.AddSegment(NewSegment);
-			NewSegment.SetHasTriangleOnLeft();
-			SegmentStack.Add(&NewSegment);
-		}
-		return true;
-	};
-
-	const SlopeMethod GetSlopeAtStartNode = ClockwiseSlope;
-	const SlopeMethod GetSlopeAtEndNode = CounterClockwiseSlope;
-
-	TFunction<void(FIsoSegment*, bool)> FindBestTriangle = [&](FIsoSegment* Segment, bool bOrientation)
-	{
-		// StartNode = A
-		FIsoNode& StartNode = bOrientation ? Segment->GetFirstNode() : Segment->GetSecondNode();
-		// EndNode = B
-		FIsoNode& EndNode = bOrientation ? Segment->GetSecondNode() : Segment->GetFirstNode();
-
-
-		//
-		// For each extremity (A, B) of a segment, in the existing connected segments, the segment with the smallest relative slop is identified ([A, X0] and [B, Xn]).
-		// These segments define the sector in which the best triangle could be.
-		// The triangle to build is the best triangle (according to the Isosceles Criteria) connecting the Segment to one of the allowed nodes (X) between X0 and Xn.
-		// Allowed nodes (X) are in the sector, Disallowed nodes (Z) are outside the sector	
-		//
-		//                                       ------Z------X0-------X------X-----X-------Xn----Z-----Z---
-		//                                                     \                           /
-		//                                                      \    Allowed triangles    /   
-		//                                Not allowed triangles  \                       /   Not allowed triangles
-		//                                                        \                     /
-		//                                             ----Z-------A------Segment------B------Z---
-		//
-		//                                                         Not allowed triangles
-		//
-		// These computations are done in the UniformScaled space to avoid numerical error due to length distortion between U or V space and U or V Length.
-		// i.e. if:
-		// (UMax - UMin) / (VMax - VMin) is big 
-		// and 
-		// "medium length along U" / "medium length along V" is small 
-		// The computed angles or slot is not representative of the 3D space.
-		//
-		// The computation is not done is Scale space to don't have problem with degenerated segments
-		//
-		// To avoid flat triangle, a candidate point must defined a minimal slop with [A, X0] or [B, Xn] to not be aligned with one of them. 
-		//
-		// TODO ? Need to add test to check if the candidate node will not generate flat triangle 
-
-		// PreviousSegment = [A, X0]
-		FIsoSegment* PreviousSegment = FindNextSegment(EGridSpace::UniformScaled, Segment, &StartNode, GetSlopeAtStartNode);
-		// NextSegment = [B, Xn]
-		FIsoSegment* NextSegment = FindNextSegment(EGridSpace::UniformScaled, Segment, &EndNode, GetSlopeAtEndNode);
-
-		// PreviousNode = X0
-		FIsoNode& PreviousNode = PreviousSegment->GetOtherNode(&StartNode);
-		// NextNode = Xn
-		FIsoNode& NextNode = NextSegment->GetOtherNode(&EndNode);
-
-#ifdef DEBUG_FIND_BEST_TRIANGLE
-		static int32 TriangleIndex = 0;
-		if (bDisplay)
-		{
-			{
-				TriangleIndex++;
-				F3DDebugSession _(FString::Printf(TEXT("Start Segment %d %d"), TriangleIndex, bOrientation));
-				Grid.DisplayIsoSegment(EGridSpace::UniformScaled, *Segment);
-				Grid.DisplayIsoNode(EGridSpace::UniformScaled, StartNode, 0, EVisuProperty::RedPoint);
-				Grid.DisplayIsoNode(EGridSpace::UniformScaled, EndNode);
-			}
-
-			{
-				F3DDebugSession _(TEXT("Next Segments"));
-				Grid.DisplayIsoSegment(EGridSpace::UniformScaled, *PreviousSegment);
-				Grid.DisplayIsoSegment(EGridSpace::UniformScaled, *NextSegment, 0, EVisuProperty::BlueCurve);
-				Grid.DisplayIsoNode(EGridSpace::UniformScaled, EndNode, 0, EVisuProperty::RedPoint);
-				Grid.DisplayIsoNode(EGridSpace::UniformScaled, NextNode);
-				Wait(TriangleIndex == 1000);
-			}
-		}
-#endif
-
-		FIsoNode* CandidatNode = nullptr;
-		FIsoSegment* StartToCandiatSegment = nullptr;
-		FIsoSegment* EndToCandiatSegment = nullptr;
-
-		if (!NextSegment->IsDegenerated() && !PreviousSegment->IsDegenerated())
-		{
-			const FPoint2D& StartPoint2D = StartNode.Get2DPoint(EGridSpace::UniformScaled, Grid);
-			const FPoint2D& EndPoint2D = EndNode.Get2DPoint(EGridSpace::UniformScaled, Grid);
-
-			const FPoint& StartPoint3D = StartNode.Get3DPoint(Grid);
-			const FPoint& EndPoint3D = EndNode.Get3DPoint(Grid);
-
-			// StartMaxSlope and EndMaxSlope are at most equal to 4, because if the slop with candidate node is biggest to 4, the nez triangle will be inverted
-			double StartReferenceSlope = ComputePositiveSlope(StartPoint2D, EndPoint2D, 0);
-			double StartMaxSlope = GetSlopeAtStartNode(StartPoint2D, PreviousNode.Get2DPoint(EGridSpace::UniformScaled, Grid), StartReferenceSlope);
-			if (&EndNode != &PreviousNode)
-			{
-				// Case of probable auto-intersection cycle at PreviousNode, cancel is preferred 
-				if (FMath::IsNearlyEqual(StartMaxSlope, 8., DOUBLE_KINDA_SMALL_NUMBER))
-				{
-					return;
-				}
-			}
-			StartMaxSlope = FMath::Min(StartMaxSlope, 4.);
-
-			double EndReferenceSlope = StartReferenceSlope < 4 ? StartReferenceSlope + 4 : StartReferenceSlope - 4;
-			double EndMaxSlope = GetSlopeAtEndNode(EndPoint2D, NextNode.Get2DPoint(EGridSpace::UniformScaled, Grid), EndReferenceSlope);
-			if (&StartNode != &NextNode)
-			{
-				// Case of probable auto-intersection cycle at PreviousNode, cancel is preferred 
-				if (FMath::IsNearlyEqual(EndMaxSlope, 8., DOUBLE_KINDA_SMALL_NUMBER))
-				{
-					return;
-				}
-			}
-			EndMaxSlope = FMath::Min(EndMaxSlope, 4.);
-
-			double MinCriteria = HUGE_VALUE;
-			const double MinSlopeToNotBeAligned = 0.01;
-			const double EpsilonSlope = 0.0001;
-			double CandidateSlopeAtStartNode = 8.;
-			double CandidateSlopeAtEndNode = 8.;
-
-			// To avoid to create triangle crossing natural iso line (a line carrying iso edges)
-			// the iso edges connecting an inside node to an external node (vs the loop) are added to InnerToOuterSegmentsIntersectionTool
-			// the selected triangle is the triangle minimizing the number of intersection with these lines
-			int32 MaxMaxIntersections = 1;
-			for (FIsoNode* Node : CycleNodes)
-			{
-				if (Node->IsALoopNode())
-				{
-					MaxMaxIntersections++;
-				}
-			}
-
-			int32 MaxIntersections = 1;
-			for (int32 IntersectionCountAllow = 0; IntersectionCountAllow < MaxIntersections && IntersectionCountAllow <= MaxMaxIntersections; IntersectionCountAllow++)
-			{
-				for (FIsoNode* Node : CycleNodes)
-				{
-					if (Node == &StartNode)
-					{
-						continue;
-					}
-					if (Node == &EndNode)
-					{
-						continue;
-					}
-
-					// Check if the node is inside the sector (X) or outside (Z)
-					const FPoint2D& NodePoint2D = Node->Get2DPoint(EGridSpace::UniformScaled, Grid);
-					double SlopeAtStartNode = GetSlopeAtStartNode(StartPoint2D, NodePoint2D, StartReferenceSlope);
-					double SlopeAtEndNode = GetSlopeAtEndNode(EndPoint2D, NodePoint2D, EndReferenceSlope);
-
-					if (Node != &PreviousNode)
-					{
-						if (SlopeAtStartNode <= EpsilonSlope || SlopeAtStartNode >= StartMaxSlope - EpsilonSlope)
-						{
-							continue;
-						}
-					}
-
-					if (Node != &NextNode)
-					{
-						if (SlopeAtEndNode <= EpsilonSlope || SlopeAtEndNode >= EndMaxSlope - EpsilonSlope)
-						{
-							continue;
-						}
-					}
-
-#ifdef DEBUG_FIND_BEST_TRIANGLE
-					if (TriangleIndex == 14)
-					{
-						static int32 NodeIndex = 1;
-						F3DDebugSession _(FString::Printf(TEXT("Next Node %d"), NodeIndex);
-						Grid.DisplayIsoNode(EGridSpace::UniformScaled, *Node, ++NodeIndex, EVisuProperty::YellowPoint);
-						Wait();
-					}
-#endif
-
-					if (FMath::IsNearlyEqual(SlopeAtStartNode, CandidateSlopeAtStartNode, MinSlopeToNotBeAligned) && SlopeAtEndNode > CandidateSlopeAtEndNode)
-					{
-						continue;
-					}
-
-					if (FMath::IsNearlyEqual(SlopeAtEndNode, CandidateSlopeAtEndNode, EpsilonSlope) && SlopeAtStartNode > CandidateSlopeAtStartNode)
-					{
-						continue;
-					}
-
-#ifdef D3_COTANGENT_CRITERIA
-				// need more test, need to be tested with IsoscelesCriteria instead of 3D CotangentCriteria
-					const FPoint& NodePoint3D = Node->Get3DPoint(Grid);
-					FPoint NodeNormal;
-					double PointCriteria = FMath::Abs(CotangentCriteria(StartPoint3D, EndPoint3D, NodePoint3D, NodeNormal));
-					double CosAngle = FMath::Abs(ComputeCosinus(NodeNormal, Node->GetNormal(Grid)));
-
-					// the criteria is weighted according to the cosine of the angle between the normal of the candidate triangle and the normal at the tested point
-					if (CosAngle > DOUBLE_SMALL_NUMBER)
-					{
-						PointCriteria /= CosAngle;
-					}
-					else
-					{
-						PointCriteria = HUGE_VALUE;
-					}
-#else 
-					double PointCriteria = IsoTriangulatorImpl::IsoscelesCriteria(StartPoint2D, EndPoint2D, NodePoint2D);
-#endif
-					if (
-						// the candidate triangle is inside the current candidate triangle
-						((SlopeAtStartNode < (CandidateSlopeAtStartNode + MinSlopeToNotBeAligned)) && (SlopeAtEndNode < (CandidateSlopeAtEndNode + MinSlopeToNotBeAligned)))
-						||
-						// the candidate triangle is better the current candidate triangle and doesn't contain the current candidate triangle
-						((PointCriteria < MinCriteria) && ((SlopeAtStartNode > CandidateSlopeAtStartNode) ^ (SlopeAtEndNode > CandidateSlopeAtEndNode))))
-					{
-						// check if the candidate segment is not in intersection with existing segments
-						// if the segment exist, it has already been tested
-						FIsoSegment* StartSegment = StartNode.GetSegmentConnectedTo(Node);
-						FIsoSegment* EndSegment = EndNode.GetSegmentConnectedTo(Node);
-
-						if (!StartSegment && CycleIntersectionTool.DoesIntersect(StartNode, *Node))
-						{
-							continue;
-						}
-
-						if (!EndSegment && CycleIntersectionTool.DoesIntersect(EndNode, *Node))
-						{
-							continue;
-						}
-
-#ifdef DEBUG_FIND_BEST_TRIANGLE
-						if (bDisplay)
-						{
-							F3DDebugSession _(TEXT("Candidate triangle"));
-							Grid.DisplayIsoSegment(EGridSpace::UniformScaled, StartNode, *Node, 1, EVisuProperty::RedCurve);
-							Grid.DisplayIsoSegment(EGridSpace::UniformScaled, EndNode, *Node, 2, EVisuProperty::RedCurve);
-							Wait();
-						}
-#endif
-						if (IntersectionCountAllow < MaxMaxIntersections)
-						{
-							int32 IntersectionCount = InnerToOuterSegmentsIntersectionTool.CountIntersections(StartNode, *Node);
-							IntersectionCount = FMath::Max(IntersectionCount, InnerToOuterSegmentsIntersectionTool.CountIntersections(EndNode, *Node));
-							if (IntersectionCount > IntersectionCountAllow)
-							{
-								if (IntersectionCount >= MaxIntersections)
-								{
-									MaxIntersections = IntersectionCount + 1;
-								}
-								continue;
-							}
-						}
-
-						MinCriteria = PointCriteria;
-						CandidatNode = Node;
-						StartToCandiatSegment = StartSegment;
-						EndToCandiatSegment = EndSegment;
-						CandidateSlopeAtStartNode = SlopeAtStartNode;
-						CandidateSlopeAtEndNode = SlopeAtEndNode;
-					}
-				}
-
-				if (CandidatNode != nullptr)
-				{
-					break;
-				}
-			}
-		}
-
-		if (CandidatNode)
-		{
-			if (bOrientation)
-			{
-				Segment->SetHasTriangleOnRight();
-			}
-			else
-			{
-				Segment->SetHasTriangleOnLeft();
-			}
-
-			if (!BuildSegmentIfNeeded(&StartNode, CandidatNode, StartToCandiatSegment))
-			{
-				return;
-			}
-			if (!BuildSegmentIfNeeded(CandidatNode, &EndNode, EndToCandiatSegment))
-			{
-				return;
-			}
-			Mesh->AddTriangle(EndNode.GetFaceIndex(), StartNode.GetFaceIndex(), CandidatNode->GetFaceIndex());
-
-#ifdef ADD_TRIANGLE_2D
-			if (bDisplay)
-			{
-				{
-					F3DDebugSession _(FString::Printf(TEXT("Triangle")));
-					Grid.DisplayTriangle(EGridSpace::UniformScaled, EndNode, StartNode, *CandidatNode);
-				}
-			}
-#endif 
-
-			if (StartToCandiatSegment == nullptr || EndToCandiatSegment == nullptr)
-			{
-				CycleIntersectionTool.Sort();
-			}
-		}
-							};
-
-	for (int32 Index = 0; Index < SegmentStack.Num(); ++Index)
-	{
-		FIsoSegment* Segment = SegmentStack[Index];
-
-		if (Segment->IsDegenerated())
-		{
-			continue;
-		}
-		if (!Segment->HasTriangleOnLeft())
-		{
-			FindBestTriangle(Segment, false);
-		}
-		if (!Segment->HasTriangleOnRight())
-		{
-			FindBestTriangle(Segment, true);
-		}
-	}
-
-	// in case of incomplete cycle meshing, remove degenerated flags and rerun the process
-	for (FIsoSegment* Segment : SegmentStack)
-	{
-		Segment->ResetDegenerated();
-	}
-
-	for (int32 Index = 0; Index < SegmentStack.Num(); ++Index)
-	{
-		FIsoSegment* Segment = SegmentStack[Index];
-
-		if (!Segment->HasTriangleOnLeft())
-		{
-			FindBestTriangle(Segment, false);
-		}
-		if (!Segment->HasTriangleOnRight())
-		{
-			FindBestTriangle(Segment, true);
-		}
-	}
-
-	// Reset the flags "has triangle" of cycle's segments to avoid to block the meshing of next cycles
-	for (FIsoSegment* Segment : Cycle)
-	{
-		Segment->ResetHasTriangle();
-	}
-}
-
-/**
- * The purpose is to add surrounding segments to the small loop to intersection tool to prevent traversing inner segments
- * A loop is inside inner segments
- *									|			 |
- *								   -----------------
- *									|	 XXX	 |
- *									|	XXXXX	 |
- *									|	 XXX	 |
- *								   -----------------
- *									|			 |
- *
- */
-void FIsoTriangulator::FindInnerGridCellSurroundingSmallLoop()
-{
-	FTimePoint StartTime = FChrono::Now();
-
-	if (GlobalIndexToIsoInnerNodes.Num() == 0)
-	{
-		// No inner node
-		return;
-	}
-
-	// when an internal loop is inside inner UV cell
-	int32 NumU = Grid.GetCuttingCount(EIso::IsoU);
-	int32 NumV = Grid.GetCuttingCount(EIso::IsoV);
-	const TArray<double>& UCoordinates = Grid.GetCuttingCoordinatesAlongIso(EIso::IsoU);
-	const TArray<double>& VCoordinates = Grid.GetCuttingCoordinatesAlongIso(EIso::IsoV);
-
-	const TArray<TArray<FPoint2D>>& Loops = Grid.GetLoops2D(EGridSpace::Default2D);
-	for (int32 LoopIndex = 1; LoopIndex < Loops.Num(); ++LoopIndex)
-	{
-		FPoint2D FirstPoint = Loops[LoopIndex][0];
-
-		int32 IndexU = 0;
-		for (; IndexU < NumU - 1; ++IndexU)
-		{
-			if ((FirstPoint.U > UCoordinates[IndexU]) && (FirstPoint.U < UCoordinates[IndexU + 1] + DOUBLE_SMALL_NUMBER))
-			{
-				break;
-			}
-		}
-
-		int32 IndexV = 0;
-		for (; IndexV < NumV - 1; ++IndexV)
-		{
-			if ((FirstPoint.V > VCoordinates[IndexV]) && (FirstPoint.V < VCoordinates[IndexV + 1] + DOUBLE_SMALL_NUMBER))
-			{
-				break;
-			}
-		}
-
-		double UMin = UCoordinates[IndexU];
-		double UMax = UCoordinates[IndexU + 1] + DOUBLE_SMALL_NUMBER;
-		double VMin = VCoordinates[IndexV];
-		double VMax = VCoordinates[IndexV + 1] + DOUBLE_SMALL_NUMBER;
-
-		bool bBoudardyIsSurrounded = true;
-		for (const FPoint2D& LoopPoint : Loops[LoopIndex])
-		{
-			if (LoopPoint.U < UMin || LoopPoint.U > UMax || LoopPoint.V < VMin || LoopPoint.V > VMax)
-			{
-				bBoudardyIsSurrounded = false;
-				break;
-			}
+			if (LoopPoint.U < UMin || LoopPoint.U > UMax || LoopPoint.V < VMin || LoopPoint.V > VMax)
+			{
+				bBoudardyIsSurrounded = false;
+				break;
+			}
 		}
 
 		if (bBoudardyIsSurrounded)
@@ -2388,9 +1944,12 @@ void FIsoTriangulator::FindInnerGridCellSurroundingSmallLoop()
 #endif
 }
 
-//#define FIND_CYCLE
 void FIsoTriangulator::TriangulateOverCycle(const EGridSpace Space)
 {
+#ifdef ADD_TRIANGLE_2D
+	F3DDebugSession G(bDisplay, ("Triangulate cycles"));
+#endif 
+
 	FTimePoint StartTime = FChrono::Now();
 
 	TArray<FIsoSegment*> Cycle;
@@ -2398,8 +1957,8 @@ void FIsoTriangulator::TriangulateOverCycle(const EGridSpace Space)
 	TArray<bool> CycleOrientation;
 	CycleOrientation.Reserve(100);
 
-	int32 CycleIndex = 0;
 #ifdef FIND_CYCLE
+	int32 OverCycleIndex = 0;
 	if (Grid.GetFace()->GetId() == FaceToDebug)
 	{
 		Open3DDebugSession(TEXT("Triangulate Over Cycle"));
@@ -2418,11 +1977,11 @@ void FIsoTriangulator::TriangulateOverCycle(const EGridSpace Space)
 			if (Grid.GetFace()->GetId() == FaceToDebug)
 			{
 				F3DDebugSession G(TEXT("Find & mesh cycles"));
-				FString Message = FString::Printf(TEXT("MeshCycle - cycle %d"), CycleIndex++);
+				FString Message = FString::Printf(TEXT("MeshCycle - cycle %d"), OverCycleIndex++);
 				Display(EGridSpace::UniformScaled, *Message, Cycle, false);
 			}
 #endif
-			MeshCycle(Space, Cycle, CycleOrientation);
+			MeshCycle(Cycle, CycleOrientation);
 		}
 	}
 
@@ -2441,12 +2000,12 @@ void FIsoTriangulator::TriangulateOverCycle(const EGridSpace Space)
 			if (Grid.GetFace()->GetId() == FaceToDebug)
 			{
 				Open3DDebugSession(TEXT("Find & mesh cycles"));
-				FString Message = FString::Printf(TEXT("MeshCycle - cycle %d"), CycleIndex++);
+				FString Message = FString::Printf(TEXT("MeshCycle - cycle %d"), OverCycleIndex++);
 				Display(EGridSpace::UniformScaled, *Message, Cycle, false);
 				Close3DDebugSession();
 			}
 #endif
-			MeshCycle(Space, Cycle, CycleOrientation);
+			MeshCycle(Cycle, CycleOrientation);
 		}
 
 		if (!Segment->HasCycleOnRight())
@@ -2462,12 +2021,12 @@ void FIsoTriangulator::TriangulateOverCycle(const EGridSpace Space)
 			if (Grid.GetFace()->GetId() == FaceToDebug)
 			{
 				Open3DDebugSession(TEXT("Find & mesh cycles"));
-				FString Message = FString::Printf(TEXT("MeshCycle - cycle %d"), CycleIndex++);
+				FString Message = FString::Printf(TEXT("MeshCycle - cycle %d"), OverCycleIndex++);
 				Display(EGridSpace::UniformScaled, *Message, Cycle, false);
 				Close3DDebugSession();
 			}
 #endif
-			MeshCycle(Space, Cycle, CycleOrientation);
+			MeshCycle(Cycle, CycleOrientation);
 		}
 	}
 #ifdef FIND_CYCLE
@@ -2479,7 +2038,6 @@ void FIsoTriangulator::TriangulateOverCycle(const EGridSpace Space)
 #endif
 }
 
-//#define DEBUG_FIND_CYCLE
 #ifdef DEBUG_FIND_CYCLE
 static int32 CycleId = -1;
 static int32 CycleIndex = 0;
@@ -2630,28 +2188,11 @@ bool FIsoTriangulator::FindCycle(FIsoSegment* StartSegment, bool LeftSide, TArra
 	return true;
 }
 
-//#define DEBUG_FIND_NEXTSEGMENT
-#ifdef DEBUG_FIND_NEXTSEGMENT
-static bool bDisplayStar = false;
-#endif
-
 FIsoSegment* FIsoTriangulator::FindNextSegment(EGridSpace Space, const FIsoSegment* StartSegment, const FIsoNode* StartNode, SlopeMethod GetSlope) const
 {
 	const FPoint2D& StartPoint = StartNode->Get2DPoint(Space, Grid);
 	const FPoint2D& EndPoint = (StartNode == &StartSegment->GetFirstNode()) ? StartSegment->GetSecondNode().Get2DPoint(Space, Grid) : StartSegment->GetFirstNode().Get2DPoint(Space, Grid);
 
-#ifdef DEBUG_FIND_NEXTSEGMENT
-	F3DDebugSession G(bDisplayStar, TEXT("FindNextSegment"));
-	if (bDisplayStar)
-	{
-		F3DDebugSession _(TEXT("Start Segment"));
-		Display(EGridSpace::Default2D, *StartNode);
-		Display(EGridSpace::Default2D, *StartSegment);
-		Display(Space, *StartNode);
-		Display(Space, *StartSegment);
-	}
-#endif
-
 	double ReferenceSlope = ComputePositiveSlope(StartPoint, EndPoint, 0);
 
 	double MaxSlope = 8.1;
@@ -2662,16 +2203,10 @@ FIsoSegment* FIsoTriangulator::FindNextSegment(EGridSpace Space, const FIsoSegme
 		const FPoint2D& OtherPoint = (StartNode == &Segment->GetFirstNode()) ? Segment->GetSecondNode().Get2DPoint(Space, Grid) : Segment->GetFirstNode().Get2DPoint(Space, Grid);
 
 		double Slope = GetSlope(StartPoint, OtherPoint, ReferenceSlope);
-		if (Slope < SMALL_NUMBER_SQUARE) Slope = 8;
-
-#ifdef DEBUG_FIND_NEXTSEGMENT
-		if (bDisplayStar)
+		if (Slope < SMALL_NUMBER_SQUARE) 
 		{
-			F3DDebugSession _(FString::Printf(TEXT("Segment slop %f"), Slope));
-			Display(EGridSpace::Default2D, *Segment);
-			Display(Space, *Segment);
+			Slope = 8;
 		}
-#endif
 
 		if (Slope < MaxSlope || NextSegment == StartSegment)
 		{
@@ -2744,8 +2279,8 @@ void FIsoTriangulator::TriangulateInnerNodes()
 			Grid.DisplayTriangle(EGridSpace::UniformScaled, *GlobalIndexToIsoInnerNodes[OppositIndex], *GlobalIndexToIsoInnerNodes[OppositIndex - 1], *GlobalIndexToIsoInnerNodes[Index]);
 #endif 
 
-			Mesh->AddTriangle(GlobalIndexToIsoInnerNodes[Index]->GetFaceIndex(), GlobalIndexToIsoInnerNodes[Index + 1]->GetFaceIndex(), GlobalIndexToIsoInnerNodes[OppositIndex]->GetFaceIndex());
-			Mesh->AddTriangle(GlobalIndexToIsoInnerNodes[OppositIndex]->GetFaceIndex(), GlobalIndexToIsoInnerNodes[OppositIndex - 1]->GetFaceIndex(), GlobalIndexToIsoInnerNodes[Index]->GetFaceIndex());
+			Mesh.AddTriangle(GlobalIndexToIsoInnerNodes[Index]->GetFaceIndex(), GlobalIndexToIsoInnerNodes[Index + 1]->GetFaceIndex(), GlobalIndexToIsoInnerNodes[OppositIndex]->GetFaceIndex());
+			Mesh.AddTriangle(GlobalIndexToIsoInnerNodes[OppositIndex]->GetFaceIndex(), GlobalIndexToIsoInnerNodes[OppositIndex - 1]->GetFaceIndex(), GlobalIndexToIsoInnerNodes[Index]->GetFaceIndex());
 		}
 		Index++;
 	}
@@ -2754,7 +2289,6 @@ void FIsoTriangulator::TriangulateInnerNodes()
 #endif
 }
 
-//#define DEBUG_CONNECT_CELL_SUB_LOOPS_BY_NEIGHBORHOOD
 void FIsoTriangulator::ConnectCellSubLoopsByNeighborhood(FCell& Cell)
 {
 	FTimePoint StartTime = FChrono::Now();
@@ -2856,7 +2390,8 @@ void FIsoTriangulator::ConnectCellSubLoopsByNeighborhood(FCell& Cell)
 
 		Cell.SelectSegmentInCandidateSegments(IsoSegmentFactory);
 #ifdef DEBUG_CONNECT_CELL_SUB_LOOPS_BY_NEIGHBORHOOD
-		Display(EGridSpace::UniformScaled, TEXT("Build Segments Connect inner close loops"), Cell.FinalSegments, false, false, EVisuProperty::YellowCurve);
+		Grid.DisplayIsoSegments(TEXT("Build Segments Connect inner close loops"), EGridSpace::UniformScaled, Cell.FinalSegments, false, false, false, EVisuProperty::YellowCurve);
+		Wait();
 #endif
 	}
 
@@ -2888,14 +2423,10 @@ void FIsoTriangulator::ConnectCellSubLoopsByNeighborhood(FCell& Cell)
 
 		Cell.SelectSegmentInCandidateSegments(IsoSegmentFactory);
 #ifdef DEBUG_CONNECT_CELL_SUB_LOOPS_BY_NEIGHBORHOOD
-		Display(EGridSpace::UniformScaled, TEXT("Build Segments Connect inner to outer loop"), Cell.FinalSegments, false, false, EVisuProperty::YellowCurve);
+		Grid.DisplayIsoSegments(TEXT("Build Segments Connect inner to outer loop"), EGridSpace::UniformScaled, Cell.FinalSegments, false, false, false, EVisuProperty::YellowCurve);
+		Wait();
 #endif
 	}
-
-
-#ifdef CADKERNEL_DEV
-	Chronos.FindSegmentToLinkLoopToLoopByDelaunayDuration += FChrono::Elapse(StartTime);
-#endif
 }
 
 void FIsoTriangulator::TryToConnectTwoSubLoopsWithShortestSegment(FCell& Cell, const TArray<FLoopNode*>& LoopA, const TArray<FLoopNode*>& LoopB)
@@ -2907,7 +2438,7 @@ void FIsoTriangulator::TryToConnectTwoSubLoopsWithShortestSegment(FCell& Cell, c
 	for (int32 IndexA = 0; IndexA < LoopA.Num(); ++IndexA)
 	{
 		const FLoopNode* NodeA = LoopA[IndexA];
-		if (NodeA->IsDelete())
+		if (NodeA->IsDeleteOrThinNode())
 		{
 			continue;
 		}
@@ -2916,7 +2447,7 @@ void FIsoTriangulator::TryToConnectTwoSubLoopsWithShortestSegment(FCell& Cell, c
 		for (int32 IndexB = 0; IndexB < LoopB.Num(); ++IndexB)
 		{
 			const FLoopNode* NodeB = LoopB[IndexB];
-			if (NodeB->IsDelete())
+			if (NodeB->IsDeleteOrThinNode())
 			{
 				continue;
 			}
@@ -2943,24 +2474,29 @@ void FIsoTriangulator::TryToConnectTwoSubLoopsWithShortestSegment(FCell& Cell, c
 		if (TryToCreateSegment(Cell, NodeA, ACoordinates, NodeB, BCoordinates, 0.1))
 		{
 			NodeA = &LoopA[MinIndexA]->GetNextNode();
+			if (!NodeA->IsDeleteOrThinNode() && !NodeB->IsDeleteOrThinNode())
+			{
 			if (Cell.Contains(NodeA))
 			{
 				NodeB = LoopB[MinIndexB];
 				TryToCreateSegment(Cell, NodeA, NodeA->Get2DPoint(EGridSpace::UniformScaled, Grid), NodeB, NodeB->Get2DPoint(EGridSpace::UniformScaled, Grid), 0.1);
 			}
+			}
 
 			NodeB = &LoopB[MinIndexB]->GetNextNode();
-			if (Cell.Contains(NodeB))
+			if (!NodeA->IsDeleteOrThinNode() && !NodeB->IsDeleteOrThinNode())
 			{
-				NodeA = LoopA[MinIndexA];
-				TryToCreateSegment(Cell, NodeA, NodeA->Get2DPoint(EGridSpace::UniformScaled, Grid), NodeB, NodeB->Get2DPoint(EGridSpace::UniformScaled, Grid), 0.1);
+				if (Cell.Contains(NodeB))
+				{
+					NodeA = LoopA[MinIndexA];
+					TryToCreateSegment(Cell, NodeA, NodeA->Get2DPoint(EGridSpace::UniformScaled, Grid), NodeB, NodeB->Get2DPoint(EGridSpace::UniformScaled, Grid), 0.1);
+				}
 			}
 		}
 	}
 
 };
 
-//#define DEBUG_TRY_TO_CONNECT
 void FIsoTriangulator::TryToConnectTwoLoopsWithIsocelesTriangle(FCell& Cell, const TArray<FLoopNode*>& LoopA, const TArray<FLoopNode*>& LoopB)
 {
 
@@ -3016,15 +2552,20 @@ void FIsoTriangulator::TryToConnectTwoLoopsWithIsocelesTriangle(FCell& Cell, con
 		double CandidateSlopeAtStartNode = 8.;
 		double CandidateSlopeAtEndNode = 8.;
 
-		for (FIsoNode* Node : LoopB)
+		for (FLoopNode* Node : LoopB)
 		{
+			if (Node->IsDeleteOrThinNode())
+			{
+				continue;
+			}
+
 			// Check if the node is inside the sector (X) or outside (Z)
 			const FPoint2D& NodePoint2D = Node->Get2DPoint(EGridSpace::UniformScaled, Grid);
 
-			double PointCriteria = IsoTriangulatorImpl::IsoscelesCriteria(StartPoint2D, EndPoint2D, NodePoint2D);
+			double PointCriteria = IsoTriangulatorImpl::IsoscelesCriteriaMax(StartPoint2D, EndPoint2D, NodePoint2D);
 
 			// Triangle that are too open (more than rectangle triangle) are not tested 
-			if (PointCriteria > 0.6)
+			if (PointCriteria > Slope::RightSlope)
 			{
 				continue;
 			}
@@ -3032,7 +2573,7 @@ void FIsoTriangulator::TryToConnectTwoLoopsWithIsocelesTriangle(FCell& Cell, con
 			double SlopeAtStartNode = GetSlopeAtStartNode(StartPoint2D, NodePoint2D, StartReferenceSlope);
 			double SlopeAtEndNode = GetSlopeAtEndNode(EndPoint2D, NodePoint2D, EndReferenceSlope);
 
-			// check the side of the candidate point accordint to the segment
+			// check the side of the candidate point according to the segment
 			if (SlopeAtStartNode <= MinSlopeToNotBeAligned)
 			{
 				continue;
@@ -3077,13 +2618,19 @@ void FIsoTriangulator::TryToConnectTwoLoopsWithIsocelesTriangle(FCell& Cell, con
 	{
 		FLoopNode* NodeA1 = LoopA[IndexA];
 		FLoopNode* NodeA2 = LoopA[IndexA + 1];
+
+		if (NodeA1->IsDeleteOrThinNode() || NodeA2->IsDeleteOrThinNode())
+		{
+			continue;
+		}
+
 		const FPoint2D& A1Coordinates = NodeA1->Get2DPoint(EGridSpace::UniformScaled, Grid);
 		const FPoint2D& A2Coordinates = NodeA2->Get2DPoint(EGridSpace::UniformScaled, Grid);
 
 		FIsoSegment* Segment = NodeA1->GetSegmentConnectedTo(NodeA2);
 
 		FIsoNode* Node = FindBestTriangle(Segment);
-		if (Node)
+		if (Node && !Node->IsDeleteOrThinNode())
 		{
 #ifdef DEBUG_TRY_TO_CONNECT
 			if (bDisplay)
@@ -3095,14 +2642,19 @@ void FIsoTriangulator::TryToConnectTwoLoopsWithIsocelesTriangle(FCell& Cell, con
 			}
 #endif 
 			const FPoint2D& NodeCoordinates = Node->Get2DPoint(EGridSpace::UniformScaled, Grid);
-			TryToCreateSegment(Cell, NodeA1, A1Coordinates, Node, NodeCoordinates, 0.1);
-			TryToCreateSegment(Cell, NodeA2, A2Coordinates, Node, NodeCoordinates, 0.1);
+			if(!NodeA1->IsDeleteOrThinNode())
+			{
+				TryToCreateSegment(Cell, NodeA1, A1Coordinates, Node, NodeCoordinates, 0.1);
+			}
+			if (!NodeA2->IsDeleteOrThinNode())
+			{
+				TryToCreateSegment(Cell, NodeA2, A2Coordinates, Node, NodeCoordinates, 0.1);
+			}
 		}
 	}
 
 };
 
-//#define DEBUG_FIND_ISO_SEGMENT_TO_LINK_OUTER_LOOP_NODES
 void FIsoTriangulator::TryToConnectVertexSubLoopWithTheMostIsoSegment(FCell& Cell, const TArray<FLoopNode*>& Loop)
 {
 #ifdef DEBUG_FIND_ISO_SEGMENT_TO_LINK_OUTER_LOOP_NODES
@@ -3125,6 +2677,11 @@ void FIsoTriangulator::TryToConnectVertexSubLoopWithTheMostIsoSegment(FCell& Cel
 		FLoopNode* CandidateB = nullptr;
 
 		FLoopNode* CandidateA = Loop[IndexA];
+		if (CandidateA->IsThinZoneNode())
+		{
+			continue;
+		}
+
 		const FPoint2D& ACoordinates = CandidateA->Get2DPoint(EGridSpace::UniformScaled, Grid);
 
 		FLoopNode* NextA = Loop[IndexA + 1];
@@ -3152,6 +2709,11 @@ void FIsoTriangulator::TryToConnectVertexSubLoopWithTheMostIsoSegment(FCell& Cel
 		for (int32 IndexB = IndexA + 2; IndexB < LoopCount; ++IndexB)
 		{
 			FLoopNode* NodeB = Loop[IndexB];
+			if (NodeB->IsThinZoneNode())
+			{
+				continue;
+			}
+
 			const FPoint2D& BCoordinates = NodeB->Get2DPoint(EGridSpace::UniformScaled, Grid);
 
 			double Slope = ComputeSlopeRelativeToReferenceAxis(ACoordinates, BCoordinates, ReferenceSlope);
@@ -3190,6 +2752,11 @@ void FIsoTriangulator::TryToConnectTwoSubLoopsWithTheMostIsoSegment(FCell& Cell,
 
 	for (FLoopNode* CandidateA : LoopA)
 	{
+		if (CandidateA->IsThinZoneNode())
+		{
+			continue;
+		}
+
 		FLoopNode* CandidateB = nullptr;
 		const FPoint2D& ACoordinates = CandidateA->Get2DPoint(EGridSpace::UniformScaled, Grid);
 
@@ -3198,6 +2765,11 @@ void FIsoTriangulator::TryToConnectTwoSubLoopsWithTheMostIsoSegment(FCell& Cell,
 
 		for (FLoopNode* NodeB : LoopB)
 		{
+			if (NodeB->IsThinZoneNode())
+			{
+				continue;
+			}
+
 			const FPoint2D& BCoordinates = NodeB->Get2DPoint(EGridSpace::UniformScaled, Grid);
 
 			double Slope = ComputeSlopeRelativeToNearestAxis(ACoordinates, BCoordinates);
@@ -3231,7 +2803,6 @@ void FIsoTriangulator::TryToConnectTwoSubLoopsWithTheMostIsoSegment(FCell& Cell,
 	}
 }
 
-//#define DEBUG_TRY_TO_CREATE_SEGMENT
 bool FIsoTriangulator::TryToCreateSegment(FCell& Cell, FLoopNode* NodeA, const FPoint2D& ACoordinates, FIsoNode* NodeB, const FPoint2D& BCoordinates, const double FlatAngle)
 {
 
@@ -3336,6 +2907,11 @@ void FIsoTriangulator::ConnectCellCornerToInnerLoop(FCell& Cell)
 		for (int32 IndexA = 0; IndexA < LoopA.Num(); ++IndexA)
 		{
 			const FLoopNode* NodeA = LoopA[IndexA];
+			if (NodeA->IsThinZoneNode())
+			{
+				continue;
+			}
+
 			const FPoint2D& ACoordinates = NodeA->Get2DPoint(EGridSpace::UniformScaled, Grid);
 
 			double SquareDistance = ACoordinates.SquareDistance(InnerCoordinates);
diff --git a/Engine/Source/Runtime/Datasmith/CADKernel/Private/CADKernel/Mesh/Meshers/IsoTriangulator.h b/Engine/Source/Runtime/Datasmith/CADKernel/Private/CADKernel/Mesh/Meshers/IsoTriangulator.h
--- a/Engine/Source/Runtime/Datasmith/CADKernel/Private/CADKernel/Mesh/Meshers/IsoTriangulator.h
+++ b/Engine/Source/Runtime/Datasmith/CADKernel/Private/CADKernel/Mesh/Meshers/IsoTriangulator.h
@@ -8,6 +8,7 @@
 #include "CADKernel/Mesh/Meshers/IsoTriangulator/IntersectionSegmentTool.h"
 #include "CADKernel/Mesh/Meshers/IsoTriangulator/IsoNode.h"
 #include "CADKernel/Mesh/Meshers/IsoTriangulator/IsoSegment.h"
+#include "CADKernel/Mesh/Meshers/ParametricMesherConstantes.h"
 #include "CADKernel/UI/Visu.h"
 
 #ifdef CADKERNEL_DEV
@@ -15,8 +16,6 @@
 #include "CADKernel/Mesh/Meshers/MesherReport.h"
 #endif
 
-//#define NEED_TO_CHECK_USEFULNESS
-//#define DEBUG_DELAUNAY
 namespace UE::CADKernel
 {
 
@@ -25,13 +24,10 @@ class FIntersectionSegmentTool;
 class FFaceMesh;
 struct FCell;
 
+using FMeshPolygonFunc = TFunction<void(const FGrid&, FIsoNode*[], FFaceMesh&)>;
 
 struct CADKERNEL_API FIsoTriangulatorChronos
 {
-	FDuration TriangulateDuration1 = FChrono::Init();
-	FDuration TriangulateDuration2 = FChrono::Init();
-	FDuration TriangulateDuration3 = FChrono::Init();
-	FDuration TriangulateDuration4 = FChrono::Init();
 	FDuration TriangulateDuration = FChrono::Init();
 	FDuration BuildIsoNodesDuration = FChrono::Init();
 	FDuration BuildLoopSegmentsDuration = FChrono::Init();
@@ -41,7 +37,6 @@ struct CADKERNEL_API FIsoTriangulatorChronos
 	FDuration FindIsoSegmentToLinkInnerToLoopDuration = FChrono::Init();
 	FDuration FindInnerSegmentToLinkLoopToLoopDuration = FChrono::Init();
 	FDuration FindSegmentToLinkLoopToLoopDuration = FChrono::Init();
-	FDuration FindSegmentToLinkLoopToLoopByDelaunayDuration = FChrono::Init();
 	FDuration FindSegmentToLinkInnerToLoopDuration = FChrono::Init();
 	FDuration TriangulateOverCycleDuration = FChrono::Init();
 	FDuration TriangulateInnerNodesDuration = FChrono::Init();
@@ -58,7 +53,6 @@ struct CADKERNEL_API FIsoTriangulatorChronos
 		IsoTriangulerDuration += FindLoopSegmentOfInnerTriangulationDuration;
 		IsoTriangulerDuration += FindIsoSegmentToLinkInnerToLoopDuration;
 		IsoTriangulerDuration += FindSegmentToLinkLoopToLoopDuration;
-		IsoTriangulerDuration += FindSegmentToLinkLoopToLoopByDelaunayDuration;
 		IsoTriangulerDuration += FindSegmentToLinkInnerToLoopDuration;
 		IsoTriangulerDuration += TriangulateOverCycleDuration;
 		IsoTriangulerDuration += TriangulateInnerNodesDuration;
@@ -71,28 +65,23 @@ struct CADKERNEL_API FIsoTriangulatorChronos
 		FChrono::PrintClockElapse(Log, TEXT("    "), TEXT("FindLoopSegmentOfInnerTriangulation"), FindLoopSegmentOfInnerTriangulationDuration);
 		FChrono::PrintClockElapse(Log, TEXT("      "), TEXT("FindSegmentIsoUVSurroundingSmallLoop"), FindSegmentIsoUVSurroundingSmallLoopDuration);
 		FChrono::PrintClockElapse(Log, TEXT("    "), TEXT("Find IsoSegment ToLink InnerToLoop"), FindIsoSegmentToLinkInnerToLoopDuration);
-		FChrono::PrintClockElapse(Log, TEXT("    "), TEXT("Find Segment ToLink LoopToLoop by Delaunay"), FindSegmentToLinkLoopToLoopByDelaunayDuration);
 		FChrono::PrintClockElapse(Log, TEXT("    "), TEXT("Find Segment ToLink LoopToLoop"), FindSegmentToLinkLoopToLoopDuration);
 		FChrono::PrintClockElapse(Log, TEXT("    "), TEXT("Find Segment ToLink InnerToLoop"), FindSegmentToLinkInnerToLoopDuration);
 		FChrono::PrintClockElapse(Log, TEXT("    "), TEXT("Mesh Over Cycle"), TriangulateOverCycleDuration);
 		FChrono::PrintClockElapse(Log, TEXT("    "), TEXT("Mesh Inner Nodes"), TriangulateInnerNodesDuration);
-		FChrono::PrintClockElapse(Log, TEXT("  "), TEXT("Triangulate1"), TriangulateDuration1);
-		FChrono::PrintClockElapse(Log, TEXT("  "), TEXT("Triangulate2"), TriangulateDuration2);
-		FChrono::PrintClockElapse(Log, TEXT("  "), TEXT("Triangulate3"), TriangulateDuration3);
-		FChrono::PrintClockElapse(Log, TEXT("  "), TEXT("Triangulate4"), TriangulateDuration4);
-		FChrono::PrintClockElapse(Log, TEXT("  "), TEXT("Triangulate "), TriangulateDuration);
 	}
 };
 
 class FIsoTriangulator
 {
-	friend class FParametricMesher;
+	friend class FCycleTriangulator;
 	friend class FLoopCleaner;
+	friend class FParametricMesher;
 
 protected:
 
 	FGrid& Grid;
-	TSharedRef<FFaceMesh> Mesh;
+	FFaceMesh& Mesh;
 
 	TArray<int32> LoopStartIndex;
 	TArray<FLoopNode> LoopNodes;
@@ -126,6 +115,8 @@ protected:
 	FIntersectionSegmentTool InnerSegmentsIntersectionTool;
 	FIntersectionSegmentTool InnerToLoopSegmentsIntersectionTool;
 	FIntersectionNodePairTool InnerToOuterSegmentsIntersectionTool;
+	FIntersectionSegmentTool ThinZoneIntersectionTool;
+
 
 	/**
 	 * Define all the lower left index of grid node that the upper cell is surrounding a loop
@@ -152,27 +143,12 @@ protected:
 
 	bool bNeedCheckOrientation = false;
 
-	static const double GeometricToMeshingToleranceFactor;
-	const double GeometricTolerance;
-	const double SquareGeometricTolerance;
-	const double SquareGeometricTolerance2;
-	const double MeshingTolerance;
-	const double SquareMeshingTolerance;
-
-#ifdef CADKERNEL_DEV
-	FMesherReport* MesherReport;
-#endif
-
 public:
+	const FMeshingTolerances& Tolerances;
 
-	FIsoTriangulator(FGrid& InGrid, TSharedRef<FFaceMesh> EntityMesh);
+public:
 
-#ifdef CADKERNEL_DEV
-	void SetMesherReport(FMesherReport& InMesherReport)
-	{
-		MesherReport = &InMesherReport;
-	}
-#endif
+	FIsoTriangulator(FGrid& InGrid, FFaceMesh& OutMesh, const FMeshingTolerances& InTolerance);
 
 	/**
 	 * Main method
@@ -194,7 +170,8 @@ public:
 	 * Add temporary loops defining thin zones to avoid the tessellation of these zone.
 	 * These zones are tessellated in a specific process
 	 */
-	void BuildThinZoneSegments();
+	void GetThinZonesMesh();
+	void GetThinZoneMesh(const TMap<int32, FLoopNode*>& IndexToNode, const FThinZone2D& ThinZone);
 
 	/**
 	 * Fill mesh node data (Position, normal, UV, Index) of the FFaceMesh object
@@ -227,13 +204,6 @@ public:
 	 */
 	void FindInnerGridCellSurroundingSmallLoop();
 
-#ifdef NEED_TO_CHECK_USEFULNESS
-	/**
-	 *
-	 */
-	void CompleteIsoSegmentLoopToLoop();
-#endif
-
 	void ConnectCellLoops();
 	void FindCellContainingBoundaryNodes(TArray<FCell>& Cells);
 
@@ -263,11 +233,6 @@ public:
 
 	void FindCandidateSegmentsToLinkInnerAndLoop();
 
-	/**
-	 * Complete the final set of segments with the best subset of segments in CandidateSegments
-	 */
-	void SelectSegmentInCandidateSegments();
-
 	/**
 	 * The goal of this algorithm is to connect unconnected inner segment extremity i.e. extremity with one or less connected segment to the closed boundary node
 	 */
@@ -290,7 +255,21 @@ public:
 	 * Generate the "Delaunay" tessellation of the cycle.
 	 * The algorithm is based on frontal process
 	 */
-	void MeshCycle(const EGridSpace Space, const TArray<FIsoSegment*>& cycle, const TArray<bool>& cycleOrientation);
+	void MeshCycle(const TArray<FIsoSegment*>& Cycle, const TArray<bool>& CycleOrientation);
+	void MeshLargeCycle(const TArray<FIsoSegment*>& Cycle, const TArray<bool>& CycleOrientation);
+
+	template <uint32 Dim>
+	void MeshCycleOf(const TArray<FIsoSegment*>& Cycle, const TArray<bool>& CycleOrientation, FMeshPolygonFunc MeshPolygonFunc)
+	{
+		FIsoNode* Nodes[Dim];
+		for (int32 Index = 0; Index < Dim; ++Index)
+		{
+			Nodes[Index] = CycleOrientation[Index] ? &Cycle[Index]->GetFirstNode() : &Cycle[Index]->GetSecondNode();
+		}
+
+		MeshPolygonFunc(Grid, Nodes, Mesh);
+	}
+
 
 	bool CanCycleBeMeshed(const TArray<FIsoSegment*>& Cycle, FIntersectionSegmentTool& CycleIntersectionTool);
 
@@ -363,10 +342,12 @@ private:
 	void TryToConnectVertexSubLoopWithTheMostIsoSegment(FCell& Cell, const TArray<FLoopNode*>& SubLoop);
 	bool TryToCreateSegment(FCell& Cell, FLoopNode* NodeA, const FPoint2D& ACoordinates, FIsoNode* NodeB, const FPoint2D& BCoordinates, const double FlatAngle);
 
-#ifdef CADKERNEL_DEV
 public:
+#ifdef CADKERNEL_DEV
 	FIsoTriangulatorChronos Chronos;
+#endif
 
+#ifdef CADKERNEL_DEBUG
 	void DisplayPixels(TArray<uint8>& Pixel) const;
 
 	void DisplayPixel(const int32 Index) const;
@@ -440,11 +421,48 @@ inline double CotangentCriteria(const FPoint2D& APoint, const FPoint2D& BPoint,
 	return ScalareProduct / NormOFPointProduct;
 }
 
-template<class PointType>
-inline double IsoscelesCriteria(const PointType& APoint, const PointType& BPoint, const PointType& IsoscelesVertex)
+struct FPairOfDouble
+{
+	double Value1;
+	double Value2;
+};
+
+inline FPairOfDouble IsoscelesCriteria(const FPoint2D& APoint, const FPoint2D& BPoint, const FPoint2D& CPoint)
+{
+	const double SlopeAB = ComputeSlope(APoint, BPoint);
+	const double SlopeAC = ComputeSlope(APoint, CPoint);
+	const double SlopeBC = ComputeSlope(BPoint, CPoint);
+	const double SlopeBA = SwapSlopeOrientation(SlopeAB);
+
+	return { TransformIntoOrientedSlope(SlopeAC - SlopeAB), TransformIntoOrientedSlope(SlopeBA - SlopeBC) };
+}
+
+inline double IsoscelesCriteriaMax(const FPoint2D& APoint, const FPoint2D& BPoint, const FPoint2D& CPoint)
+{
+	const FPairOfDouble Criteria = IsoscelesCriteria(APoint, BPoint, CPoint);
+	return FMath::Max(Criteria.Value1, Criteria.Value2);
+}
+
+inline double IsoscelesCriteriaMin(const FPoint2D& APoint, const FPoint2D& BPoint, const FPoint2D& CPoint)
+{
+	const FPairOfDouble Criteria = IsoscelesCriteria(APoint, BPoint, CPoint);
+	return FMath::Min(Criteria.Value1, Criteria.Value2);
+}
+
+inline double EquilateralSlopeCriteria(const FPoint2D& APoint, const FPoint2D& BPoint, const FPoint2D& CPoint)
 {
-	double Coord = CoordinateOfProjectedPointOnSegment(IsoscelesVertex, APoint, BPoint, false);
-	return FMath::Abs(Coord - 0.5);
+	const double SlopeAB = ComputeSlope(APoint, BPoint);
+	const double SlopeAC = ComputeSlope(APoint, CPoint);
+	const double SlopeBC = ComputeSlope(BPoint, CPoint);
+	const double SlopeBA = SwapSlopeOrientation(SlopeAB);
+	const double SlopeCA = SwapSlopeOrientation(SlopeAC);
+	const double SlopeCB = SwapSlopeOrientation(SlopeBC);
+
+	const double A = TransformIntoOrientedSlope(SlopeAC - SlopeAB);
+	const double B = TransformIntoOrientedSlope(SlopeBA - SlopeBC);
+	const double C = TransformIntoOrientedSlope(SlopeCB - SlopeCA);
+
+	return FMath::Max3(A, B, C);
 }
 
 template<class PointType>
diff --git a/Engine/Source/Runtime/Datasmith/CADKernel/Private/CADKernel/Mesh/Meshers/IsoTriangulator/IntersectionSegmentTool.cpp b/Engine/Source/Runtime/Datasmith/CADKernel/Private/CADKernel/Mesh/Meshers/IsoTriangulator/IntersectionSegmentTool.cpp
--- a/Engine/Source/Runtime/Datasmith/CADKernel/Private/CADKernel/Mesh/Meshers/IsoTriangulator/IntersectionSegmentTool.cpp
+++ b/Engine/Source/Runtime/Datasmith/CADKernel/Private/CADKernel/Mesh/Meshers/IsoTriangulator/IntersectionSegmentTool.cpp
@@ -22,6 +22,22 @@ FIsoSegment* FIntersectionSegmentTool::FindIntersectingSegment(const FIsoSegment
 	return const_cast<FIsoSegment*> (static_cast<const FIntersectionSegmentTool*>(this)->TIntersectionSegmentTool<IntersectionSegmentTool::FSegment>::FindIntersectingSegment(&Segment.GetFirstNode(), &Segment.GetSecondNode()));
 }
 
+bool FIntersectionSegmentTool::DoesIntersect(const FIsoNode* StartNode, const FIsoNode* EndNode, FIsoSegment** Segment) const
+{
+	if (!StartNode || !EndNode)
+	{
+		return false;
+	}
+
+	*Segment = StartNode->GetSegmentConnectedTo(EndNode);
+	if (*Segment)
+	{
+		return false;
+	}
+	return DoesIntersect(StartNode, EndNode);
+}
+
+
 IntersectionSegmentTool::FSegment::FSegment(const FGrid& Grid, const FIsoNode& StartNode, const FIsoNode& EndNode)
 	: IntersectionToolBase::FSegment(StartNode.Get2DPoint(EGridSpace::UniformScaled, Grid), EndNode.Get2DPoint(EGridSpace::UniformScaled, Grid))
 	, IsoSegment(nullptr)
diff --git a/Engine/Source/Runtime/Datasmith/CADKernel/Private/CADKernel/Mesh/Meshers/IsoTriangulator/IntersectionSegmentTool.h b/Engine/Source/Runtime/Datasmith/CADKernel/Private/CADKernel/Mesh/Meshers/IsoTriangulator/IntersectionSegmentTool.h
--- a/Engine/Source/Runtime/Datasmith/CADKernel/Private/CADKernel/Mesh/Meshers/IsoTriangulator/IntersectionSegmentTool.h
+++ b/Engine/Source/Runtime/Datasmith/CADKernel/Private/CADKernel/Mesh/Meshers/IsoTriangulator/IntersectionSegmentTool.h
@@ -66,8 +66,8 @@ struct FSegment
 
 	EConnectionType IsSuperimposed(const FSegment2D& SegmentAB, const FSegment2D& SegmentCD, bool bSameOrientation) const
 	{
-		const FPoint2D AB = SegmentAB.GetVector();
-		const FPoint2D CD = SegmentCD.GetVector(); 
+		const FPoint2D AB = SegmentAB.GetVector().Normalize();
+		const FPoint2D CD = SegmentCD.GetVector().Normalize();
 		const double ParallelCoef = AB ^ CD;
 		if (FMath::IsNearlyZero(ParallelCoef, DOUBLE_KINDA_SMALL_NUMBER))
 		{
@@ -89,7 +89,7 @@ struct FSegment
 		}
 		if (GetSecondNode() == StartNode)
 		{
-			constexpr bool bNotSameOrientation = true;
+			constexpr bool bNotSameOrientation = false;
 			return IsSuperimposed(Segment2D, InSegment, bNotSameOrientation);
 		}
 
@@ -165,7 +165,7 @@ struct FSegment
 			return false;
 		}
 
-		return IntersectSegments2D(Segment2D, Segment.Segment2D);
+		return DoIntersect(Segment2D, Segment.Segment2D);
 	}
 };
 }
@@ -506,6 +506,7 @@ public:
 		return TIntersectionSegmentTool<IntersectionSegmentTool::FSegment>::FindIntersectingSegment(&StartPoint, &EndPoint) != nullptr;
 	}
 
+	
 	/**
 	 * Allow StartNode to be connected to one segment
 	 * WARNING EndPoint must be defined in EGridSpace::UniformScaled
@@ -515,6 +516,16 @@ public:
 		return TIntersectionSegmentTool<IntersectionSegmentTool::FSegment>::FindIntersectingSegment(&StartNode, &EndPoint) != nullptr;
 	}
 
+	/**
+	 * Find if a Segment already exists, if return false with the segment
+	 */
+	bool DoesIntersect(const FIsoNode* StartNode, const FIsoNode* EndNode, FIsoSegment** Segment) const;
+
+	bool DoesIntersect(const FIsoNode* StartNode, const FIsoNode* EndNode) const
+	{
+		return TIntersectionSegmentTool<IntersectionSegmentTool::FSegment>::FindIntersectingSegment(StartNode, EndNode) != nullptr;
+	}
+
 	bool DoesIntersect(const FIsoNode& StartNode, const FIsoNode& EndNode) const
 	{
 		return TIntersectionSegmentTool<IntersectionSegmentTool::FSegment>::FindIntersectingSegment(&StartNode, &EndNode) != nullptr;
diff --git a/Engine/Source/Runtime/Datasmith/CADKernel/Private/CADKernel/Mesh/Meshers/IsoTriangulator/IsoNode.h b/Engine/Source/Runtime/Datasmith/CADKernel/Private/CADKernel/Mesh/Meshers/IsoTriangulator/IsoNode.h
--- a/Engine/Source/Runtime/Datasmith/CADKernel/Private/CADKernel/Mesh/Meshers/IsoTriangulator/IsoNode.h
+++ b/Engine/Source/Runtime/Datasmith/CADKernel/Private/CADKernel/Mesh/Meshers/IsoTriangulator/IsoNode.h
@@ -30,10 +30,13 @@ enum class EIsoNodeStates : uint16
 	LinkedToPreviousV = 0x0004u,
 	LinkedToNextV = 0x0008u,
 
-	TriangleComplete = 0x000Fu,
+	NodeComplete = 0x000Fu,
 
 	LinkedToLoop = 0x0010u,
 	InnerMeshLoop = 0x0020u,
+	ThinZoneNode = 0x0040u,
+
+	Delete = 0x0080u,
 
 	FirstQuarter = 0x0100u,
 	SecondQuarter = 0x0200u,
@@ -43,8 +46,6 @@ enum class EIsoNodeStates : uint16
 	NearlyIsoU = 0x1000u,  // Loop nodes
 	NearlyIsoV = 0x2000u,  // Loop nodes
 
-	Delete = 0x8000u,
-
 	All = 0xFFFFu
 };
 
@@ -64,7 +65,7 @@ protected:
 
 	int32 Index; // Index of the node either in loop nodes either in inner nodes
 	int32 FaceIndex; // Index of the node in the face
-	int32 Id;
+	int32 NodeId; // Id of the node in the mesh
 
 public:
 	FIsoNode(int32 InNodeIndex, int32 InFaceIndex, int32 InNodeId)
@@ -72,7 +73,7 @@ public:
 		, States(EIsoNodeStates::None)
 		, Index(InNodeIndex)
 		, FaceIndex(InFaceIndex)
-		, Id(InNodeId)
+		, NodeId(InNodeId)
 	{
 		ConnectedSegments.Reserve(5);
 	};
@@ -102,9 +103,9 @@ public:
 		return FaceIndex;
 	}
 
-	const int32 GetId() const
+	const int32 GetNodeId() const
 	{
-		return Id;
+		return NodeId;
 	}
 
 	const TArray<FIsoSegment*>& GetConnectedSegments() const
@@ -126,7 +127,23 @@ public:
 
 	const virtual bool IsALoopNode() const = 0;
 
-	void SetAsLinkedToLoop()
+	void SetThinZoneNodeMarker()
+	{
+		States |= EIsoNodeStates::ThinZoneNode;
+	}
+
+	bool IsThinZoneNode() const
+	{
+		return (States & EIsoNodeStates::ThinZoneNode) == EIsoNodeStates::ThinZoneNode;
+	}
+
+	bool IsDeleteOrThinNode() const
+	{
+		constexpr EIsoNodeStates DeleteOrThinNode = EIsoNodeStates::ThinZoneNode | EIsoNodeStates::Delete;
+		return (States & DeleteOrThinNode) != EIsoNodeStates::None;
+	}
+
+	void SetLinkedToLoopMarker()
 	{
 		States |= EIsoNodeStates::LinkedToLoop;
 	}
@@ -475,7 +492,7 @@ public:
 	/** the node is connected in its 4 directions*/
 	bool IsComplete() const
 	{
-		return (States & EIsoNodeStates::TriangleComplete) == EIsoNodeStates::TriangleComplete;
+		return (States & EIsoNodeStates::NodeComplete) == EIsoNodeStates::NodeComplete;
 	}
 
 	bool IsLinkedToBoundary() const
@@ -520,7 +537,7 @@ public:
 
 	void OffsetId(int32 StartId)
 	{
-		Id += StartId;
+		NodeId += StartId;
 	}
 
 	virtual uint32 GetTypeHash() const override
diff --git a/Engine/Source/Runtime/Datasmith/CADKernel/Private/CADKernel/Mesh/Meshers/IsoTriangulator/IsoSegment.h b/Engine/Source/Runtime/Datasmith/CADKernel/Private/CADKernel/Mesh/Meshers/IsoTriangulator/IsoSegment.h
--- a/Engine/Source/Runtime/Datasmith/CADKernel/Private/CADKernel/Mesh/Meshers/IsoTriangulator/IsoSegment.h
+++ b/Engine/Source/Runtime/Datasmith/CADKernel/Private/CADKernel/Mesh/Meshers/IsoTriangulator/IsoSegment.h
@@ -178,24 +178,20 @@ public:
 		return (States & Side) == Side;
 	}
 
-	void SetHasTriangleOn(EIsoSegmentStates Side)
-	{
-		States |= Side;
-	}
-
 	bool HasTriangleOnLeft()
 	{
 		return (States & EIsoSegmentStates::LeftTriangle) == EIsoSegmentStates::LeftTriangle;
 	}
 
-	void SetHasTriangleOnLeft()
+	bool HasTriangleOnRight()
 	{
-		States |= EIsoSegmentStates::LeftTriangle;
+		return (States & EIsoSegmentStates::RightTriangle) == EIsoSegmentStates::RightTriangle;
 	}
 
-	bool HasTriangleOnRight()
+	bool HasntTriangle()
 	{
-		return (States & EIsoSegmentStates::RightTriangle) == EIsoSegmentStates::RightTriangle;
+		constexpr EIsoSegmentStates BoothSides = EIsoSegmentStates::RightTriangle | EIsoSegmentStates::LeftTriangle;
+		return (States & BoothSides) == EIsoSegmentStates::None;
 	}
 
 	bool HasTriangleOnRightAndLeft()
@@ -204,11 +200,33 @@ public:
 		return (States & BoothSides) == BoothSides;
 	}
 
+	void SetHasTriangleOn(EIsoSegmentStates Side)
+	{
+		States |= Side;
+	}
+
+	void SetHasTriangleOnLeft()
+	{
+		States |= EIsoSegmentStates::LeftTriangle;
+	}
+
 	void SetHasTriangleOnRight()
 	{
 		States |= EIsoSegmentStates::RightTriangle;
 	}
 
+	void SetHasInnerTriangle(bool bOrientation)
+	{
+		if (bOrientation)
+		{
+			SetHasTriangleOnLeft();
+		}
+		else
+		{
+			SetHasTriangleOnRight();
+		}
+	}
+
 	void ResetHasTriangle()
 	{
 		States &= ~EIsoSegmentStates::SegmentComplete;
@@ -249,6 +267,11 @@ public:
 		return (States & EIsoSegmentStates::Degenerate) == EIsoSegmentStates::Degenerate;
 	}
 
+	bool IsFirstNode(const FIsoNode* Node) const
+	{
+		return FirstNode == Node;
+	}
+
 	const FIsoNode& GetFirstNode() const
 	{
 		return *FirstNode;
@@ -369,6 +392,11 @@ inline uint32 GetTypeHash(const FIsoSegment& Segment0, const FIsoSegment& Segmen
 
 inline FIsoSegment* FIsoNode::GetSegmentConnectedTo(const FIsoNode* Node) const
 {
+	if (this == Node)
+	{
+		return nullptr;
+	}
+
 	for (FIsoSegment* Segment : ConnectedSegments)
 	{
 		if (&Segment->GetFirstNode() == Node)
diff --git a/Engine/Source/Runtime/Datasmith/CADKernel/Private/CADKernel/Mesh/Meshers/MesherTools.cpp b/Engine/Source/Runtime/Datasmith/CADKernel/Private/CADKernel/Mesh/Meshers/MesherTools.cpp
--- a/Engine/Source/Runtime/Datasmith/CADKernel/Private/CADKernel/Mesh/Meshers/MesherTools.cpp
+++ b/Engine/Source/Runtime/Datasmith/CADKernel/Private/CADKernel/Mesh/Meshers/MesherTools.cpp
@@ -298,4 +298,76 @@ void FMesherTools::ComputeFinalCuttingPointsWithImposedCuttingPoints(const TArra
 		LastCrossingIndex = CrossingIndex;
 	}
 };
+
+void FMesherTools::FillImposedIsoCuttingPoints(TArray<double>& UEdgeSetOfIntersectionWithIso, ECoordinateType CoordinateType, double EdgeToleranceGeo, const FTopologicalEdge& Edge, TArray<FCuttingPoint>& OutImposedIsoVertexSet)
+{
+	FLinearBoundary EdgeBoundary = Edge.GetBoundary();
+
+	int32 StartIndex = OutImposedIsoVertexSet.Num();
+	Algo::Sort(UEdgeSetOfIntersectionWithIso);
+	double PreviousU = -HUGE_VALUE;
+	for (double InterU : UEdgeSetOfIntersectionWithIso)
+	{
+		// Remove coordinate nearly equal to boundary
+		if ((InterU - EdgeToleranceGeo) < EdgeBoundary.GetMin() || (InterU + EdgeToleranceGeo) > EdgeBoundary.GetMax())
+		{
+			continue;
+		}
+
+		// Remove coordinate inside thin zone
+		for (FLinearBoundary ThinZone : Edge.GetThinZoneBounds())
+		{
+			if (ThinZone.Contains(InterU))
+			{
+				continue;
+			}
+		}
+
+		// Remove nearly duplicate 
+		if (InterU - PreviousU < EdgeToleranceGeo)
+		{
+			continue;
+		}
+
+		OutImposedIsoVertexSet.Emplace(InterU, CoordinateType);
+		PreviousU = InterU;
+	}
+
+	int32 Index;
+	int32 NewCoordinateCount = OutImposedIsoVertexSet.Num() - StartIndex;
+	switch (NewCoordinateCount)
+	{
+	case 0:
+		return;
+
+	case 1:
+	{
+		int32 CuttingPointIndex = 0;
+		while (CuttingPointIndex < Edge.GetCrossingPointUs().Num() && Edge.GetCrossingPointUs()[CuttingPointIndex] + DOUBLE_SMALL_NUMBER <= OutImposedIsoVertexSet[StartIndex].Coordinate)
+		{
+			++CuttingPointIndex;
+		};
+		if (CuttingPointIndex > 0)
+		{
+			--CuttingPointIndex;
+		}
+		OutImposedIsoVertexSet[StartIndex].IsoDeltaU = Edge.GetDeltaUMaxs()[CuttingPointIndex] * AQuarter;
+		break;
+	}
+
+	default:
+	{
+		OutImposedIsoVertexSet[StartIndex].IsoDeltaU = (OutImposedIsoVertexSet[StartIndex + 1].Coordinate - OutImposedIsoVertexSet[StartIndex].Coordinate) * AQuarter;
+		for (Index = StartIndex + 1; Index < OutImposedIsoVertexSet.Num() - 1; ++Index)
+		{
+			OutImposedIsoVertexSet[Index].IsoDeltaU = (OutImposedIsoVertexSet[Index + 1].Coordinate - OutImposedIsoVertexSet[Index - 1].Coordinate) * AEighth;
+		}
+		OutImposedIsoVertexSet[Index].IsoDeltaU = (OutImposedIsoVertexSet[Index].Coordinate - OutImposedIsoVertexSet[Index - 1].Coordinate) * AQuarter;
+		break;
+	}
+
+	}
+}
+
+
 }
\ No newline at end of file
diff --git a/Engine/Source/Runtime/Datasmith/CADKernel/Private/CADKernel/Mesh/Meshers/MesherTools.h b/Engine/Source/Runtime/Datasmith/CADKernel/Private/CADKernel/Mesh/Meshers/MesherTools.h
--- a/Engine/Source/Runtime/Datasmith/CADKernel/Private/CADKernel/Mesh/Meshers/MesherTools.h
+++ b/Engine/Source/Runtime/Datasmith/CADKernel/Private/CADKernel/Mesh/Meshers/MesherTools.h
@@ -3,9 +3,12 @@
 
 #include "CADKernel/Core/Types.h"
 #include "CADKernel/Math/Boundary.h"
+#include "CADKernel/Mesh/MeshEnum.h"
 
 namespace UE::CADKernel
 {
+class FTopologicalEdge;
+
 struct FCuttingPoint;
 struct FLinearBoundary;
 
@@ -62,6 +65,8 @@ public:
 			OutCuttingPointCoordinates.Add(OutCuttingPointCoordinates.Last() + Delta);
 		}
 	}
+
+	static void FillImposedIsoCuttingPoints(TArray<double>& UEdgeSetOfIntersectionWithIso, ECoordinateType CoordinateType, double EdgeToleranceGeo, const FTopologicalEdge& Edge, TArray<FCuttingPoint>& OutImposedIsoVertexSet);
 };
 
 } // namespace UE::CADKernel
