CADKernel update: Meshing of thin faces

diff --git a/Engine/Source/Runtime/Datasmith/CADKernel/Private/CADKernel/Mesh/Meshers/ParametricFaceMesher.cpp b/Engine/Source/Runtime/Datasmith/CADKernel/Private/CADKernel/Mesh/Meshers/ParametricFaceMesher.cpp
new file mode 100644
--- /dev/null
+++ b/Engine/Source/Runtime/Datasmith/CADKernel/Private/CADKernel/Mesh/Meshers/ParametricFaceMesher.cpp
@@ -0,0 +1,1353 @@
+// Copyright Epic Games, Inc. All Rights Reserved.
+
+#include "CADKernel/Mesh/Meshers/ParametricFaceMesher.h"
+
+#include "CADKernel/Mesh/Criteria/Criterion.h"
+#include "CADKernel/Mesh/Meshers/IsoTriangulator.h"
+#include "CADKernel/Mesh/Meshers/MesherTools.h"
+#include "CADKernel/Mesh/Meshers/ParametricMesherConstantes.h"
+#include "CADKernel/Mesh/Structure/EdgeMesh.h"
+#include "CADKernel/Mesh/Structure/FaceMesh.h"
+#include "CADKernel/Mesh/Structure/ModelMesh.h"
+#include "CADKernel/Mesh/Structure/ThinZone2DFinder.h"
+#include "CADKernel/Mesh/Structure/ThinZone2D.h"
+#include "CADKernel/Mesh/Structure/VertexMesh.h"
+#include "CADKernel/Topo/TopologicalEdge.h"
+#include "CADKernel/Topo/TopologicalFace.h"
+
+
+#ifdef CADKERNEL_DEV
+#include "CADKernel/Mesh/Meshers/MesherReport.h"
+#endif
+
+namespace UE::CADKernel
+{
+
+FParametricFaceMesher::FParametricFaceMesher(FTopologicalFace& InFace, FModelMesh& InMeshModel, const FMeshingTolerances& InTolerances, bool bActivateThinZoneMeshing)
+	: Face(InFace)
+	, MeshModel(InMeshModel)
+	, Tolerances(InTolerances)
+	, bThinZoneMeshing(bActivateThinZoneMeshing)
+	, Grid(Face, MeshModel)
+{
+#ifdef DEBUG_ONLY_SURFACE_TO_DEBUG
+	bDisplay = (Face.GetId() == FaceToDebug);
+#endif
+}
+
+void FParametricFaceMesher::Mesh()
+{
+#ifdef DEBUG_ONLY_SURFACE_TO_DEBUG
+	bDisplay = (Face.GetId() == FaceToDebug);
+#endif
+
+	Face.GetOrCreateMesh(MeshModel);
+
+	if (Face.IsNotMeshable())
+	{
+		return;
+	}
+
+	FMessage::Printf(EVerboseLevel::Debug, TEXT("Meshing of surface %d\n"), Face.GetId());
+
+	FProgress _(1, TEXT("Meshing Entities : Mesh Surface"));
+
+#ifdef DEBUG_CADKERNEL
+	F3DDebugSession S(bDisplayDebugMeshStep, FString::Printf(TEXT("Mesh of surface %d"), Face.GetId()));
+#endif
+
+	FTimePoint StartTime = FChrono::Now();
+
+
+	if (!GenerateCloud() || Grid.IsDegenerated())
+	{
+#ifdef CADKERNEL_DEV
+		FMesherReport::GetLogs().AddDegeneratedGrid();
+#endif
+		FMessage::Printf(EVerboseLevel::Log, TEXT("The meshing of the surface %d failed due to a degenerated grid\n"), Face.GetId());
+
+		Face.IsDegenerated();
+		Face.SetMeshedMarker();
+		return;
+	}
+
+	FTimePoint IsoTriangulerStartTime = FChrono::Now();
+
+	FFaceMesh& SurfaceMesh = Face.GetOrCreateMesh(MeshModel);
+	FIsoTriangulator IsoTrianguler(Grid, SurfaceMesh, Tolerances);
+
+	if (IsoTrianguler.Triangulate())
+	{
+		if (Face.IsBackOriented())
+		{
+			SurfaceMesh.InverseOrientation();
+		}
+		MeshModel.AddMesh(SurfaceMesh);
+	}
+	Face.SetMeshedMarker();
+
+	FDuration TriangulateDuration = FChrono::Elapse(IsoTriangulerStartTime);
+	FDuration Duration = FChrono::Elapse(StartTime);
+
+#ifdef CADKERNEL_DEV
+	FMesherReport::GetChronos().GlobalPointCloudDuration += Grid.Chronos.GeneratePointCloudDuration;
+	FMesherReport::GetChronos().GlobalTriangulateDuration += TriangulateDuration;
+	FMesherReport::GetChronos().GlobalMeshDuration += Duration;
+#endif
+
+	FChrono::PrintClockElapse(EVerboseLevel::Debug, TEXT("   "), TEXT("Meshing"), Duration);
+
+}
+
+bool FParametricFaceMesher::GenerateCloud()
+{
+	FTimePoint GenerateCloudStartTime = FChrono::Now();
+	Grid.DefineCuttingParameters();
+	if (!Grid.GeneratePointCloud())
+	{
+		return false;
+	}
+
+	if (bThinZoneMeshing)
+	{
+		FTimePoint StartTime = FChrono::Now();
+		if (Grid.GetFace().HasThinZone())
+		{
+			MeshThinZones();
+		}
+#ifdef CADKERNEL_DEV
+		FMesherReport::GetChronos().GlobalMeshThinZones += FChrono::Elapse(StartTime);
+#endif
+	}
+
+	MeshFaceLoops();
+
+	Grid.ProcessPointCloud();
+
+#ifdef CADKERNEL_DEV
+	FMesherReport::GetChronos().GlobalGeneratePointCloudDuration += FChrono::Elapse(GenerateCloudStartTime);
+#endif
+
+	return true;
+}
+
+void FParametricFaceMesher::MeshThinZones(TArray<FTopologicalEdge*>& EdgesToMesh, const bool bFinalMeshing)
+{
+	{
+#ifdef DEBUG_MESH_THINZONE_EDGES
+		F3DDebugSession A(bDisplay, TEXT("   - Define Cutting Points"));
+#endif
+		for (FTopologicalEdge* Edge : EdgesToMesh)
+		{
+			if (Edge->IsPreMeshed())
+			{
+				continue;
+			}
+
+#ifdef DEBUG_MESH_THINZONE_EDGES
+			{
+				F3DDebugSession B(bDisplay, FString::Printf(TEXT("Edge %d"), Edge->GetId()));
+				Display2DWithScale(*Edge, EVisuProperty::GreenCurve);
+			}
+#endif
+
+			{
+#ifdef DEBUG_MESH_THINZONE_EDGES
+				F3DDebugSession B(bDisplay, FString::Printf(TEXT("Edge %d"), Edge->GetId()));
+#endif
+				for (FThinZoneSide* ZoneSide : Edge->GetThinZoneSides())
+				{
+					if (ZoneSide->HasMarker1())
+					{
+						continue;
+					}
+					ZoneSide->SetMarker1();
+
+					DefineImposedCuttingPointsBasedOnOtherSideMesh(*ZoneSide);
+				}
+			}
+		}
+
+		for (FTopologicalEdge* Edge : EdgesToMesh)
+		{
+			for (FThinZoneSide* ZoneSide : Edge->GetThinZoneSides())
+			{
+				if (ZoneSide->HasMarker1())
+				{
+					ZoneSide->ResetMarker1();
+				}
+			}
+		}
+	}
+
+	{
+#ifdef DEBUG_MESH_THINZONE_EDGES
+		F3DDebugSession A(bDisplay, TEXT("   - Mesh"));
+#endif
+		for (FTopologicalEdge* Edge : EdgesToMesh)
+		{
+			if (Edge->IsMeshed())
+			{
+				continue;
+			}
+
+			Mesh(*Edge, bFinalMeshing);
+		}
+	}
+	//Wait();
+}
+
+namespace ThinZoneMesherTools
+{
+void ResetMarkers(TArray<FTopologicalEdge*>& EdgesWithThinZones, TArray<FThinZone2D>& ThinZones)
+{
+	for (FTopologicalEdge* Edge : EdgesWithThinZones)
+	{
+		Edge->ResetMarkers();
+	}
+
+	for (FThinZone2D& Zone : ThinZones)
+	{
+		Zone.ResetMarkers();
+		Zone.GetFirstSide().ResetMarkers();
+		Zone.GetSecondSide().ResetMarkers();
+	}
+};
+};
+
+void FParametricFaceMesher::SortThinZoneSides(TArray<FThinZone2D*>& ThinZones)
+{
+#ifdef DEBUG_SORT_THIN_ZONE_SIDES
+	F3DDebugSession A(bDisplay, ("SortThinZoneSides"));
+#endif
+
+	Face.ResetMarkersRecursively();
+	int32 EdgeCount = Face.EdgeCount();
+	ZoneAEdges.Reserve(EdgeCount);
+	ZoneBEdges.Reserve(EdgeCount);
+
+	WaitingThinZones.Reserve(ThinZones.Num());
+
+	TFunction<void(FThinZone2D*)> AddToWaitingList = [&WaitingList = WaitingThinZones](FThinZone2D* Zone)
+	{
+		Zone->SetWaitingMarker();
+		WaitingList.Add(Zone);
+	};
+
+	TFunction<void(FThinZone2D*)> SetAndGet = [&WaitingList = WaitingThinZones, &ZoneA = ZoneAEdges, &ZoneB = ZoneBEdges](FThinZone2D* Zone)
+	{
+		Zone->SetEdgesZoneSide();
+		Zone->GetEdges(ZoneA, ZoneB);
+	};
+
+
+	int32 Index = 1;
+	for (FThinZone2D* Zone : ThinZones)
+	{
+#ifdef DEBUG_SORT_THIN_ZONE_SIDES
+		Zone.Display(FString::Printf(TEXT("Zone %d"), Index++), EVisuProperty::YellowCurve);
+#endif
+
+		Zone->CheckEdgesZoneSide();
+
+
+		if (Zone->GetFirstSide().HasMarker1And2())
+		{
+			AddToWaitingList(Zone);
+			continue;
+		}
+
+		if (Zone->GetSecondSide().HasMarker1And2())
+		{
+			AddToWaitingList(Zone);
+			continue;
+		}
+
+		if (!Zone->GetFirstSide().HasMarker1Or2() && !Zone->GetSecondSide().HasMarker1Or2())
+		{
+			SetAndGet(Zone);
+			continue;
+		}
+
+		if (Zone->GetFirstSide().HasMarker1() && !Zone->GetSecondSide().HasMarker1())
+		{
+			SetAndGet(Zone);
+			continue;
+		}
+
+		if (Zone->GetFirstSide().HasMarker1() && Zone->GetSecondSide().HasMarker1())
+		{
+			AddToWaitingList(Zone);
+			continue;
+		}
+
+		if (!Zone->GetFirstSide().HasMarker2() && !Zone->GetSecondSide().HasMarker1())
+		{
+			SetAndGet(Zone);
+			continue;
+		}
+
+		if (Zone->GetFirstSide().HasMarker2() && Zone->GetSecondSide().HasMarker2()) // E
+		{
+			AddToWaitingList(Zone);
+			continue;
+		}
+
+		if (Zone->GetFirstSide().HasMarker2() && !Zone->GetSecondSide().HasMarker2()) // F
+		{
+			Zone->Swap();
+			SetAndGet(Zone);
+			continue;
+		}
+
+		if (!Zone->GetFirstSide().HasMarker1() && Zone->GetSecondSide().HasMarker1())
+		{
+			Zone->Swap();
+			SetAndGet(Zone);
+			continue;
+		}
+
+		ensureCADKernel(false);
+	}
+
+	Face.ResetMarkersRecursively();
+}
+
+void FParametricFaceMesher::MeshThinZones()
+{
+	TArray<FThinZone2D>& FaceThinZones = Face.GetThinZones();
+
+	if (FaceThinZones.IsEmpty())
+	{
+		return;
+	}
+
+	TArray<FThinZone2D*> ThinZones;
+	ThinZones.Reserve(FaceThinZones.Num());
+	for (FThinZone2D& FaceThinZone : FaceThinZones)
+	{
+		ThinZones.Add(&FaceThinZone);
+	}
+
+	while (ThinZones.Num())
+	{
+#ifdef DEBUG_THIN_ZONES
+		F3DDebugSession A(bDisplay, ("Mesh thin zone"));
+		DisplayMeshOfFaceLoop();
+#endif
+		int32 WaitingThinZoneCount = ThinZones.Num();
+
+		MeshThinZones(ThinZones);
+
+		if (WaitingThinZoneCount == WaitingThinZones.Num())
+		{
+			break;
+		}
+		ThinZones = MoveTemp(WaitingThinZones);
+	}
+}
+
+
+void FParametricFaceMesher::MeshThinZones(TArray<FThinZone2D*>& ThinZones)
+{
+	SortThinZoneSides(ThinZones);
+
+	TFunction<void(TArray<FTopologicalEdge*>)> TransfereCuttingPointsFromMeshedEdges = [](TArray<FTopologicalEdge*> Edges)
+	{
+		for (FTopologicalEdge* Edge : Edges)
+		{
+			FAddCuttingPointFunc AddCuttingPoint = [&Edge](const double Coordinate, const ECoordinateType Type, const FPairOfIndex OppositNodeIndices, const double DeltaU)
+			{
+				Edge->AddTwinsCuttingPoint(Coordinate, DeltaU);
+			};
+
+			constexpr bool bOnlyOppositNode = false;
+			Edge->TransferCuttingPointFromMeshedEdge(bOnlyOppositNode, AddCuttingPoint);
+		}
+	};
+
+	TransfereCuttingPointsFromMeshedEdges(ZoneAEdges);
+	TransfereCuttingPointsFromMeshedEdges(ZoneBEdges);
+
+#ifdef DEBUG_THIN_ZONES
+	if(bDisplay)
+	{
+		DisplayThinZoneEdges(TEXT("Side A"), ZoneAEdges, EVisuProperty::BlueCurve, EVisuProperty::PurpleCurve);
+		DisplayThinZoneEdges(TEXT("Side B"), ZoneBEdges, EVisuProperty::RedCurve, EVisuProperty::PinkCurve);
+
+		{
+			F3DDebugSession A(bDisplay, ("Waiting Zone"));
+			for(FThinZone2D* Zone : WaitingThinZones)
+			{
+				Zone->Display(TEXT("Waiting Zone"), EVisuProperty::YellowCurve);
+			}
+		}
+
+		//Wait();
+	}
+#endif
+
+	bool bFinalMeshing = false;
+	{
+#ifdef DEBUG_MESHTHINSURF
+		F3DDebugSession _(bDisplay, TEXT("Step 1"));
+#endif
+		MeshThinZones(ZoneAEdges, bFinalMeshing);
+	}
+
+	{
+#ifdef DEBUG_MESHTHINSURF
+		F3DDebugSession _(bDisplay, TEXT("Step 2"));
+#endif
+		bFinalMeshing = true;
+		MeshThinZones(ZoneBEdges, bFinalMeshing);
+	}
+
+	for (FTopologicalEdge* Edge : ZoneAEdges)
+	{
+		Edge->RemovePreMesh();
+	}
+
+	{
+#ifdef DEBUG_MESHTHINSURF
+		F3DDebugSession _(bDisplay, TEXT("Step 3"));
+#endif
+		MeshThinZones(ZoneAEdges, bFinalMeshing);
+	}
+
+#ifdef DEBUG_THIN_ZONES
+	{
+		DisplayMeshOfFaceLoop();
+	}
+	//Wait(bDisplay);
+#endif
+}
+
+void FParametricFaceMesher::DefineImposedCuttingPointsBasedOnOtherSideMesh(FThinZoneSide& SideToConstrain)
+{
+	using namespace ParametricMesherTool;
+
+	FThinZoneSide& FrontSide = SideToConstrain.GetFrontThinZoneSide();
+
+	TMap<int32, FPoint2D> ExistingMeshNodes;
+	TArray<FCrossZoneElement> CrossZoneElements;
+
+	FAddMeshNodeFunc AddToCrossZoneElements = [&CrossZoneElements](const int32 NodeIndice, const FPoint2D& MeshNode2D, double MeshingTolerance3D, const FEdgeSegment& EdgeSegment, const FPairOfIndex& OppositeNodeIndices)
+	{
+		if (CrossZoneElements.Num() && CrossZoneElements.Last().VertexId == NodeIndice)
+		{
+			CrossZoneElements.Last().Add(OppositeNodeIndices);
+		}
+		else
+		{
+			CrossZoneElements.Emplace(NodeIndice, MeshNode2D, MeshingTolerance3D, &EdgeSegment, OppositeNodeIndices);
+		}
+	};
+
+	FAddMeshNodeFunc AddToExistingMeshNodes = [&ExistingMeshNodes](const int32 NodeIndice, const FPoint2D& MeshNode2D, double MeshingTolerance3D, const FEdgeSegment& EdgeSegment, const FPairOfIndex& OppositeNodeIndices)
+	{
+		ExistingMeshNodes.Emplace(NodeIndice, MeshNode2D);
+	};
+
+	FReserveContainerFunc ReserveCrossZoneElements = [&CrossZoneElements](int32 MeshVertexCount)
+	{
+		CrossZoneElements.Reserve(MeshVertexCount);
+	};
+
+	FReserveContainerFunc ReserveExistingMeshNodes = [&ExistingMeshNodes](int32 MeshVertexCount)
+	{
+		ExistingMeshNodes.Reserve(MeshVertexCount);
+	};
+
+	SideToConstrain.GetExistingMeshNodes(Face, MeshModel, ReserveExistingMeshNodes, AddToExistingMeshNodes, /*bWithTolerance*/ false);
+	FrontSide.GetExistingMeshNodes(Face, MeshModel, ReserveCrossZoneElements, AddToCrossZoneElements, /*bWithTolerance*/ true);
+
+#ifdef DEBUG_MESHTHINSURF
+	if (bDisplay)
+	{
+		{
+			F3DDebugSession E(bDisplay, FString::Printf(TEXT("ThinZone")));
+			ThinZone::DisplayThinZoneSide(FrontSide, 0, EVisuProperty::BlueCurve);
+			ThinZone::DisplayThinZoneSide(SideToConstrain, 1, EVisuProperty::YellowPoint);
+			Wait(false);
+		}
+
+		F3DDebugSession A(TEXT("Existing Mesh nodes"));
+		{
+			F3DDebugSession A(TEXT("Existing Mesh nodes side to mesh"));
+			for (const TPair<int32, FPoint2D>& MeshNode : ExistingMeshNodes)
+			{
+				DisplayPoint2DWithScale(MeshNode.Value, EVisuProperty::RedPoint, MeshNode.Key);
+			}
+		}
+		{
+			F3DDebugSession A(TEXT("Existing Mesh nodes reference side"));
+			for (const FCrossZoneElement& MeshNode : CrossZoneElements)
+			{
+				DisplayPoint2DWithScale(MeshNode.VertexPoint2D, EVisuProperty::BluePoint, MeshNode.VertexId);
+				if (MeshNode.OppositeVertexIndices[0] >= 0)
+				{
+					FPoint2D* OppositeVertex = ExistingMeshNodes.Find(MeshNode.OppositeVertexIndices[0]);
+					if (OppositeVertex)
+					{
+						DisplayPoint2DWithScale(*OppositeVertex, EVisuProperty::RedPoint, MeshNode.OppositeVertexIndices[0]);
+						DisplaySegmentWithScale(*OppositeVertex, MeshNode.VertexPoint2D, 0, EVisuProperty::BlueCurve);
+					}
+				}
+			}
+		}
+		//Wait();
+	}
+
+	F3DDebugSession A(bDisplay, TEXT("Find the best projection of existing mesh vertices"));
+#endif
+
+	// Find the best projection of existing mesh vertices (CrossZone Vertex)
+	for (FCrossZoneElement& CrossZoneElement : CrossZoneElements)
+	{
+		if (CrossZoneElement.OppositeVertexIndices[0] >= 0)
+		{
+			FPoint2D* OppositeVertex = ExistingMeshNodes.Find(CrossZoneElement.OppositeVertexIndices[0]);
+			if (OppositeVertex)
+			{
+				CrossZoneElement.OppositePoint2D = *OppositeVertex;
+				CrossZoneElement.SquareThickness = 0;
+			}
+			continue;
+		}
+
+#ifdef DEBUG_MESHTHINSURF_
+		if (bDisplay)
+		{
+			F3DDebugSession A(TEXT("Point"));
+			DisplayPoint2DWithScale(CrossZoneElement.VertexPoint2D, EVisuProperty::BluePoint, CrossZoneElement.VertexId);
+			Wait();
+		}
+#endif
+
+		double MinSquareThickness = HUGE_VALUE;
+		FPoint2D ClosePoint;
+		FEdgeSegment* CloseSegment = nullptr;
+		double ClosePointCoordinate = -1;
+
+		const FPoint2D MeshNodeCoordinate = CrossZoneElement.VertexPoint2D;
+
+		for (FEdgeSegment& Segment : SideToConstrain.GetSegments())
+		{
+			// check the angle between segment and Middle-SegementStart, Middle-SegementEnd.
+			const double SlopeS = Segment.ComputeOrientedSlopeOf(MeshNodeCoordinate, Segment.GetExtemity(ELimit::Start));
+			const double SlopeE = Segment.ComputeOrientedSlopeOf(MeshNodeCoordinate, Segment.GetExtemity(ELimit::End));
+			if (SlopeE < 1. || SlopeS > 3.)
+			{
+#ifdef DEBUG_MESHTHINSURF_
+				if (bDisplay)
+				{
+					F3DDebugSession A(FString::Printf(TEXT("Not candidate slope end %f"), SlopeE));
+					DisplaySegmentWithScale(Segment.GetExtemity(Start), Segment.GetExtemity(End), 0, EVisuProperty::RedCurve);
+					DisplaySegmentWithScale(MeshNodeCoordinate, Segment.GetExtemity(End), 0, EVisuProperty::YellowCurve);
+					DisplaySegmentWithScale(MeshNodeCoordinate, Segment.GetExtemity(Start), 0, EVisuProperty::YellowCurve);
+				}
+#endif
+				continue;
+			}
+
+			double CoordSegmentU;
+			FPoint2D Projection = Segment.ProjectPoint(MeshNodeCoordinate, CoordSegmentU);
+
+
+#ifdef DEBUG_MESHTHINSURF_
+			if (bDisplay)
+			{
+				const double Slope = Segment.ComputeOrientedSlopeOf(MeshNodeCoordinate, Projection);
+				F3DDebugSession A(FString::Printf(TEXT("Close Point slope %f %f %f"), Slope, SlopeS, SlopeE));
+				DisplaySegmentWithScale(Segment.GetExtemity(Start), Segment.GetExtemity(End), 0, EVisuProperty::BlueCurve);
+				DisplaySegmentWithScale(MeshNodeCoordinate, Projection, 0, EVisuProperty::GreenCurve);
+				DisplayPoint2DWithScale(Projection, EVisuProperty::YellowPoint, CrossZoneElement.VertexId);
+				//Wait();
+			}
+#endif
+
+			const double SquareDistance = MeshNodeCoordinate.SquareDistance(Projection);
+			if (MinSquareThickness > SquareDistance)
+			{
+				FEdgeSegment* SegmentPtr = &Segment;
+				// Forbid the common extremity as candidate
+				{
+					constexpr double NearlyZero = 0. + DOUBLE_SMALL_NUMBER;
+					constexpr double NearlyOne = 1. - DOUBLE_SMALL_NUMBER;
+
+					const FEdgeSegment* CrossZoneSegment = CrossZoneElement.Segment;
+					if (SegmentPtr == CrossZoneSegment->GetPrevious() && CoordSegmentU > NearlyOne)
+					{
+						continue;
+					}
+					if (SegmentPtr == CrossZoneSegment->GetNext() && CoordSegmentU < NearlyZero)
+					{
+						continue;
+					}
+				}
+
+				MinSquareThickness = SquareDistance;
+				ClosePoint = Projection;
+				ClosePointCoordinate = CoordSegmentU;
+				CloseSegment = SegmentPtr;
+			}
+		}
+
+		if (CloseSegment)
+		{
+			CrossZoneElement.OppositePoint2D = ClosePoint;
+			CrossZoneElement.OppositeSegment = CloseSegment;
+			CrossZoneElement.OppositePointCoordinate = ClosePointCoordinate;
+			CrossZoneElement.SquareThickness = MinSquareThickness;
+		}
+	}
+
+	Algo::Sort(CrossZoneElements, [](FCrossZoneElement& ElementA, FCrossZoneElement& ElementB) { return ElementA.SquareThickness < ElementB.SquareThickness; });
+
+#ifdef DEBUG_MESHTHINSURF_
+	if (bDisplay)
+	{
+		F3DDebugSession A(TEXT("Result before intersection filtering"));
+		for (FCrossZoneElement& CrossZoneElement : CrossZoneElements)
+		{
+			//F3DDebugSession A(TEXT("Point and Close"));
+			DisplayPoint2DWithScale(CrossZoneElement.VertexPoint2D, EVisuProperty::BluePoint, CrossZoneElement.VertexId);
+			DisplayPoint2DWithScale(CrossZoneElement.OppositePoint2D, EVisuProperty::RedPoint);
+			DisplaySegmentWithScale(CrossZoneElement.VertexPoint2D, CrossZoneElement.OppositePoint2D);
+			if (CrossZoneElement.OppositeSegment)
+			{
+				DisplaySegmentWithScale(CrossZoneElement.OppositeSegment->GetExtemity(ELimit::End), CrossZoneElement.OppositeSegment->GetExtemity(ELimit::Start));
+			}
+		}
+		Wait();
+	}
+#endif
+
+	// Find candidates i.e. CrossZoneElement that are not in intersection with the sides and with selected CrossZoneElement
+	FIntersectionTool IntersectionTool(FrontSide.GetSegments(), SideToConstrain.GetSegments(), CrossZoneElements.Num());
+	for (FCrossZoneElement& CrossZoneElement : CrossZoneElements)
+	{
+		if (!CrossZoneElement.OppositeSegment)
+		{
+			continue;
+		}
+
+		if (!IntersectionTool.IsIntersectSides(CrossZoneElement))
+		{
+			IntersectionTool.AddCrossZoneElement(CrossZoneElement);
+			CrossZoneElement.bIsSelected = true;
+		}
+	}
+
+#ifdef DEBUG_MESHTHINSURF
+	if (bDisplay)
+	{
+		F3DDebugSession M(TEXT("Result final"));
+		for (FCrossZoneElement& CrossZoneElement : CrossZoneElements)
+		{
+			if (CrossZoneElement.bIsSelected)
+			{
+				//F3DDebugSession A(TEXT("Point and Close"));
+				DisplayPoint2DWithScale(CrossZoneElement.VertexPoint2D, EVisuProperty::BluePoint, CrossZoneElement.VertexId);
+				DisplayPoint2DWithScale(CrossZoneElement.OppositePoint2D, EVisuProperty::RedPoint);
+				DisplaySegmentWithScale(CrossZoneElement.VertexPoint2D, CrossZoneElement.OppositePoint2D, 0, EVisuProperty::YellowCurve);
+			}
+		}
+		//Wait();
+	}
+#endif
+
+	// Add ImposedCuttingPointU
+	for (FCrossZoneElement& CrossZoneElement : CrossZoneElements)
+	{
+		if (CrossZoneElement.bIsSelected && (CrossZoneElement.OppositeVertexIndices[0] < 0))
+		{
+			FTopologicalEdge* OppositeEdge = CrossZoneElement.OppositeSegment->GetEdge();
+			if (OppositeEdge == nullptr)
+			{
+				continue;
+			}
+
+			const double OppositeCuttingPointU = CrossZoneElement.OppositeSegment->ComputeEdgeCoordinate(CrossZoneElement.OppositePointCoordinate);
+			const double DeltaU = CrossZoneElement.OppositeSegment->ComputeDeltaU(CrossZoneElement.Tolerance3D);
+
+			OppositeEdge->AddImposedCuttingPointU(OppositeCuttingPointU, CrossZoneElement.VertexId, DeltaU);
+		}
+	}
+}
+
+void FParametricFaceMesher::MeshFaceLoops()
+{
+#ifdef DEBUG_MESHTHINSURF
+	F3DDebugSession M(bDisplay, ("Mesh other edges"));
+#endif
+
+	for (const TSharedPtr<FTopologicalLoop>& Loop : Face.GetLoops())
+	{
+		for (const FOrientedEdge& Edge : Loop->GetEdges())
+		{
+			Mesh(*Edge.Entity);
+		}
+	}
+}
+
+void FParametricFaceMesher::Mesh(FTopologicalVertex& InVertex)
+{
+	InVertex.GetOrCreateMesh(MeshModel);
+}
+
+#ifdef DEBUG_INTERSECTEDGEISOS
+void DebugIntersectEdgeIsos(const FTopologicalFace& Face, const TArray<double>& IsoCoordinates, EIso TypeIso);
+#endif
+
+void FParametricFaceMesher::Mesh(FTopologicalEdge& InEdge, bool bFinalMeshing)
+{
+  	FTopologicalEdge& ActiveEdge = *InEdge.GetLinkActiveEntity();
+	if (ActiveEdge.IsMeshed())
+	{
+		if (ActiveEdge.GetMesh()->GetNodeCount() > 0)
+		{
+			return;
+		}
+
+		// In some case the 2d curve is a smooth curve and the 3d curve is a line and vice versa
+		// In the particular case where the both case are opposed, we can have the 2d line sampled with 4 points, and the 2d curve sampled with 2 points (because in 3d, the 2d curve is a 3d line)
+		// In this case, the loop is flat i.e. in 2d the meshes of the 2d line and 2d curve are coincident
+		// So the grid is degenerated and the surface is not meshed
+		// to avoid this case, the Edge is virtually meshed i.e. the nodes inside the edge have the id of the mesh of the vertices.
+		InEdge.SetVirtuallyMeshedMarker();
+	}
+
+	if (ActiveEdge.IsThinPeak())
+	{
+		TArray<FCuttingPoint>& FinalEdgeCuttingPointCoordinates = ActiveEdge.GetCuttingPoints();
+		FinalEdgeCuttingPointCoordinates.Emplace(ActiveEdge.GetStartCurvilinearCoordinates(), ECoordinateType::VertexCoordinate);
+		FinalEdgeCuttingPointCoordinates.Emplace(ActiveEdge.GetEndCurvilinearCoordinates(), ECoordinateType::VertexCoordinate);
+		ActiveEdge.GenerateMeshElements(MeshModel);
+		return;
+	}
+
+	const FSurfacicTolerance& ToleranceIso = Face.GetIsoTolerances();
+
+	// Get Edge intersection with inner surface mesh grid
+	TArray<double> EdgeIntersectionWithIsoU_Coordinates;
+	TArray<double> EdgeIntersectionWithIsoV_Coordinates;
+
+	const TArray<double>& SurfaceTabU = Face.GetCuttingCoordinatesAlongIso(EIso::IsoU);
+	const TArray<double>& SurfaceTabV = Face.GetCuttingCoordinatesAlongIso(EIso::IsoV);
+
+	ApplyEdgeCriteria(InEdge);
+
+#ifdef DEBUG_MESH_EDGE
+	if (bDisplay)
+	{
+		F3DDebugSession _(FString::Printf(TEXT("EdgePointsOnDomain %d"), InEdge.GetId()));
+		Display2D(InEdge);
+		Wait();
+	}
+#endif
+
+#ifdef DEBUG_INTERSECTEDGEISOS
+	DebugIntersectEdgeIsos(Face, SurfaceTabU, EIso::IsoU);
+	DebugIntersectEdgeIsos(Face, SurfaceTabV, EIso::IsoV);
+	{
+		F3DDebugSession _(FString::Printf(TEXT("Edge 2D %d"), InEdge.GetId()));
+		UE::CADKernel::Display2D(InEdge);
+	}
+#endif
+
+	InEdge.ComputeIntersectionsWithIsos(SurfaceTabU, EIso::IsoU, ToleranceIso, EdgeIntersectionWithIsoU_Coordinates);
+	InEdge.ComputeIntersectionsWithIsos(SurfaceTabV, EIso::IsoV, ToleranceIso, EdgeIntersectionWithIsoV_Coordinates);
+
+#ifdef DEBUG_INTERSECTEDGEISOS
+	{
+		F3DDebugSession _(FString::Printf(TEXT("Edge %d Intersect with iso"), InEdge.GetId()));
+		TArray<FPoint2D> Intersections;
+		InEdge.Approximate2DPoints(EdgeIntersectionWithIsoU_Coordinates, Intersections);
+		for (const FPoint2D& Point : Intersections)
+		{
+			DisplayPoint(Point);
+		}
+		Intersections.Empty();
+		InEdge.Approximate2DPoints(EdgeIntersectionWithIsoV_Coordinates, Intersections);
+		for (const FPoint2D& Point : Intersections)
+		{
+			DisplayPoint(Point);
+		}
+		Wait();
+	}
+	{
+		F3DDebugSession _(FString::Printf(TEXT("Thin Zone")));
+		for (FLinearBoundary ThinZone : InEdge.GetThinZoneBounds())
+		{
+			TArray<double> Coords;
+			Coords.Add(ThinZone.GetMin());
+			Coords.Add(ThinZone.GetMax());
+			TArray<FPoint2D> ThinZone2D;
+			InEdge.Approximate2DPoints(Coords, ThinZone2D);
+			DisplaySegment(ThinZone2D[0], ThinZone2D[1], EVisuProperty::YellowCurve);
+		}
+		Wait();
+
+	}
+#endif
+
+	FLinearBoundary EdgeBounds = InEdge.GetBoundary();
+
+	TArray<double>& DeltaUs = InEdge.GetDeltaUMaxs();
+
+	FAddCuttingPointFunc AddCuttingPoint = [&InEdge](const double Coordinate, const ECoordinateType Type, const FPairOfIndex OppositNodeIndices, const double DeltaU)
+	{
+		for (int32 Index = 0; Index < 2; ++Index)
+		{
+			if (OppositNodeIndices[Index] >= 0)
+			{
+				InEdge.AddImposedCuttingPointU(Coordinate, OppositNodeIndices[Index], DeltaU);
+			}
+		}
+	};
+
+	// Case of self connected surface (e.g. cylinder) an edge 
+	// The first edge is premeshed at step 1, but the activeEdge is not yet meshed
+	// the twin edge is meshed at step 2
+	if (ActiveEdge.IsPreMeshed())
+	{
+		constexpr bool bOnlyOppositNode = true;
+		InEdge.TransferCuttingPointFromMeshedEdge(bOnlyOppositNode, AddCuttingPoint);
+
+		FTopologicalEdge* PreMeshEdge = InEdge.GetPreMeshedTwin();
+		if(PreMeshEdge)
+		{
+			PreMeshEdge->RemovePreMesh();
+		}
+	}
+
+	InEdge.SortImposedCuttingPoints();
+	const TArray<FImposedCuttingPoint>& EdgeImposedCuttingPoints = InEdge.GetImposedCuttingPoints();
+
+	// build a edge mesh compiling inner surface cutting (based on criteria applied on the surface) and edge cutting (based on criteria applied on the curve)
+	TArray<FCuttingPoint> ImposedIsoCuttingPoints;
+
+	TFunction<void(int32&, int32&)> UpdateDeltaU = [&ImposedIsoCuttingPoints](int32& NewIndex, int32& Index)
+	{
+		if (ImposedIsoCuttingPoints[NewIndex].IsoDeltaU > ImposedIsoCuttingPoints[Index].IsoDeltaU)
+		{
+			ImposedIsoCuttingPoints[NewIndex].IsoDeltaU = ImposedIsoCuttingPoints[Index].IsoDeltaU;
+		}
+	};
+
+	TFunction<void(int32&, int32&)> UpdateOppositNodeIndices = [&ImposedIsoCuttingPoints](int32& NewIndex, int32& Index)
+	{
+		if (ImposedIsoCuttingPoints[NewIndex].OppositNodeIndices[0] == -1)
+		{
+			ImposedIsoCuttingPoints[NewIndex].OppositNodeIndices[0] = ImposedIsoCuttingPoints[Index].OppositNodeIndices[0];
+		}
+		else if (ImposedIsoCuttingPoints[NewIndex].OppositNodeIndices[0] != ImposedIsoCuttingPoints[Index].OppositNodeIndices[0])
+		{
+			ImposedIsoCuttingPoints[NewIndex].OppositNodeIndices[1] = ImposedIsoCuttingPoints[Index].OppositNodeIndices[0];
+		}
+	};
+
+	TFunction<void(int32&, int32&, ECoordinateType)> MergeImposedCuttingPoints = [&ImposedIsoCuttingPoints, UpdateOppositNodeIndices, UpdateDeltaU](int32& Index, int32& NewIndex, ECoordinateType NewType)
+	{
+		double DeltaU = FMath::Max(ImposedIsoCuttingPoints[NewIndex].IsoDeltaU, ImposedIsoCuttingPoints[Index].IsoDeltaU);
+
+		if (ImposedIsoCuttingPoints[NewIndex].Coordinate + DeltaU > ImposedIsoCuttingPoints[Index].Coordinate)
+		{
+			if (ImposedIsoCuttingPoints[Index].Type == VertexCoordinate)
+			{
+				ImposedIsoCuttingPoints[NewIndex].Coordinate = ImposedIsoCuttingPoints[Index].Coordinate;
+				ImposedIsoCuttingPoints[NewIndex].IsoDeltaU = ImposedIsoCuttingPoints[Index].IsoDeltaU;
+				ImposedIsoCuttingPoints[NewIndex].Type = ImposedIsoCuttingPoints[Index].Type;
+				UpdateOppositNodeIndices(NewIndex, Index);
+				UpdateDeltaU(NewIndex, Index);
+			}
+			else if (ImposedIsoCuttingPoints[NewIndex].Type == VertexCoordinate)
+			{
+				if (ImposedIsoCuttingPoints[Index].Type == ImposedCoordinate)
+				{
+					UpdateOppositNodeIndices(NewIndex, Index);
+					UpdateDeltaU(NewIndex, Index);
+				}
+			}
+			else if (ImposedIsoCuttingPoints[NewIndex].Type == ImposedCoordinate)
+			{
+				if (ImposedIsoCuttingPoints[Index].Type == ImposedCoordinate)
+				{
+					ImposedIsoCuttingPoints[NewIndex].Coordinate = (ImposedIsoCuttingPoints[NewIndex].Coordinate + ImposedIsoCuttingPoints[Index].Coordinate) * 0.5;
+					UpdateOppositNodeIndices(NewIndex, Index);
+					UpdateDeltaU(NewIndex, Index);
+				}
+			}
+			else if (ImposedIsoCuttingPoints[Index].Type == ImposedCoordinate)
+			{
+				ImposedIsoCuttingPoints[NewIndex].Coordinate = ImposedIsoCuttingPoints[Index].Coordinate;
+				ImposedIsoCuttingPoints[NewIndex].Type = ImposedCoordinate;
+				UpdateOppositNodeIndices(NewIndex, Index);
+				UpdateDeltaU(NewIndex, Index);
+			}
+			else if (ImposedIsoCuttingPoints[NewIndex].Type != ImposedIsoCuttingPoints[Index].Type)
+			{
+				ImposedIsoCuttingPoints[NewIndex].Coordinate = (ImposedIsoCuttingPoints[NewIndex].Coordinate + ImposedIsoCuttingPoints[Index].Coordinate) * 0.5;
+				ImposedIsoCuttingPoints[NewIndex].Type = IsoUVCoordinate;
+				UpdateDeltaU(NewIndex, Index);
+			}
+		}
+		else
+		{
+			++NewIndex;
+			ImposedIsoCuttingPoints[NewIndex] = ImposedIsoCuttingPoints[Index];
+		}
+	};
+
+	{
+		int32 NbImposedCuttingPoints = EdgeImposedCuttingPoints.Num() + EdgeIntersectionWithIsoU_Coordinates.Num() + EdgeIntersectionWithIsoV_Coordinates.Num() + 2;
+		ImposedIsoCuttingPoints.Reserve(NbImposedCuttingPoints);
+	}
+
+
+	const double EdgeDeltaUAtMin = DeltaUs[0] * AQuarter;
+	const double EdgeDeltaUAtMax = DeltaUs.Last() * AQuarter;
+
+	ImposedIsoCuttingPoints.Emplace(EdgeBounds.GetMin(), ECoordinateType::VertexCoordinate, FPairOfIndex::Undefined, EdgeDeltaUAtMin);
+	ImposedIsoCuttingPoints.Emplace(EdgeBounds.GetMax(), ECoordinateType::VertexCoordinate, FPairOfIndex::Undefined, EdgeDeltaUAtMax);
+
+	int32 Index = 0;
+	for (const FImposedCuttingPoint& CuttingPoint : EdgeImposedCuttingPoints)
+	{
+		const double CuttingPointDeltaU = CuttingPoint.DeltaU;
+		ImposedIsoCuttingPoints.Emplace(CuttingPoint.Coordinate, ECoordinateType::ImposedCoordinate, CuttingPoint.OppositNodeIndex, CuttingPointDeltaU * AThird);
+	}
+
+	// Add Edge intersection with inner surface grid Iso
+	FPoint2D ExtremityTolerances = InEdge.GetCurve()->GetExtremityTolerances(EdgeBounds);
+	double EdgeTolerance = FMath::Min(ExtremityTolerances[0], ExtremityTolerances[1]);
+	if (!EdgeIntersectionWithIsoU_Coordinates.IsEmpty())
+	{
+		FMesherTools::FillImposedIsoCuttingPoints(EdgeIntersectionWithIsoU_Coordinates, IsoUCoordinate, EdgeTolerance, InEdge, ImposedIsoCuttingPoints);
+	}
+
+	if (!EdgeIntersectionWithIsoV_Coordinates.IsEmpty())
+	{
+		FMesherTools::FillImposedIsoCuttingPoints(EdgeIntersectionWithIsoV_Coordinates, IsoVCoordinate, EdgeTolerance, InEdge, ImposedIsoCuttingPoints);
+	}
+
+	ImposedIsoCuttingPoints.Sort([](const FCuttingPoint& Point1, const FCuttingPoint& Point2) { return Point1.Coordinate < Point2.Coordinate; });
+
+	// If a pair of point isoU/isoV is too close, get the middle of the points
+	if (ImposedIsoCuttingPoints.Num() > 1)
+	{
+		int32 NewIndex = 0;
+		for (int32 Andex = 1; Andex < ImposedIsoCuttingPoints.Num(); ++Andex)
+		{
+			if (ImposedIsoCuttingPoints[Andex].Type > ECoordinateType::ImposedCoordinate)
+			{
+				bool bIsDelete = false;
+				for (const FLinearBoundary& ThinZone : InEdge.GetThinZoneBounds())
+				{
+					if (ThinZone.Contains(ImposedIsoCuttingPoints[Andex].Coordinate))
+					{
+						bIsDelete = true;
+						break; // or continue
+					}
+				}
+				if (bIsDelete)
+				{
+					continue;
+				}
+			}
+
+			if (ImposedIsoCuttingPoints[NewIndex].Type == ECoordinateType::ImposedCoordinate || ImposedIsoCuttingPoints[Andex].Type == ECoordinateType::ImposedCoordinate)
+			{
+				MergeImposedCuttingPoints(Andex, NewIndex, ECoordinateType::ImposedCoordinate);
+			}
+			else if (ImposedIsoCuttingPoints[NewIndex].Type != ImposedIsoCuttingPoints[Andex].Type)
+			{
+				MergeImposedCuttingPoints(Andex, NewIndex, ECoordinateType::IsoUVCoordinate);
+			}
+			else
+			{
+				++NewIndex;
+				ImposedIsoCuttingPoints[NewIndex] = ImposedIsoCuttingPoints[Andex];
+			}
+		}
+		ImposedIsoCuttingPoints.SetNum(NewIndex + 1);
+	}
+
+	if (ImposedIsoCuttingPoints.Num() > 1 && (EdgeBounds.GetMax() - ImposedIsoCuttingPoints.Last().Coordinate) < FMath::Min(ImposedIsoCuttingPoints.Last().IsoDeltaU, InEdge.GetDeltaUMaxs().Last()))
+	{
+		ImposedIsoCuttingPoints.Last().Coordinate = EdgeBounds.GetMax();
+		ImposedIsoCuttingPoints.Last().Type = VertexCoordinate;
+	}
+	else
+	{
+		ImposedIsoCuttingPoints.Emplace(EdgeBounds.GetMax(), ECoordinateType::VertexCoordinate, -1, InEdge.GetDeltaUMaxs().Last() * AQuarter);
+	}
+
+	// Final array of the edge mesh vertex 
+	TArray<FCuttingPoint>& FinalEdgeCuttingPointCoordinates = InEdge.GetCuttingPoints();
+	{
+		// max count of vertex
+		double MinDeltaU = HUGE_VALUE;
+		for (const double& DeltaU : DeltaUs)
+		{
+			if (DeltaU < MinDeltaU)
+			{
+				MinDeltaU = DeltaU;
+			}
+		}
+
+		int32 MaxNumberOfVertex = FMath::IsNearlyZero(MinDeltaU) ? 5 : (int32)((EdgeBounds.GetMax() - EdgeBounds.GetMin()) / MinDeltaU) + 5;
+		FinalEdgeCuttingPointCoordinates.Empty(ImposedIsoCuttingPoints.Num() + MaxNumberOfVertex);
+	}
+
+#ifdef DEBUG_GETPREFERREDUVCOORDINATESFROMNEIGHBOURS
+	TArray<double> CuttingPoints2;
+	{
+		double ToleranceGeoEdge = FMath::Min(ExtremityTolerances[0], ExtremityTolerances[1]);
+
+		TArray<FCuttingPoint> Extremities;
+		Extremities.Reserve(2);
+		Extremities.Emplace(EdgeBounds.GetMin(), ECoordinateType::VertexCoordinate, -1, ToleranceGeoEdge);
+		Extremities.Emplace(EdgeBounds.GetMax(), ECoordinateType::VertexCoordinate, -1, ToleranceGeoEdge);
+
+		FMesherTools::ComputeFinalCuttingPointsWithImposedCuttingPoints(InEdge.GetCrossingPointUs(), InEdge.GetDeltaUMaxs(), Extremities, CuttingPoints2);
+	}
+#endif
+
+	if (InEdge.IsDegenerated() || InEdge.IsVirtuallyMeshed())
+	{
+		if (ImposedIsoCuttingPoints.Num() == 2)
+		{
+			ImposedIsoCuttingPoints.EmplaceAt(1, (ImposedIsoCuttingPoints[0].Coordinate + ImposedIsoCuttingPoints[1].Coordinate) * 0.5, ECoordinateType::OtherCoordinate);
+		}
+
+		for (FCuttingPoint CuttingPoint : ImposedIsoCuttingPoints)
+		{
+			FinalEdgeCuttingPointCoordinates.Emplace(CuttingPoint.Coordinate, ECoordinateType::OtherCoordinate);
+		}
+		InEdge.GetLinkActiveEdge()->SetMeshedMarker();
+		return;
+	}
+
+	TArray<double> CuttingPoints;
+	FMesherTools::ComputeFinalCuttingPointsWithImposedCuttingPoints(InEdge.GetCrossingPointUs(), InEdge.GetDeltaUMaxs(), ImposedIsoCuttingPoints, CuttingPoints);
+	int32 ImposedIndex = 0;
+	int32 ImposedIsoCuttingPointsCount = ImposedIsoCuttingPoints.Num();
+	for (const double& Coordinate : CuttingPoints)
+	{
+		if (FMath::IsNearlyEqual(ImposedIsoCuttingPoints[ImposedIndex].Coordinate, Coordinate))
+		{
+			FinalEdgeCuttingPointCoordinates.Emplace(ImposedIsoCuttingPoints[ImposedIndex]);
+			++ImposedIndex;
+		}
+		else
+		{
+			while (ImposedIndex < ImposedIsoCuttingPointsCount && ImposedIsoCuttingPoints[ImposedIndex].Coordinate < Coordinate)
+			{
+				++ImposedIndex;
+			}
+			FinalEdgeCuttingPointCoordinates.Emplace(Coordinate, ECoordinateType::OtherCoordinate);
+		}
+	}
+
+#ifdef DEBUG_GETPREFERREDUVCOORDINATESFROMNEIGHBOURS
+	if (InEdge.IsThinZone() && EdgeImposedCuttingPoints.Num())
+	{
+		F3DDebugSession G(TEXT("Mesh(TSharedRef<FEdge> InEdge"));
+		{
+			F3DDebugSession G(TEXT("U From Iso"));
+			for (const FCuttingPoint& CuttingU : ImposedIsoCuttingPoints)
+			{
+				if (CuttingU.OppositNodeIndex >= 0)
+				{
+					UE::CADKernel::DisplayPoint(FPoint2D(CuttingU.Coordinate, 0.0), EVisuProperty::RedPoint);
+				}
+				else
+				{
+					UE::CADKernel::DisplayPoint(FPoint2D(CuttingU.Coordinate, 0.0));
+				}
+			}
+		}
+		{
+			F3DDebugSession _(InEdge.GetThinZoneBounds().Num() != 0, FString::Printf(TEXT("Thin Zone"), InEdge.GetId()));
+			for (FLinearBoundary ThinZone : InEdge.GetThinZoneBounds())
+			{
+				UE::CADKernel::DisplayPoint(FPoint2D(ThinZone.GetMin(), 0.01), EVisuProperty::BluePoint);
+				UE::CADKernel::DisplayPoint(FPoint2D(ThinZone.GetMax(), 0.01), EVisuProperty::BluePoint);
+				DisplaySegment(FPoint2D(ThinZone.GetMin(), 0.01), FPoint2D(ThinZone.GetMax(), 0.01), EVisuProperty::BlueCurve);
+			}
+			Wait();
+		}
+
+		{
+			F3DDebugSession G(TEXT("U From Criteria"));
+			for (double CuttingU : CuttingPoints2)
+			{
+				UE::CADKernel::DisplayPoint(FPoint2D(CuttingU, 0.02), EVisuProperty::RedPoint);
+			}
+		}
+		{
+			F3DDebugSession G(TEXT("U Final (Criteria & Iso)"));
+			for (double CuttingU : CuttingPoints)
+			{
+				UE::CADKernel::DisplayPoint(FPoint2D(CuttingU, 0.04), EVisuProperty::YellowPoint);
+			}
+		}
+		Wait();
+	}
+#endif
+
+	if (bFinalMeshing)
+	{
+		InEdge.GenerateMeshElements(MeshModel);
+	}
+	else
+	{
+		ActiveEdge.SetPreMeshedMarker();
+		InEdge.SetPreMeshedMarker();
+	}
+
+#ifdef DEBUG_THIN_ZONES
+	if (bFinalMeshing && bDisplay)
+	{
+#ifdef DEBUG_MESHTHINSURF
+		F3DDebugSession A(bDisplay, FString::Printf(TEXT("Edge %d"), InEdge.GetId()));
+#endif
+		Display2DWithScale(InEdge);
+		TArray<FPoint> MeshVertices = InEdge.GetMesh()->GetNodeCoordinates();
+		MeshVertices.Append(InEdge.GetStartVertex()->GetMesh()->GetNodeCoordinates());
+		MeshVertices.Append(InEdge.GetEndVertex()->GetMesh()->GetNodeCoordinates());
+		if (MeshVertices.Num())
+		{
+			TArray<double> CoordinatesOfMesh;
+			TArray<FPoint> ProjectedPoints;
+			TArray<FPoint2D> Mesh2DPoints;
+			InEdge.ProjectPoints(MeshVertices, CoordinatesOfMesh, ProjectedPoints);
+			Algo::Sort(CoordinatesOfMesh);
+			InEdge.Approximate2DPoints(CoordinatesOfMesh, Mesh2DPoints);
+			int32 Index = 0;
+			for (FPoint2D Point : Mesh2DPoints)
+			{
+				DisplayPoint2DWithScale(Point, EVisuProperty::RedPoint, ++Index);
+			}
+		}
+#ifdef DEBUG_MESHTHINSURF
+		//Wait(bDisplay);
+#endif
+	}
+#endif
+
+}
+
+void FParametricFaceMesher::ApplyEdgeCriteria(FTopologicalEdge& Edge)
+{
+	FTopologicalEdge& ActiveEdge = *Edge.GetLinkActiveEdge();
+
+	if (Edge.Length() < Edge.GetTolerance3D())
+	{
+		Edge.SetAsDegenerated();
+	}
+
+	Edge.ComputeCrossingPointCoordinates();
+	Edge.InitDeltaUs();
+	const TArray<double>& CrossingPointUs = Edge.GetCrossingPointUs();
+
+	TArray<double> Coordinates;
+	Coordinates.SetNum(CrossingPointUs.Num() * 2 - 1);
+	Coordinates[0] = CrossingPointUs[0];
+	for (int32 ICuttingPoint = 1; ICuttingPoint < Edge.GetCrossingPointUs().Num(); ICuttingPoint++)
+	{
+		Coordinates[2 * ICuttingPoint - 1] = (CrossingPointUs[ICuttingPoint - 1] + CrossingPointUs[ICuttingPoint]) * 0.5;
+		Coordinates[2 * ICuttingPoint] = CrossingPointUs[ICuttingPoint];
+	}
+
+	TArray<FCurvePoint> Points3D;
+	Edge.EvaluatePoints(Coordinates, 0, Points3D);
+
+	const TArray<TSharedPtr<FCriterion>>& Criteria = MeshModel.GetCriteria();
+	for (const TSharedPtr<FCriterion>& Criterion : Criteria)
+	{
+		Criterion->ApplyOnEdgeParameters(Edge, CrossingPointUs, Points3D);
+	}
+
+	Edge.SetApplyCriteriaMarker();
+	ActiveEdge.SetApplyCriteriaMarker();
+}
+
+void FParametricFaceMesher::MeshThinZoneSide(FThinZoneSide& Side, bool bFinalMeshing)
+{
+	if (!Side.HasMarker2())
+	{
+		return;
+	}
+
+	if (Side.IsProcessed())
+	{
+		return;
+	}
+	Side.ResetProcessedMarker();
+
+#ifdef DEBUG_MESHTHINSURF
+	F3DDebugSession A(bFinalMeshing && bDisplay, TEXT("SideMesh"));
+#endif
+	for (FTopologicalEdge* Edge : Side.GetEdges())
+	{
+		if (Edge->IsMeshed())
+		{
+			continue;
+		}
+		Mesh(*Edge, bFinalMeshing);
+#ifdef DEBUG_THIN_ZONES
+		if (bFinalMeshing && bDisplay)
+		{
+#ifdef DEBUG_MESHTHINSURF
+			F3DDebugSession A(bDisplay, FString::Printf(TEXT("Edge %d"), Edge->GetId()));
+#endif
+			Display2DWithScale(*Edge);
+			TArray<FPoint> MeshVertices = Edge->GetMesh()->GetNodeCoordinates();
+			MeshVertices.Append(Edge->GetStartVertex()->GetMesh()->GetNodeCoordinates());
+			MeshVertices.Append(Edge->GetEndVertex()->GetMesh()->GetNodeCoordinates());
+			if (MeshVertices.Num())
+			{
+				TArray<double> CoordinatesOfMesh;
+				TArray<FPoint> ProjectedPoints;
+				TArray<FPoint2D> Mesh2DPoints;
+				Edge->ProjectPoints(MeshVertices, CoordinatesOfMesh, ProjectedPoints);
+				Algo::Sort(CoordinatesOfMesh);
+				Edge->Approximate2DPoints(CoordinatesOfMesh, Mesh2DPoints);
+				int32 Index = 0;
+				for (FPoint2D Point : Mesh2DPoints)
+				{
+					DisplayPoint2DWithScale(Point, EVisuProperty::RedPoint, ++Index);
+				}
+			}
+#ifdef DEBUG_MESHTHINSURF
+					//Wait(bDisplay);
+#endif
+		}
+#endif
+	}
+#ifdef DEBUG_MESHTHINSURF
+	//Wait(bDisplay);
+#endif
+}
+
+#ifdef DEBUG_INTERSECTEDGEISOS
+TMap<int32, int32> SurfaceDrawed;
+bool bDisplayIsoCurve = true;
+
+void DebugIntersectEdgeIsos(const FTopologicalFace& Face, const TArray<double>& IsoCoordinates, EIso TypeIso)
+{
+	if (SurfaceDrawed.Find(Face.GetId()) == nullptr)
+	{
+		SurfaceDrawed.Add(Face.GetId(), 0);
+	}
+
+	if (bDisplayIsoCurve && SurfaceDrawed[Face.GetId()] < 2)
+	{
+		SurfaceDrawed[Face.GetId()]++;
+
+		FSurfacicBoundary Bounds = Face.GetBoundary();
+
+		F3DDebugSession _(FString::Printf(TEXT("Iso %s 2D %d"), TypeIso == EIso::IsoU ? TEXT("U") : TEXT("V"), Face.GetId()));
+		if (TypeIso == EIso::IsoU)
+		{
+			for (double U : IsoCoordinates)
+			{
+				FPoint2D Start(U, Bounds[EIso::IsoV].Min);
+				FPoint2D End(U, Bounds[EIso::IsoV].Max);
+				DisplaySegment(Start, End, 0, EVisuProperty::Iso);
+			}
+		}
+		else
+		{
+			for (double V : IsoCoordinates)
+			{
+				FPoint2D Start(Bounds[EIso::IsoU].Min, V);
+				FPoint2D End(Bounds[EIso::IsoU].Max, V);
+				DisplaySegment(Start, End, 0, EVisuProperty::Iso);
+			}
+		}
+	}
+}
+
+#endif
+
+#ifdef CADKERNEL_DEBUG
+
+void FParametricFaceMesher::DisplayMeshOfFaceLoop()
+{
+	if (bDisplay)
+	{
+		F3DDebugSession _(bDisplay, FString::Printf(TEXT("Mesh of Thin Face %d"), Face.GetId()));
+		for (const TSharedPtr<FTopologicalLoop>& Loop : Face.GetLoops())
+		{
+			for (const FOrientedEdge& OrientedEdge : Loop->GetEdges())
+			{
+				FTopologicalEdge& Edge = *OrientedEdge.Entity;
+				if (Edge.GetLinkActiveEdge()->IsMeshed())
+				{
+					const TArray<FPoint>& MeshVetrices = Edge.GetMesh()->GetNodeCoordinates();
+					TArray<FPoint> ProjectedPoints;
+					TArray<double> Coordinates;
+					Edge.ProjectPoints(MeshVetrices, Coordinates, ProjectedPoints);
+					Algo::Sort(Coordinates);
+					Coordinates.Insert(Edge.GetBoundary().GetMin(), 0);
+					Coordinates.Add(Edge.GetBoundary().GetMax());
+					TArray<FPoint2D> Points2D;
+					Edge.Approximate2DPoints(Coordinates, Points2D);
+
+					DisplayPolylineWithScale(Points2D, EVisuProperty::BlueCurve);
+					for (FPoint2D Point : Points2D)
+					{
+						DisplayPoint2DWithScale(Point);
+					}
+				}
+				else
+				{
+					Display2DWithScale(Edge, EVisuProperty::GreenCurve);
+				}
+			}
+		}
+		//Wait();
+	}
+}
+
+void FParametricFaceMesher::DisplayThinZoneEdges(const TCHAR* Text, TArray<FTopologicalEdge*>& Edges, EVisuProperty Color, EVisuProperty Color2)
+{
+	if (bDisplay)
+	{
+		F3DDebugSession _(bDisplay, Text);
+		for (const FTopologicalEdge* Edge : Edges)
+		{
+			if (Edge->GetLinkActiveEdge()->IsMeshed())
+			{
+				const TArray<FPoint>& MeshVetrices = Edge->GetMesh()->GetNodeCoordinates();
+				TArray<FPoint> ProjectedPoints;
+				TArray<double> Coordinates;
+				Edge->ProjectPoints(MeshVetrices, Coordinates, ProjectedPoints);
+				Algo::Sort(Coordinates);
+				Coordinates.Insert(Edge->GetBoundary().GetMin(), 0);
+				Coordinates.Add(Edge->GetBoundary().GetMax());
+				TArray<FPoint2D> Points2D;
+				Edge->Approximate2DPoints(Coordinates, Points2D);
+
+				DisplayPolylineWithScale(Points2D, Color2);
+				for (FPoint2D Point : Points2D)
+				{
+					DisplayPoint2DWithScale(Point, (EVisuProperty)(Color2 - 1));
+				}
+			}
+			else
+			{
+				Display2DWithScale(*Edge, Color);
+			}
+		}
+	}
+}
+
+
+#endif
+} // namespace
\ No newline at end of file
diff --git a/Engine/Source/Runtime/Datasmith/CADKernel/Private/CADKernel/Mesh/Meshers/ParametricMesher.cpp b/Engine/Source/Runtime/Datasmith/CADKernel/Private/CADKernel/Mesh/Meshers/ParametricMesher.cpp
--- a/Engine/Source/Runtime/Datasmith/CADKernel/Private/CADKernel/Mesh/Meshers/ParametricMesher.cpp
+++ b/Engine/Source/Runtime/Datasmith/CADKernel/Private/CADKernel/Mesh/Meshers/ParametricMesher.cpp
@@ -4,35 +4,44 @@
 
 #include "CADKernel/Mesh/Criteria/CriteriaGrid.h"
 #include "CADKernel/Mesh/Criteria/Criterion.h"
-#include "CADKernel/Mesh/Meshers/IsoTriangulator.h"
-#include "CADKernel/Mesh/Meshers/MesherTools.h"
-#include "CADKernel/Mesh/Structure/EdgeMesh.h"
-#include "CADKernel/Mesh/Structure/FaceMesh.h"
+#include "CADKernel/Mesh/Meshers/Mesher.h"
+#include "CADKernel/Mesh/Meshers/ParametricFaceMesher.h"
 #include "CADKernel/Mesh/Structure/ModelMesh.h"
 #include "CADKernel/Mesh/Structure/ThinZone2DFinder.h"
-#include "CADKernel/Mesh/Structure/VertexMesh.h"
-#include "CADKernel/Topo/Body.h"
-#include "CADKernel/Topo/Model.h"
-#include "CADKernel/Topo/Shell.h"
 #include "CADKernel/Topo/TopologicalFace.h"
-#include "CADKernel/Utils/Util.h"
+#include "CADKernel/Topo/TopologicalShapeEntity.h"
 
-#ifndef CADKERNEL_DEV
 #include "Async/ParallelFor.h"
-#endif
 
 #ifdef CADKERNEL_DEV
 #include "CADKernel/Mesh/Meshers/MesherReport.h"
 #endif
 
-//#define DEBUG_MESH_EDGE
-//#define DEBUG_GETPREFERREDUVCOORDINATESFROMNEIGHBOURS
 namespace UE::CADKernel
 {
 
-FParametricMesher::FParametricMesher(FModelMesh& InMeshModel)
-	: MeshModel(InMeshModel)
-	, bThinZoneMeshing(false)
+#ifdef CADKERNEL_DEV
+TUniquePtr<FMesherReport> FMesherReport::MesherReport;
+#endif
+
+FMesher::FMesher(FModelMesh& InMeshModel, double InGeometricTolerance, bool bActivateThinZoneMeshing)
+	: GeometricTolerance(InGeometricTolerance)
+	, bThinZoneMeshing(bActivateThinZoneMeshing)
+	, MeshModel(InMeshModel)
+{
+}
+
+void FMesher::MeshEntities(TArray<FTopologicalShapeEntity*>& InEntities)
+{
+	FParametricMesher Mesher(MeshModel, GeometricTolerance, bThinZoneMeshing);
+	Mesher.MeshEntities(InEntities);
+}
+
+
+FParametricMesher::FParametricMesher(FModelMesh& InMeshModel, double GeometricTolerance, bool bActivateThinZoneMeshing)
+	: Tolerances(GeometricTolerance)
+	, bThinZoneMeshing(bActivateThinZoneMeshing)
+	, MeshModel(InMeshModel)
 {
 }
 
@@ -98,7 +107,6 @@ void FParametricMesher::PreMeshingTasks()
 {
 	TRACE_CPUPROFILER_EVENT_SCOPE(ParametricMesher::PreMeshingTasks);
 
-	FTimePoint StartTime = FChrono::Now();
 	FTimePoint ApplyCriteriaStartTime = FChrono::Now();
 
 	FProgress ProgressBar(Faces.Num() * 2, TEXT("Meshing Entities : Apply Surface Criteria"));
@@ -108,32 +116,36 @@ void FParametricMesher::PreMeshingTasks()
 	// ============================================================================================================
 	//      Apply Surface Criteria
 	// ============================================================================================================
+	const double MeshingTolerance = Tolerances.MeshingTolerance;
 
 	TArray<FTopologicalFace*>& LocalFaces = Faces;
-#ifndef CADKERNEL_DEV
-	ParallelFor(Faces.Num(), [&LocalFaces, &Criteria, &bThinZone = bThinZoneMeshing](int32 Index)
-#else
-	bool bThinZone = bThinZoneMeshing;
-	for (int32 Index = 0; Index < Faces.Num(); ++Index)
-#endif
-	{
+	ParallelFor(Faces.Num(), [&LocalFaces, &Criteria, &MeshingTolerance, &bThinZone = bThinZoneMeshing](int32 Index) {
+
+		FProgress _(1, TEXT("Meshing Entities : Apply Surface Criteria"));
+
 		FTopologicalFace* Face = LocalFaces[Index];
 		if (Face == nullptr || Face->IsNotMeshable())
 		{
 			return;
 		}
-		ApplyFaceCriteria(*Face, Criteria, bThinZone);
+		ApplyFaceCriteria(*Face, Criteria, MeshingTolerance, bThinZone);
 		if (!Face->IsDeletedOrDegenerated())
 		{
 			Face->ComputeSurfaceSideProperties();
 		}
-	}
-#ifndef CADKERNEL_DEV
-	);
-#endif
+	});
 
 #ifdef CADKERNEL_DEV
-	MesherReport.Chronos.ApplyCriteriaDuration = FChrono::Elapse(ApplyCriteriaStartTime);
+	FMesherReport::GetChronos().ApplyCriteriaDuration = FChrono::Elapse(ApplyCriteriaStartTime);
+
+	for (FTopologicalFace* Face : Faces)
+	{
+		if (Face->IsDegenerated())
+		{
+			FMesherReport::GetLogs().AddDegeneratedGrid();
+		}
+		FMesherReport::GetLogs().AddThinZone(Face->GetThinZones().Num());
+	}
 #endif
 }
 
@@ -145,32 +157,29 @@ void FParametricMesher::MeshEntities()
 
 	FProgress ProgressBar(Faces.Num() * 2, TEXT("Meshing Entities : Find quad surfaces"));
 
-#ifdef CADKERNEL_DEV
-	MesherReport.Chronos.ApplyCriteriaDuration = FChrono::Elapse(StartTime);
-#endif
-
 	FTimePoint MeshingStartTime = FChrono::Now();
 
 	// ============================================================================================================
 	//      Find and sort quad surfaces 
 	// ============================================================================================================
 
-	TArray<FCostToFace> QuadTrimmedSurfaceSet;
+	TArray<ParametricMesherTool::FCostToFace> QuadTrimmedSurfaceSet;
 
 	if (Faces.Num() > 1)
 	{
+		FTimePoint IsolateQuadPatchStartTime = FChrono::Now();
 		TArray<FTopologicalFace*> OtherEntities;
 
 		FMessage::Printf(Log, TEXT("  Isolate QuadPatch\n"));
-		FTimePoint IsolateQuadPatchStartTime = FChrono::Now();
 
 		IsolateQuadFace(QuadTrimmedSurfaceSet, OtherEntities);
 
+		FMessage::Printf(Log, TEXT("  %d Quad Surfaces found\n"), QuadTrimmedSurfaceSet.Num());
+
 #ifdef CADKERNEL_DEV
-		MesherReport.Chronos.IsolateQuadPatchDuration = FChrono::Elapse(IsolateQuadPatchStartTime);
+		FMesherReport::GetLogs().AddQuadPatch(QuadTrimmedSurfaceSet.Num());
+		FMesherReport::GetChronos().IsolateQuadPatchDuration = FChrono::Elapse(IsolateQuadPatchStartTime);
 #endif
-
-		FMessage::Printf(Log, TEXT("  %d Quad Surfaces found\n"), QuadTrimmedSurfaceSet.Num());
 	}
 
 	// ============================================================================================================
@@ -179,14 +188,15 @@ void FParametricMesher::MeshEntities()
 
 	FTimePoint MeshStartTime = FChrono::Now();
 	MeshSurfaceByFront(QuadTrimmedSurfaceSet);
+
 #ifdef CADKERNEL_DEV
-	MesherReport.Chronos.GlobalMeshDuration = FChrono::Elapse(MeshStartTime);
-	MesherReport.Chronos.GlobalDuration = FChrono::Elapse(StartTime);
+	FMesherReport::GetChronos().GlobalMeshDuration = FChrono::Elapse(MeshStartTime);
+	FMesherReport::GetChronos().GlobalDuration = FChrono::Elapse(StartTime);
 #endif
 
 }
 
-void FParametricMesher::ApplyFaceCriteria(FTopologicalFace& Face, const TArray<TSharedPtr<FCriterion>>& Criteria, bool bThinZoneMeshing)
+void FParametricMesher::ApplyFaceCriteria(FTopologicalFace& Face, const TArray<TSharedPtr<FCriterion>>& Criteria, const double MeshingTolerance, bool bThinZoneMeshing)
 {
 	if (Face.IsApplyCriteria())
 	{
@@ -200,12 +210,20 @@ void FParametricMesher::ApplyFaceCriteria(FTopologicalFace& Face, const TArray<T
 	}
 
 	FCriteriaGrid Grid(Face);
+	Grid.ComputeFaceMinMaxThicknessAlongIso();
+
+	if (Grid.CheckIfIsDegenerate())
+	{
+		Face.Remove();
+		return;
+	}
 
 	Face.InitDeltaUs();
 	Face.ApplyCriteria(Criteria, Grid);
 
 	if (bThinZoneMeshing)
 	{
+		FTimePoint ThinZone2DFinderStartTime = FChrono::Now();
 		Grid.ScaleGrid();
 
 #ifdef DEBUG_ONLY_SURFACE_TO_DEBUG
@@ -222,13 +240,23 @@ void FParametricMesher::ApplyFaceCriteria(FTopologicalFace& Face, const TArray<T
 		// With this ratio each edges of the mesh should be defined by at least 3 segments. 
 		// This should ensure to identified all thin zones according to the mesh size and minimizing the size of the loop sampling
 		constexpr double ElementRatio = 3.;
-		const double Size = Face.GetEstimatedMinimalElementLength() / ElementRatio;
-		const bool bHasThinZones = ThinZoneFinder.SearchThinZones(Size);
+		const double SizeVsElementLength = Face.GetEstimatedMinimalElementLength() / ElementRatio;
+
+		const double FaceSize = Grid.GetCharacteristicThicknessOfFace();
+		const double SizeVsLoop = FaceSize / (double)Face.LoopCount();
+
+		const double SizeVsThickness = FaceSize / MeshConst::ElementRatioVsThickness;
+
+		const bool bHasThinZones = ThinZoneFinder.SearchThinZones(FMath::Max(FMath::Min3(SizeVsElementLength, SizeVsLoop, SizeVsThickness), MeshingTolerance));
 		if (bHasThinZones)
 		{
-			Face.SetHasThinZone();
+			Face.SetHasThinZoneMarker();
 			Face.MoveThinZones(ThinZoneFinder.GetThinZones());
 		}
+
+#ifdef CADKERNEL_DEV
+		FMesherReport::GetChronos().FindThinSurfaceDuration += FChrono::Elapse(ThinZone2DFinderStartTime);
+#endif
 	}
 
 	if (Face.IsDegenerated())
@@ -240,7 +268,16 @@ void FParametricMesher::ApplyFaceCriteria(FTopologicalFace& Face, const TArray<T
 void FParametricMesher::ApplyEdgeCriteria(FTopologicalEdge& Edge)
 {
 	FTopologicalEdge& ActiveEdge = *Edge.GetLinkActiveEdge();
-	ensureCADKernel(Edge.IsVirtuallyMeshed() || !ActiveEdge.IsApplyCriteria());
+
+	if (Edge.Length() < Edge.GetTolerance3D())
+	{
+		Edge.SetAsDegenerated();
+	}
+
+	if (ActiveEdge.IsApplyCriteria())
+	{
+		return;
+	}
 
 	Edge.ComputeCrossingPointCoordinates();
 	Edge.InitDeltaUs();
@@ -270,616 +307,11 @@ void FParametricMesher::ApplyEdgeCriteria(FTopologicalEdge& Edge)
 
 void FParametricMesher::Mesh(FTopologicalFace& Face)
 {
-#ifdef DEBUG_ONLY_SURFACE_TO_DEBUG
-	bDisplay = (Face.GetId() == FaceToDebug);
-#endif
-
-	if (Face.IsNotMeshable())
-	{
-		return;
-	}
-
-	FMessage::Printf(EVerboseLevel::Debug, TEXT("Meshing of surface %d\n"), Face.GetId());
-
-	FProgress _(1, TEXT("Meshing Entities : Mesh Surface"));
-
-#ifdef DEBUG_CADKERNEL
-	F3DDebugSession S(bDisplayDebugMeshStep, FString::Printf(TEXT("Mesh of surface %d"), Face.GetId()));
-#endif
-
-	FTimePoint StartTime = FChrono::Now();
-
-	FTimePoint GenerateCloudStartTime = FChrono::Now();
-
-	FGrid Grid(Face, MeshModel);
-	if (!GenerateCloud(Grid) || Grid.IsDegenerated())
-	{
-#ifdef CADKERNEL_DEV
-		MesherReport.Logs.AddDegeneratedGrid();
-#endif
-		FMessage::Printf(EVerboseLevel::Log, TEXT("The meshing of the surface %d failed due to a degenerated grid\n"), Face.GetId());
-		Face.SetMeshed();
-		return;
-	}
-
-	FDuration GenerateCloudDuration = FChrono::Elapse(GenerateCloudStartTime);
-
-	FTimePoint IsoTriangulerStartTime = FChrono::Now();
-
-	TSharedRef<FFaceMesh> SurfaceMesh = StaticCastSharedRef<FFaceMesh>(Face.GetOrCreateMesh(MeshModel));
-
-	FIsoTriangulator IsoTrianguler(Grid, SurfaceMesh);
-#ifdef CADKERNEL_DEV
-	IsoTrianguler.SetMesherReport(MesherReport);
-#endif
-
-	if (IsoTrianguler.Triangulate())
-	{
-		if (Face.IsBackOriented())
-		{
-			SurfaceMesh->InverseOrientation();
-		}
-		MeshModel.AddMesh(SurfaceMesh);
-	}
-	Face.SetMeshed();
-
-	FDuration TriangulateDuration = FChrono::Elapse(IsoTriangulerStartTime);
-	FDuration Duration = FChrono::Elapse(StartTime);
-
-#ifdef CADKERNEL_DEV
-	MesherReport.Chronos.GlobalPointCloudDuration += Grid.Chronos.GeneratePointCloudDuration;
-	MesherReport.Chronos.GlobalGeneratePointCloudDuration += GenerateCloudDuration;
-	MesherReport.Chronos.GlobalTriangulateDuration += TriangulateDuration;
-	MesherReport.Chronos.GlobalDelaunayDuration += IsoTrianguler.Chronos.FindSegmentToLinkLoopToLoopByDelaunayDuration;
-	MesherReport.Chronos.GlobalMeshDuration += Duration;
-#endif
-
-	FChrono::PrintClockElapse(EVerboseLevel::Debug, TEXT("   "), TEXT("Meshing"), Duration);
-
-}
-
-bool FParametricMesher::GenerateCloud(FGrid& Grid)
-{
-	Grid.DefineCuttingParameters();
-	if (!Grid.GeneratePointCloud())
-	{
-		return false;
-	}
-
-	if (bThinZoneMeshing)
-	{
-		FTimePoint StartTime = FChrono::Now();
-		if (Grid.GetFace().HasThinZone())
-		{
-			MeshThinZoneEdges(Grid.GetFace());
-		}
-#ifdef CADKERNEL_DEV
-		MesherReport.Chronos.GlobalThinZones += FChrono::Elapse(StartTime);
-#endif
-	}
-
-	FTimePoint StartTime = FChrono::Now();
-	MeshFaceLoops(Grid);
-
-	Grid.ProcessPointCloud();
-
-#ifdef CADKERNEL_DEV
-	MesherReport.Chronos.GlobalMeshAndGetLoopNodes += FChrono::Elapse(StartTime);
-#endif
-
-	return true;
-}
-
-void FParametricMesher::MeshFaceLoops(FGrid& Grid)
-{
-	const FTopologicalFace& Face = Grid.GetFace();
-
-	FTimePoint StartTime = FChrono::Now();
-
-	for (const TSharedPtr<FTopologicalLoop>& Loop : Face.GetLoops())
-	{
-		for (const FOrientedEdge& Edge : Loop->GetEdges())
-		{
-			Mesh(*Edge.Entity, Face);
-		}
-	}
-
-#ifdef CADKERNEL_DEV
-	MesherReport.Chronos.GlobalMeshEdges += FChrono::Elapse(StartTime);
-#endif
-}
-
-static void FillImposedIsoCuttingPoints(TArray<double>& UEdgeSetOfIntersectionWithIso, ECoordinateType CoordinateType, double EdgeToleranceGeo, const FTopologicalEdge& Edge, TArray<FCuttingPoint>& OutImposedIsoVertexSet)
-{
-	FLinearBoundary EdgeBoundary = Edge.GetBoundary();
-
-	int32 StartIndex = OutImposedIsoVertexSet.Num();
-	Algo::Sort(UEdgeSetOfIntersectionWithIso);
-	double PreviousU = -HUGE_VALUE;
-	for (double InterU : UEdgeSetOfIntersectionWithIso)
-	{
-		// Remove coordinate nearly equal to boundary
-		if ((InterU - EdgeToleranceGeo) < EdgeBoundary.GetMin() || (InterU + EdgeToleranceGeo) > EdgeBoundary.GetMax())
-		{
-			continue;
-		}
-
-		// Remove coordinate inside thin zone
-		for (FLinearBoundary ThinZone : Edge.GetThinZoneBounds())
-		{
-			if (ThinZone.Contains(InterU))
-			{
-				continue;
-			}
-		}
-
-		// Remove nearly duplicate 
-		if (InterU - PreviousU < EdgeToleranceGeo)
-		{
-			continue;
-		}
-
-		OutImposedIsoVertexSet.Emplace(InterU, CoordinateType);
-		PreviousU = InterU;
-	}
-
-	int32 Index;
-	int32 NewCoordinateCount = OutImposedIsoVertexSet.Num() - StartIndex;
-	switch (NewCoordinateCount)
-	{
-	case 0:
-		return;
-
-	case 1:
-	{
-		int32 CuttingPointIndex = 0;
-		while (CuttingPointIndex < Edge.GetCrossingPointUs().Num() && Edge.GetCrossingPointUs()[CuttingPointIndex] + DOUBLE_SMALL_NUMBER <= OutImposedIsoVertexSet[StartIndex].Coordinate)
-		{
-			++CuttingPointIndex;
-		};
-		if (CuttingPointIndex > 0)
-		{
-			--CuttingPointIndex;
-		}
-		OutImposedIsoVertexSet[StartIndex].IsoDeltaU = Edge.GetDeltaUMaxs()[CuttingPointIndex] * AQuarter;
-		break;
-	}
-
-	default:
-	{
-		OutImposedIsoVertexSet[StartIndex].IsoDeltaU = (OutImposedIsoVertexSet[StartIndex + 1].Coordinate - OutImposedIsoVertexSet[StartIndex].Coordinate) * AQuarter;
-		for (Index = StartIndex + 1; Index < OutImposedIsoVertexSet.Num() - 1; ++Index)
-		{
-			OutImposedIsoVertexSet[Index].IsoDeltaU = (OutImposedIsoVertexSet[Index + 1].Coordinate - OutImposedIsoVertexSet[Index - 1].Coordinate) * AEighth;
-		}
-		OutImposedIsoVertexSet[Index].IsoDeltaU = (OutImposedIsoVertexSet[Index].Coordinate - OutImposedIsoVertexSet[Index - 1].Coordinate) * AQuarter;
-		break;
-	}
-
-	}
-}
-
-void FParametricMesher::Mesh(FTopologicalVertex& InVertex)
-{
-	InVertex.GetOrCreateMesh(GetMeshModel());
-}
-
-//#define DEBUG_INTERSECTEDGEISOS
-#ifdef DEBUG_INTERSECTEDGEISOS
-void DebugIntersectEdgeIsos(const FTopologicalFace& Face, const TArray<double>& IsoCoordinates, EIso TypeIso);
-#endif
-
-void FParametricMesher::Mesh(FTopologicalEdge& InEdge, const FTopologicalFace& Face)
-{
-	{
-		FTopologicalEdge& ActiveEdge = *InEdge.GetLinkActiveEntity();
-		if (ActiveEdge.IsMeshed())
-		{
-			if (ActiveEdge.GetMesh()->GetNodeCount() > 0)
-			{
-				return;
-			}
-
-			// In some case the 2d curve is a smooth curve and the 3d curve is a line and vice versa
-			// In the particular case where the both case are opposed, we can have the 2d line sampled with 4 points, and the 2d curve sampled with 2 points (because in 3d, the 2d curve is a 3d line)
-			// In this case, the loop is flat i.e. in 2d the meshes of the 2d line and 2d curve are coincident
-			// So the grid is degenerated and the surface is not meshed
-			// to avoid this case, the Edge is virtually meshed i.e. the nodes inside the edge have the id of the mesh of the vertices.
-			InEdge.SetVirtuallyMeshedMarker();
-		}
-
-		if (ActiveEdge.IsThinPeak())
-		{
-			TArray<FCuttingPoint>& FinalEdgeCuttingPointCoordinates = ActiveEdge.GetCuttingPoints();
-			FinalEdgeCuttingPointCoordinates.Emplace(ActiveEdge.GetStartCurvilinearCoordinates(), ECoordinateType::VertexCoordinate);
-			FinalEdgeCuttingPointCoordinates.Emplace(ActiveEdge.GetEndCurvilinearCoordinates(), ECoordinateType::VertexCoordinate);
-			GenerateEdgeElements(ActiveEdge);
-			return;
-		}
-	}
-
-	const FSurfacicTolerance& ToleranceIso = Face.GetIsoTolerances();
-
-	// Get Edge intersection with inner surface mesh grid
-	TArray<double> EdgeIntersectionWithIsoU_Coordinates;
-	TArray<double> EdgeIntersectionWithIsoV_Coordinates;
-
-	const TArray<double>& SurfaceTabU = Face.GetCuttingCoordinatesAlongIso(EIso::IsoU);
-	const TArray<double>& SurfaceTabV = Face.GetCuttingCoordinatesAlongIso(EIso::IsoV);
-
-	ApplyEdgeCriteria(InEdge);
-
-#ifdef DEBUG_MESH_EDGE
-	if (bDisplay)
-	{
-		F3DDebugSession _(FString::Printf(TEXT("EdgePointsOnDomain %d"), InEdge.GetId()));
-		Display2D(InEdge);
-		Wait();
-	}
-#endif
-
-#ifdef DEBUG_INTERSECTEDGEISOS
-	DebugIntersectEdgeIsos(Face, SurfaceTabU, EIso::IsoU);
-	DebugIntersectEdgeIsos(Face, SurfaceTabV, EIso::IsoV);
-	{
-		F3DDebugSession _(FString::Printf(TEXT("Edge 2D %d"), InEdge.GetId()));
-		UE::CADKernel::Display2D(InEdge);
-	}
-#endif
-
-	InEdge.ComputeIntersectionsWithIsos(SurfaceTabU, EIso::IsoU, ToleranceIso, EdgeIntersectionWithIsoU_Coordinates);
-	InEdge.ComputeIntersectionsWithIsos(SurfaceTabV, EIso::IsoV, ToleranceIso, EdgeIntersectionWithIsoV_Coordinates);
-
-#ifdef DEBUG_INTERSECTEDGEISOS
-	{
-		F3DDebugSession _(FString::Printf(TEXT("Edge %d Intersect with iso"), InEdge.GetId()));
-		TArray<FPoint2D> Intersections;
-		InEdge.Approximate2DPoints(EdgeIntersectionWithIsoU_Coordinates, Intersections);
-		for (const FPoint2D& Point : Intersections)
-		{
-			DisplayPoint(Point);
-		}
-		Intersections.Empty();
-		InEdge.Approximate2DPoints(EdgeIntersectionWithIsoV_Coordinates, Intersections);
-		for (const FPoint2D& Point : Intersections)
-		{
-			DisplayPoint(Point);
-		}
-		Wait();
-	}
-	{
-		F3DDebugSession _(FString::Printf(TEXT("Thin Zone")));
-		for (FLinearBoundary ThinZone : InEdge.GetThinZoneBounds())
-		{
-			TArray<double> Coords;
-			Coords.Add(ThinZone.GetMin());
-			Coords.Add(ThinZone.GetMax());
-			TArray<FPoint2D> ThinZone2D;
-			InEdge.Approximate2DPoints(Coords, ThinZone2D);
-			DisplaySegment(ThinZone2D[0], ThinZone2D[1], EVisuProperty::YellowCurve);
-		}
-		Wait();
-
-	}
-#endif
-
-	FLinearBoundary EdgeBounds = InEdge.GetBoundary();
-
-	TArray<double>& DeltaUs = InEdge.GetDeltaUMaxs();
-
-	InEdge.SortImposedCuttingPoints();
-	const TArray<FImposedCuttingPoint>& EdgeImposedCuttingPoints = InEdge.GetImposedCuttingPoints();
-
-	// build a edge mesh compiling inner surface cutting (based on criteria applied on the surface) and edge cutting (based on criteria applied on the curve)
-	TArray<FCuttingPoint> ImposedIsoCuttingPoints;
-	{
-		int32 NbImposedCuttingPoints = EdgeImposedCuttingPoints.Num() + EdgeIntersectionWithIsoU_Coordinates.Num() + EdgeIntersectionWithIsoV_Coordinates.Num() + 2;
-		ImposedIsoCuttingPoints.Reserve(NbImposedCuttingPoints);
-	}
-
-	FPoint2D ExtremityTolerances = InEdge.GetCurve()->GetExtremityTolerances(EdgeBounds);
-
-	ImposedIsoCuttingPoints.Emplace(EdgeBounds.GetMin(), ECoordinateType::VertexCoordinate, -1, ExtremityTolerances[0]);
-	ImposedIsoCuttingPoints.Emplace(EdgeBounds.GetMax(), ECoordinateType::VertexCoordinate, -1, ExtremityTolerances[1]);
-
-	int32 Index = 0;
-	for (const FImposedCuttingPoint& CuttingPoint : EdgeImposedCuttingPoints)
-	{
-		double CuttingPointDeltaU = InEdge.GetDeltaUFor(CuttingPoint.Coordinate, Index);
-		ImposedIsoCuttingPoints.Emplace(CuttingPoint.Coordinate, ECoordinateType::ImposedCoordinate, CuttingPoint.OppositNodeIndex, CuttingPointDeltaU * AThird);
-	}
-
-	// Add Edge intersection with inner surface grid Iso
-	double EdgeTolerance = FMath::Min(ExtremityTolerances[0], ExtremityTolerances[1]);
-	if (!EdgeIntersectionWithIsoU_Coordinates.IsEmpty())
-	{
-		FillImposedIsoCuttingPoints(EdgeIntersectionWithIsoU_Coordinates, IsoUCoordinate, EdgeTolerance, InEdge, ImposedIsoCuttingPoints);
-	}
-
-	if (!EdgeIntersectionWithIsoV_Coordinates.IsEmpty())
-	{
-		FillImposedIsoCuttingPoints(EdgeIntersectionWithIsoV_Coordinates, IsoVCoordinate, EdgeTolerance, InEdge, ImposedIsoCuttingPoints);
-	}
-
-	ImposedIsoCuttingPoints.Sort([](const FCuttingPoint& Point1, const FCuttingPoint& Point2) { return Point1.Coordinate < Point2.Coordinate; });
-
-	TFunction<void(int32&, int32&, ECoordinateType)> MergeImposedCuttingPoints = [&](int32& Index, int32& NewIndex, ECoordinateType NewType)
-	{
-		double DeltaU = FMath::Max(ImposedIsoCuttingPoints[NewIndex].IsoDeltaU, ImposedIsoCuttingPoints[Index].IsoDeltaU);
-		if (ImposedIsoCuttingPoints[NewIndex].Type <= ImposedCoordinate && ImposedIsoCuttingPoints[Index].Type <= ImposedCoordinate)
-		{
-			DeltaU /= 5.;
-		}
-
-		if (ImposedIsoCuttingPoints[NewIndex].Coordinate + DeltaU > ImposedIsoCuttingPoints[Index].Coordinate)
-		{
-			if (ImposedIsoCuttingPoints[Index].Type == VertexCoordinate)
-			{
-				ImposedIsoCuttingPoints[NewIndex].Coordinate = ImposedIsoCuttingPoints[Index].Coordinate;
-				ImposedIsoCuttingPoints[NewIndex].IsoDeltaU = ImposedIsoCuttingPoints[Index].IsoDeltaU;
-			}
-			else if (ImposedIsoCuttingPoints[NewIndex].Type == VertexCoordinate)
-			{
-			}
-			else if (ImposedIsoCuttingPoints[NewIndex].Type == ImposedCoordinate)
-			{
-				if (ImposedIsoCuttingPoints[Index].Type == ImposedCoordinate)
-				{
-					ImposedIsoCuttingPoints[NewIndex].Coordinate = (ImposedIsoCuttingPoints[NewIndex].Coordinate + ImposedIsoCuttingPoints[Index].Coordinate) * 0.5;
-				}
-			}
-			else if (ImposedIsoCuttingPoints[Index].Type == ImposedCoordinate)
-			{
-				ImposedIsoCuttingPoints[NewIndex].Coordinate = ImposedIsoCuttingPoints[Index].Coordinate;
-				ImposedIsoCuttingPoints[NewIndex].Type = ImposedCoordinate;
-				ImposedIsoCuttingPoints[NewIndex].IsoDeltaU = ImposedIsoCuttingPoints[Index].IsoDeltaU;
-			}
-			else if (ImposedIsoCuttingPoints[NewIndex].Type != ImposedIsoCuttingPoints[Index].Type)
-			{
-				ImposedIsoCuttingPoints[NewIndex].Coordinate = (ImposedIsoCuttingPoints[NewIndex].Coordinate + ImposedIsoCuttingPoints[Index].Coordinate) * 0.5;
-				ImposedIsoCuttingPoints[NewIndex].Type = IsoUVCoordinate;
-				ImposedIsoCuttingPoints[NewIndex].IsoDeltaU = FMath::Min(ImposedIsoCuttingPoints[NewIndex].IsoDeltaU, ImposedIsoCuttingPoints[Index].IsoDeltaU);
-			}
-
-			if (ImposedIsoCuttingPoints[NewIndex].Type <= ImposedCoordinate)
-			{
-				if (ImposedIsoCuttingPoints[NewIndex].OppositNodeIndex == -1)
-				{
-					ImposedIsoCuttingPoints[NewIndex].OppositNodeIndex = ImposedIsoCuttingPoints[Index].OppositNodeIndex;
-				}
-				else
-				{
-					ImposedIsoCuttingPoints[NewIndex].OppositNodeIndex2 = ImposedIsoCuttingPoints[Index].OppositNodeIndex;
-				}
-			}
-		}
-		else
-		{
-			++NewIndex;
-			ImposedIsoCuttingPoints[NewIndex] = ImposedIsoCuttingPoints[Index];
-		}
-	};
-
-	// If a pair of point isoU/isoV is too close, get the middle of the points
-	if (ImposedIsoCuttingPoints.Num() > 1)
-	{
-		int32 NewIndex = 0;
-		for (int32 Andex = 1; Andex < ImposedIsoCuttingPoints.Num(); ++Andex)
-		{
-			if (ImposedIsoCuttingPoints[Andex].Type > ECoordinateType::ImposedCoordinate)
-			{
-				bool bIsDelete = false;
-				for (const FLinearBoundary& ThinZone : InEdge.GetThinZoneBounds())
-				{
-					if (ThinZone.Contains(ImposedIsoCuttingPoints[Andex].Coordinate))
-					{
-						bIsDelete = true;
-						break; // or copntinue
-					}
-				}
-				if (bIsDelete)
-				{
-					continue;
-				}
-			}
-
-			if (ImposedIsoCuttingPoints[NewIndex].Type == ECoordinateType::ImposedCoordinate || ImposedIsoCuttingPoints[Andex].Type == ECoordinateType::ImposedCoordinate)
-			{
-				MergeImposedCuttingPoints(Andex, NewIndex, ECoordinateType::ImposedCoordinate);
-			}
-			else if (ImposedIsoCuttingPoints[NewIndex].Type != ImposedIsoCuttingPoints[Andex].Type)
-			{
-				MergeImposedCuttingPoints(Andex, NewIndex, ECoordinateType::IsoUVCoordinate);
-			}
-			else
-			{
-				++NewIndex;
-				ImposedIsoCuttingPoints[NewIndex] = ImposedIsoCuttingPoints[Andex];
-			}
-		}
-		ImposedIsoCuttingPoints.SetNum(NewIndex + 1);
-	}
-
-	if (ImposedIsoCuttingPoints.Num() > 1 && (EdgeBounds.GetMax() - ImposedIsoCuttingPoints.Last().Coordinate) < FMath::Min(ImposedIsoCuttingPoints.Last().IsoDeltaU, InEdge.GetDeltaUMaxs().Last()))
-	{
-		ImposedIsoCuttingPoints.Last().Coordinate = EdgeBounds.GetMax();
-		ImposedIsoCuttingPoints.Last().Type = VertexCoordinate;
-	}
-	else
-	{
-		ImposedIsoCuttingPoints.Emplace(EdgeBounds.GetMax(), ECoordinateType::VertexCoordinate, -1, InEdge.GetDeltaUMaxs().Last() * AQuarter);
-	}
-
-	// Final array of the edge mesh vertex 
-	TArray<FCuttingPoint>& FinalEdgeCuttingPointCoordinates = InEdge.GetCuttingPoints();
-	{
-		// max count of vertex
-		double MinDeltaU = HUGE_VALUE;
-		for (const double& DeltaU : DeltaUs)
-		{
-			if (DeltaU < MinDeltaU)
-			{
-				MinDeltaU = DeltaU;
-			}
-		}
-
-		int32 MaxNumberOfVertex = FMath::IsNearlyZero(MinDeltaU) ? 5 : (int32)((EdgeBounds.GetMax() - EdgeBounds.GetMin()) / MinDeltaU) + 5;
-		FinalEdgeCuttingPointCoordinates.Empty(ImposedIsoCuttingPoints.Num() + MaxNumberOfVertex);
-	}
-
-#ifdef DEBUG_GETPREFERREDUVCOORDINATESFROMNEIGHBOURS
-	TArray<double> CuttingPoints2;
-	{
-		double ToleranceGeoEdge = FMath::Min(ExtremityTolerances[0], ExtremityTolerances[1]);
-
-		TArray<FCuttingPoint> Extremities;
-		Extremities.Reserve(2);
-		Extremities.Emplace(EdgeBounds.GetMin(), ECoordinateType::VertexCoordinate, -1, ToleranceGeoEdge);
-		Extremities.Emplace(EdgeBounds.GetMax(), ECoordinateType::VertexCoordinate, -1, ToleranceGeoEdge);
-
-		FMesherTools::ComputeFinalCuttingPointsWithImposedCuttingPoints(InEdge.GetCrossingPointUs(), InEdge.GetDeltaUMaxs(), Extremities, CuttingPoints2);
-	}
-#endif
-
-	if (InEdge.IsDegenerated() || InEdge.IsVirtuallyMeshed())
-	{
-		if (ImposedIsoCuttingPoints.Num() == 2)
-		{
-			ImposedIsoCuttingPoints.EmplaceAt(1, (ImposedIsoCuttingPoints[0].Coordinate + ImposedIsoCuttingPoints[1].Coordinate) * 0.5, ECoordinateType::OtherCoordinate);
-		}
-
-		for (FCuttingPoint CuttingPoint : ImposedIsoCuttingPoints)
-		{
-			FinalEdgeCuttingPointCoordinates.Emplace(CuttingPoint.Coordinate, ECoordinateType::OtherCoordinate);
-		}
-		InEdge.GetLinkActiveEdge()->SetMeshed();
-	}
-	else
-	{
-		TArray<double> CuttingPoints;
-		FMesherTools::ComputeFinalCuttingPointsWithImposedCuttingPoints(InEdge.GetCrossingPointUs(), InEdge.GetDeltaUMaxs(), ImposedIsoCuttingPoints, CuttingPoints);
-		int32 ImposedIndex = 0;
-		int32 ImposedIsoCuttingPointsCount = ImposedIsoCuttingPoints.Num();
-		for (const double& Coordinate : CuttingPoints)
-		{
-			if (FMath::IsNearlyEqual(ImposedIsoCuttingPoints[ImposedIndex].Coordinate, Coordinate))
-			{
-				FinalEdgeCuttingPointCoordinates.Emplace(ImposedIsoCuttingPoints[ImposedIndex]);
-				++ImposedIndex;
-			}
-			else
-			{
-				while (ImposedIndex < ImposedIsoCuttingPointsCount && ImposedIsoCuttingPoints[ImposedIndex].Coordinate < Coordinate)
-				{
-					++ImposedIndex;
-				}
-				FinalEdgeCuttingPointCoordinates.Emplace(Coordinate, ECoordinateType::OtherCoordinate);
-			}
-		}
-
-#ifdef DEBUG_GETPREFERREDUVCOORDINATESFROMNEIGHBOURS
-		if (InEdge.IsThinZone() && EdgeImposedCuttingPoints.Num())
-		{
-			F3DDebugSession G(TEXT("Mesh(TSharedRef<FEdge> InEdge"));
-			{
-				F3DDebugSession G(TEXT("U From Iso"));
-				for (const FCuttingPoint& CuttingU : ImposedIsoCuttingPoints)
-				{
-					if (CuttingU.OppositNodeIndex >= 0)
-					{
-						UE::CADKernel::DisplayPoint(FPoint2D(CuttingU.Coordinate, 0.0), EVisuProperty::RedPoint);
-					}
-					else
-					{
-						UE::CADKernel::DisplayPoint(FPoint2D(CuttingU.Coordinate, 0.0));
-					}
-				}
-			}
-			{
-				F3DDebugSession _(InEdge.GetThinZoneBounds().Num() != 0, FString::Printf(TEXT("Thin Zone"), InEdge.GetId()));
-				for (FLinearBoundary ThinZone : InEdge.GetThinZoneBounds())
-				{
-					UE::CADKernel::DisplayPoint(FPoint2D(ThinZone.GetMin(), 0.01), EVisuProperty::BluePoint);
-					UE::CADKernel::DisplayPoint(FPoint2D(ThinZone.GetMax(), 0.01), EVisuProperty::BluePoint);
-					DisplaySegment(FPoint2D(ThinZone.GetMin(), 0.01), FPoint2D(ThinZone.GetMax(), 0.01), EVisuProperty::BlueCurve);
-				}
-				Wait();
-			}
-
-			{
-				F3DDebugSession G(TEXT("U From Criteria"));
-				for (double CuttingU : CuttingPoints2)
-				{
-					UE::CADKernel::DisplayPoint(FPoint2D(CuttingU, 0.02), EVisuProperty::RedPoint);
-				}
-			}
-			{
-				F3DDebugSession G(TEXT("U Final (Criteria & Iso)"));
-				for (double CuttingU : CuttingPoints)
-				{
-					UE::CADKernel::DisplayPoint(FPoint2D(CuttingU, 0.04), EVisuProperty::YellowPoint);
-				}
-			}
-			Wait();
-		}
-#endif
-
-		GenerateEdgeElements(InEdge);
-	}
-
-}
-
-//#define DEBUG_MESH_EDGE
-void FParametricMesher::GenerateEdgeElements(FTopologicalEdge& Edge)
-{
-	FTopologicalEdge& ActiveEdge = *Edge.GetLinkActiveEntity();
-
-	bool bSameDirection = Edge.IsSameDirection(ActiveEdge);
-
-	TSharedRef<FEdgeMesh> EdgeMesh = ActiveEdge.GetOrCreateMesh(MeshModel);
-
-	int32 StartVertexNodeIndex = ActiveEdge.GetStartVertex()->GetOrCreateMesh(GetMeshModel())->GetMesh();
-	int32 EndVertexNodeIndex = ActiveEdge.GetEndVertex()->GetOrCreateMesh(GetMeshModel())->GetMesh();
-
-	TArray<double> CuttingPointCoordinates;
-	CuttingPointCoordinates.Reserve(Edge.GetCuttingPoints().Num());
-	for (const FCuttingPoint& CuttingPoint : Edge.GetCuttingPoints())
-	{
-		CuttingPointCoordinates.Add(CuttingPoint.Coordinate);
-	}
-	ensureCADKernel(CuttingPointCoordinates.Num() > 1);
-	CuttingPointCoordinates.RemoveAt(0);
-	CuttingPointCoordinates.Pop();
-
-	TArray<FPoint>& Coordinates = EdgeMesh->GetNodeCoordinates();
-	Edge.ApproximatePoints(CuttingPointCoordinates, Coordinates);
-
-	if (!bSameDirection)
-	{
-		Algo::Reverse(Coordinates);
-	}
-
-#ifdef DEBUG_MESH_EDGE
-	if (bDisplay)
-	{
-		F3DDebugSession _(FString::Printf(TEXT("Edge Mesh %d"), Edge.GetId()));
-		TArray<FPoint2D> Mesh2D;
-		Edge.Approximate2DPoints(CuttingPointCoordinates, Mesh2D);
-		for (const FPoint2D& Vertex : Mesh2D)
-		{
-			DisplayPoint(Vertex, EVisuProperty::RedPoint);
-		}
-		Wait();
-	}
-#endif
-
-	EdgeMesh->RegisterCoordinates();
-	EdgeMesh->Mesh(StartVertexNodeIndex, EndVertexNodeIndex);
-	MeshModel.AddMesh(EdgeMesh);
-	ActiveEdge.SetMeshed();
+	FParametricFaceMesher FaceMesher(Face, MeshModel, Tolerances, bThinZoneMeshing);
+	FaceMesher.Mesh();
 }
 
-void FParametricMesher::IsolateQuadFace(TArray<FCostToFace>& QuadSurfaces, TArray<FTopologicalFace*>& OtherSurfaces) const
+void FParametricMesher::IsolateQuadFace(TArray<ParametricMesherTool::FCostToFace>& QuadSurfaces, TArray<FTopologicalFace*>& OtherSurfaces) const
 {
 	TRACE_CPUPROFILER_EVENT_SCOPE(ParametricMesher::IsolateQuadFace);
 
@@ -903,10 +335,10 @@ void FParametricMesher::IsolateQuadFace(TArray<FCostToFace>& QuadSurfaces, TArra
 			double LocalMinCurvature;
 			double LocalMaxCurvature;
 			GetMinMax(Face.GetCurvature(EIso::IsoU).Max, Face.GetCurvature(EIso::IsoV).Max, LocalMinCurvature, LocalMaxCurvature);
-			if (LocalMaxCurvature > ConstMinCurvature)
+			if (LocalMaxCurvature > MeshConst::MinCurvature)
 			{
 				QuadSurfaces.Emplace(LocalMaxCurvature, &Face);
-				if (LocalMinCurvature > ConstMinCurvature)
+				if (LocalMinCurvature > MeshConst::MinCurvature)
 				{
 					QuadSurfaces.Emplace(LocalMinCurvature, &Face);
 				}
@@ -927,7 +359,7 @@ void FParametricMesher::IsolateQuadFace(TArray<FCostToFace>& QuadSurfaces, TArra
 		}
 	}
 
-	Algo::Sort(QuadSurfaces, [](FCostToFace& SurfaceA, FCostToFace& SurfaceB)
+	Algo::Sort(QuadSurfaces, [](ParametricMesherTool::FCostToFace& SurfaceA, ParametricMesherTool::FCostToFace& SurfaceB)
 		{
 			return (SurfaceA.Cost > SurfaceB.Cost);
 		}
@@ -967,13 +399,13 @@ void FParametricMesher::IsolateQuadFace(TArray<FCostToFace>& QuadSurfaces, TArra
 #endif
 }
 
-void FParametricMesher::LinkQuadSurfaceForMesh(TArray<FCostToFace>& QuadTrimmedSurfaceSet, TArray<TArray<FTopologicalFace*>>& OutStrips)
+void FParametricMesher::LinkQuadSurfaceForMesh(TArray<ParametricMesherTool::FCostToFace>& QuadTrimmedSurfaceSet, TArray<TArray<FTopologicalFace*>>& OutStrips)
 {
 	const double GeometricTolerance = 20. * MeshModel.GetGeometricTolerance();
 
 	OutStrips.Reserve(QuadTrimmedSurfaceSet.Num());
 
-	for (FCostToFace& Quad : QuadTrimmedSurfaceSet)
+	for (ParametricMesherTool::FCostToFace& Quad : QuadTrimmedSurfaceSet)
 	{
 		FTopologicalFace* Face = Quad.Face;
 		const FSurfaceCurvature& Curvatures = Face->GetCurvatures();
@@ -1128,9 +560,9 @@ void FParametricMesher::LinkQuadSurfaceForMesh(TArray<FCostToFace>& QuadTrimmedS
 	}
 }
 
-void FParametricMesher::MeshSurfaceByFront(TArray<FCostToFace>& QuadTrimmedSurfaceSet)
+void FParametricMesher::MeshSurfaceByFront(TArray<ParametricMesherTool::FCostToFace>& QuadTrimmedSurfaceSet)
 {
-	// Marker3 : Surfaces that have to be meshed are set Marker3
+	// WaitingMarker : Surfaces that have to be meshed are set WaitingMarker
 	// Marker1 : Surfaces added in CandidateSurfacesForMesh
 	// Marker2 : Surfaces added in SecondChoiceOfCandidateSurfacesForMesh
 
@@ -1142,7 +574,7 @@ void FParametricMesher::MeshSurfaceByFront(TArray<FCostToFace>& QuadTrimmedSurfa
 		{
 			continue;
 		}
-		Face->SetMarker3();
+		Face->SetWaitingMarker();
 	}
 
 	const double GeometricTolerance = 20. * MeshModel.GetGeometricTolerance();
@@ -1196,7 +628,7 @@ void FParametricMesher::MeshSurfaceByFront(TArray<FCostToFace>& QuadTrimmedSurfa
 				}
 
 				FTopologicalFace* NextFace = NextEdge->GetFace();
-				if ((NextFace == nullptr) || !NextFace->HasMarker3())
+				if ((NextFace == nullptr) || !NextFace->IsWaiting())
 				{
 					// not in the scope of surface to mesh
 					continue;
@@ -1347,7 +779,7 @@ void FParametricMesher::MeshSurfaceByFront(TArray<FCostToFace>& QuadTrimmedSurfa
 	};
 
 	// the front is initialized with quad surface
-	for (const FCostToFace& Quad : QuadTrimmedSurfaceSet)
+	for (const ParametricMesherTool::FCostToFace& Quad : QuadTrimmedSurfaceSet)
 	{
 		FTopologicalFace* Face = Quad.Face;
 		MeshFacesByFront(*Face);
@@ -1364,297 +796,4 @@ void FParametricMesher::MeshSurfaceByFront(TArray<FCostToFace>& QuadTrimmedSurfa
 
 }
 
-void FParametricMesher::MeshThinZoneEdges(FTopologicalFace& Face)
-{
-	TArray<FThinZone2D>& ThinZones = Face.GetThinZones();
-
-	{
-#ifdef DEBUG_MESHTHINSURF
-		if(bDisplay)
-		{
-			F3DDebugSession _(bDisplay, FString::Printf(TEXT("Mesh Thin Face %d"), Face.GetId()));
-		}
-#endif
-
-		for (FThinZone2D& Zone : ThinZones)
-		{
-			FindThinZoneBoundary(Zone.GetFirstSide());
-			FindThinZoneBoundary(Zone.GetSecondSide());
-		}
-
-		for (FThinZone2D& Zone : ThinZones)
-		{
-			EMeshingState bFirstSideMeshingState = Zone.GetFirstSide().GetMeshingState();
-			EMeshingState bSecondSideMeshingState = Zone.GetSecondSide().GetMeshingState();
-
-			if (bFirstSideMeshingState == EMeshingState::FullyMeshed)
-			{
-				MeshThinZoneSide(Zone.GetFirstSide());
-			}
-			else if (bSecondSideMeshingState == EMeshingState::FullyMeshed)
-			{
-				MeshThinZoneSide(Zone.GetSecondSide());
-			}
-			else if (bFirstSideMeshingState == EMeshingState::PartiallyMeshed)
-			{
-				MeshThinZoneSide(Zone.GetFirstSide());
-			}
-			else if (bSecondSideMeshingState == EMeshingState::PartiallyMeshed)
-			{
-				MeshThinZoneSide(Zone.GetSecondSide());
-			}
-			else if (Zone.GetFirstSide().Length() > Zone.GetSecondSide().Length())
-			{
-				MeshThinZoneSide(Zone.GetFirstSide());
-			}
-			else
-			{
-				MeshThinZoneSide(Zone.GetSecondSide());
-			}
-		}
-	}
-
-	// if the extremity of the thin zone are connected by a short edges path, the edges path are not discretized to avoid a well discretized edge connecting to thin sides
-
-}
-
-void FParametricMesher::FindThinZoneBoundary(FThinZoneSide& Side)
-{
-	FTopologicalEdge* Edge = nullptr;
-	FLinearBoundary SideEdgeCoordinate;
-
-	for (FEdgeSegment& EdgeSegment : Side.GetSegments())
-	{
-		double UMin = EdgeSegment.GetCoordinate(ELimit::Start);
-		double UMax = EdgeSegment.GetCoordinate(ELimit::End);
-		GetMinMax(UMin, UMax);
-
-		if (Edge != EdgeSegment.GetEdge())
-		{
-			if (Edge)
-			{
-				Edge->AddThinZone(SideEdgeCoordinate);
-			}
-
-			Edge = EdgeSegment.GetEdge();
-			SideEdgeCoordinate.Set(UMin, UMax);
-		}
-		else
-		{
-			SideEdgeCoordinate.ExtendTo(UMin, UMax);
-		}
-	};
-	Edge->AddThinZone(SideEdgeCoordinate);
-}
-
-void FParametricMesher::MeshThinZoneSide(FThinZoneSide& Side)
-{
-	typedef TFunction<bool(double, double)> CompareMethode;
-
-	FTopologicalEdge* Edge = nullptr;
-
-	int32 Index = 0;
-	int32 Increment = 1;
-	double UMin = 0.;
-	double UMax = 0.;
-
-	TArray<double> EdgeCuttingPointCoordinates;
-	const TArray<int32>* NodeIndices = nullptr;
-
-	TFunction<void(const FEdgeSegment&)> AddImposedCuttingPoint = [&](const FEdgeSegment& EdgeSegment)
-	{
-		for (; Index >= 0 && Index < EdgeCuttingPointCoordinates.Num(); Index += Increment)
-		{
-			if (EdgeCuttingPointCoordinates[Index] < UMin || EdgeCuttingPointCoordinates[Index] > UMax)
-			{
-				break;
-			}
-
-			FPoint2D CuttingPoint = EdgeSegment.ComputeEdgePoint(EdgeCuttingPointCoordinates[Index]);
-
-			FEdgeSegment* ClosedSegment = EdgeSegment.GetCloseSegment();
-			if (ClosedSegment == nullptr)
-			{
-#ifdef CADKERNEL_DEV
-				ensureCADKernel(false);
-				Wait();
-#endif
-				continue;
-			}
-
-			double OppositeCuttingPointSegmentU;
-			FPoint2D OppositeCuttingPoint = ClosedSegment->ProjectPoint(CuttingPoint, OppositeCuttingPointSegmentU);
-
-			double OppositeCuttingPointU = 0;
-			FTopologicalEdge* OppositeEdge = nullptr;
-			if (FMath::IsNearlyZero(OppositeCuttingPointSegmentU) && ClosedSegment->GetPrevious()->GetCloseSegment())
-			{
-				FEdgeSegment* PreviousClosedSegment = ClosedSegment->GetPrevious();
-				OppositeCuttingPoint = PreviousClosedSegment->ProjectPoint(CuttingPoint, OppositeCuttingPointSegmentU);
-				OppositeCuttingPointU = PreviousClosedSegment->ComputeEdgeCoordinate(OppositeCuttingPointSegmentU);
-				OppositeEdge = PreviousClosedSegment->GetEdge();
-			}
-			else if (FMath::IsNearlyEqual(OppositeCuttingPointSegmentU, 1.) && ClosedSegment->GetNext()->GetCloseSegment())
-			{
-				FEdgeSegment* NextClosedSegment = ClosedSegment->GetNext();
-				OppositeCuttingPoint = NextClosedSegment->ProjectPoint(CuttingPoint, OppositeCuttingPointSegmentU);
-				OppositeCuttingPointU = NextClosedSegment->ComputeEdgeCoordinate(OppositeCuttingPointSegmentU);
-				OppositeEdge = NextClosedSegment->GetEdge();
-			}
-			else
-			{
-				OppositeEdge = ClosedSegment->GetEdge();
-				OppositeCuttingPointU = ClosedSegment->ComputeEdgeCoordinate(OppositeCuttingPointSegmentU);
-			}
-
-			if (OppositeEdge != Edge)
-			{
-				OppositeEdge->AddImposedCuttingPointU(OppositeCuttingPointU, (*NodeIndices)[Index]);
-#ifdef DEBUG_MESHTHINSURF
-				DisplayPoint2DWithScale(CuttingPoint, EVisuProperty::RedPoint);
-				DisplaySegmentWithScale(ClosedSegment->GetExtemity(ELimit::End), ClosedSegment->GetExtemity(ELimit::Start));
-				DisplayPoint2DWithScale(OppositeCuttingPoint);
-				DisplaySegmentWithScale(CuttingPoint, OppositeCuttingPoint);
-				Wait(false);
-#endif
-			}
-		}
-	};
-
-	TFunction<void(double, CompareMethode)> FindFirstIndex = [&](double ULimit, CompareMethode Compare)
-	{
-		for (; Index >= 0 && Index < EdgeCuttingPointCoordinates.Num(); Index += Increment)
-		{
-			if (Compare(ULimit, EdgeCuttingPointCoordinates[Index]))
-			{
-				break;
-			}
-		}
-	};
-
-#ifdef DEBUG_MESHTHINSURF
-	Open3DDebugSession(TEXT("MeshThinZoneSide"));
-#endif
-
-	for (FEdgeSegment& EdgeSegment : Side.GetSegments())
-	{
-		UMin = EdgeSegment.GetCoordinate(ELimit::Start);
-		UMax = EdgeSegment.GetCoordinate(ELimit::End);
-		GetMinMax(UMin, UMax);
-
-		if (Edge != EdgeSegment.GetEdge())
-		{
-#ifdef DEBUG_MESHTHINSURF
-			if (Edge)
-			{
-				Close3DDebugSession();
-			}
-			Open3DDebugSession(TEXT("Projection of mesh"));
-#endif
-			Edge = EdgeSegment.GetEdge();
-			if (Edge == nullptr)
-			{
-				continue;
-			}
-
-			if (!Edge->IsMeshed())
-			{
-				Mesh(*Edge, *Edge->GetFace());
-			}
-
-			FEdgeMesh& EdgeMesh = *Edge->GetOrCreateMesh(MeshModel);
-			NodeIndices = &EdgeMesh.EdgeVerticesIndex;
-			GetCuttingPointCoordinates(Edge->GetCuttingPoints(), EdgeCuttingPointCoordinates);
-
-
-			if (EdgeCuttingPointCoordinates.Num() == 0)
-			{
-				TArray<FPoint>& NodeCoordinates = EdgeMesh.GetNodeCoordinates();
-				TArray<FPoint> ProjectedPoints;
-				Edge->ProjectPoints(NodeCoordinates, EdgeCuttingPointCoordinates, ProjectedPoints);
-				if (EdgeCuttingPointCoordinates.Num() > 1 && EdgeCuttingPointCoordinates[0] > EdgeCuttingPointCoordinates[1])
-				{
-					Algo::Reverse(EdgeCuttingPointCoordinates);
-				}
-				EdgeCuttingPointCoordinates.EmplaceAt(0, Edge->GetStartCurvilinearCoordinates());
-				EdgeCuttingPointCoordinates.Emplace(Edge->GetEndCurvilinearCoordinates());
-			}
-
-			ensureCADKernel(EdgeCuttingPointCoordinates[0] < EdgeCuttingPointCoordinates[1]);
-
-			bool bEdgeIsForward = EdgeSegment.IsForward();
-			if (bEdgeIsForward)
-			{
-				Index = 0;
-				Increment = 1;
-				FindFirstIndex(UMin, [](double Value1, double Value2) {return (Value1 < Value2); });
-			}
-			else
-			{
-				Index = (int32)EdgeCuttingPointCoordinates.Num() - 1;
-				Increment = -1;
-				FindFirstIndex(UMax, [](double Value1, double Value2) {return (Value1 > Value2); });
-			}
-		}
-
-		AddImposedCuttingPoint(EdgeSegment);
-	}
-
-#ifdef DEBUG_MESHTHINSURF
-	Close3DDebugSession();
-	Close3DDebugSession();
-	//Wait();
-#endif
-
-}
-
-#ifdef DEBUG_INTERSECTEDGEISOS
-TMap<int32, int32> SurfaceDrawed;
-bool bDisplayIsoCurve = true;
-
-void DebugIntersectEdgeIsos(const FTopologicalFace& Face, const TArray<double>& IsoCoordinates, EIso TypeIso)
-{
-	if (SurfaceDrawed.Find(Face.GetId()) == nullptr)
-	{
-		SurfaceDrawed.Add(Face.GetId(), 0);
-	}
-
-	if (bDisplayIsoCurve && SurfaceDrawed[Face.GetId()] < 2)
-	{
-		SurfaceDrawed[Face.GetId()]++;
-
-		FSurfacicBoundary Bounds = Face.GetBoundary();
-
-		//{
-		//	F3DDebugSession _(FString::Printf(TEXT("Iso %s 3D %d"), TypeIso == EIso::IsoU ? TEXT("U") : TEXT("V"), Face.GetId()));
-		//	for (double U : IsoCoordinates)
-		//	{
-		//		DisplayIsoCurve(*Face.GetCarrierSurface(), U, TypeIso);
-		//	}
-		//}
-
-		F3DDebugSession _(FString::Printf(TEXT("Iso %s 2D %d"), TypeIso == EIso::IsoU ? TEXT("U") : TEXT("V"), Face.GetId()));
-		if (TypeIso == EIso::IsoU)
-		{
-			for (double U : IsoCoordinates)
-			{
-				FPoint2D Start(U, Bounds[EIso::IsoV].Min);
-				FPoint2D End(U, Bounds[EIso::IsoV].Max);
-				DisplaySegment(Start, End, 0, EVisuProperty::Iso);
-			}
-		}
-		else
-		{
-			for (double V : IsoCoordinates)
-			{
-				FPoint2D Start(Bounds[EIso::IsoU].Min, V);
-				FPoint2D End(Bounds[EIso::IsoU].Max, V);
-				DisplaySegment(Start, End, 0, EVisuProperty::Iso);
-			}
-		}
-	}
-}
-
-#endif
-
 } // namespace
\ No newline at end of file
diff --git a/Engine/Source/Runtime/Datasmith/CADKernel/Private/CADKernel/Mesh/Structure/FaceMesh.cpp b/Engine/Source/Runtime/Datasmith/CADKernel/Private/CADKernel/Mesh/Structure/FaceMesh.cpp
--- a/Engine/Source/Runtime/Datasmith/CADKernel/Private/CADKernel/Mesh/Structure/FaceMesh.cpp
+++ b/Engine/Source/Runtime/Datasmith/CADKernel/Private/CADKernel/Mesh/Structure/FaceMesh.cpp
@@ -28,13 +28,19 @@ void FFaceMesh::GetNodeIdToCoordinates(TMap<int32, const FPoint*>& NodeIdToCoord
 				continue;
 			}
 
-			const TSharedRef<FVertexMesh> StartVertexMesh = ActiveEdge->GetStartVertex()->GetMesh();
-			NodeIdToCoordinates.Add(StartVertexMesh->GetMesh(), &StartVertexMesh->GetNodeCoordinates()[0]);
+			const FVertexMesh* StartVertexMesh = ActiveEdge->GetStartVertex()->GetMesh();
+			if(StartVertexMesh)
+			{
+				NodeIdToCoordinates.Add(StartVertexMesh->GetMesh(), &StartVertexMesh->GetNodeCoordinates()[0]);
+			}
 
-			const TSharedRef<FVertexMesh> EndVertexMesh = ActiveEdge->GetEndVertex()->GetMesh();
-			NodeIdToCoordinates.Add(EndVertexMesh->GetMesh(), &EndVertexMesh->GetNodeCoordinates()[0]);
+			const FVertexMesh* EndVertexMesh = ActiveEdge->GetEndVertex()->GetMesh();
+			if (EndVertexMesh)
+			{
+				NodeIdToCoordinates.Add(EndVertexMesh->GetMesh(), &EndVertexMesh->GetNodeCoordinates()[0]);
+			}
 
-			const TSharedRef<FEdgeMesh> EdgeMesh = ActiveEdge->GetMesh();
+			const FEdgeMesh* EdgeMesh = ActiveEdge->GetMesh();
 
 			TArray<int32> NodeIds = EdgeMesh->EdgeVerticesIndex;
 			const FPoint* EdgeNodeCoordinates = EdgeMesh->GetNodeCoordinates().GetData();
