CADKernel update: Meshing of thin faces

diff --git a/Engine/Source/Runtime/Datasmith/CADKernel/Private/CADKernel/Mesh/Structure/Grid.cpp b/Engine/Source/Runtime/Datasmith/CADKernel/Private/CADKernel/Mesh/Structure/Grid.cpp
--- a/Engine/Source/Runtime/Datasmith/CADKernel/Private/CADKernel/Mesh/Structure/Grid.cpp
+++ b/Engine/Source/Runtime/Datasmith/CADKernel/Private/CADKernel/Mesh/Structure/Grid.cpp
@@ -22,7 +22,7 @@ FGrid::FGrid(FTopologicalFace& InFace, FModelMesh& InMeshModel)
 	, MinimumElementSize(Tolerance3D * 2.)
 	, MeshModel(InMeshModel)
 {
-#ifdef CADKERNEL_DEBUG
+#ifdef DEBUG_ONLY_SURFACE_TO_DEBUG
 	bDisplay = (InFace.GetId() == FaceToDebug);
 	Open3DDebugSession(bDisplay, FString::Printf(TEXT("Grid %d"), InFace.GetId()));
 #endif
@@ -30,7 +30,6 @@ FGrid::FGrid(FTopologicalFace& InFace, FModelMesh& InMeshModel)
 
 void FGrid::ProcessPointCloud()
 {
-	EGridSpace DisplaySpace = EGridSpace::UniformScaled;
 	FTimePoint StartTime = FChrono::Now();
 
 	if (!GetMeshOfLoops())
@@ -38,29 +37,30 @@ void FGrid::ProcessPointCloud()
 		return;
 	}
 
-	ScaleLoops();
-
 #ifdef DEBUG_GRID
-	DisplayGridLoop(TEXT("FGrid::Loop 2D"), GetLoops2D(EGridSpace::Default2D), true, false);
-	DisplayInnerDomainPoints(TEXT("FGrid::PointCloud 2D"), GetInner2DPoints(EGridSpace::Default2D));
+	if(bDisplay)
+	{
+		DisplayGridLoops(TEXT("FGrid::Loop 2D"), GetLoops2D(EGridSpace::Default2D), true, true, false);
+		DisplayInnerDomainPoints(TEXT("FGrid::PointCloud 2D"), GetInner2DPoints(EGridSpace::Default2D));
 
-	DisplayGridLoop(TEXT("FGrid::Loop 2D Scaled"), GetLoops2D(EGridSpace::Scaled), true, false);
-	DisplayInnerDomainPoints(TEXT("FGrid::PointCloud 2D Scaled"), GetInner2DPoints(EGridSpace::Scaled));
+		DisplayGridLoops(TEXT("FGrid::Loop 2D Scaled"), GetLoops2D(EGridSpace::Scaled), true, false, false);
+		DisplayInnerDomainPoints(TEXT("FGrid::PointCloud 2D Scaled"), GetInner2DPoints(EGridSpace::Scaled));
 
-	DisplayGridLoop(TEXT("FGrid::Loop 2D UniformScaled"), GetLoops2D(EGridSpace::UniformScaled), true, false);
-	DisplayInnerDomainPoints(TEXT("FGrid::PointCloud 2D UniformScaled"), GetInner2DPoints(EGridSpace::UniformScaled));
+		DisplayGridLoops(TEXT("FGrid::Loop 2D UniformScaled"), GetLoops2D(EGridSpace::UniformScaled), true, false, false);
+		DisplayInnerDomainPoints(TEXT("FGrid::PointCloud 2D UniformScaled"), GetInner2DPoints(EGridSpace::UniformScaled));
 
-	DisplayGridLoop(TEXT("FGrid::Loop 3D"), GetLoops3D(), true, false);
-	DisplayInnerDomainPoints(TEXT("FGrid::PointCloud 3D"), GetInner3DPoints());
+		DisplayGridLoops(TEXT("FGrid::Loop 3D"), GetLoops3D(), true, false, false);
+		DisplayInnerDomainPoints(TEXT("FGrid::PointCloud 3D"), GetInner3DPoints());
 
-	DisplayGridLoop(TEXT("FGrid::Loop 3D"), GetLoops3D(), true, false);
-	DisplayInnerDomainPoints(TEXT("FGrid::PointCloud 3D"), GetInner3DPoints());
+		DisplayGridLoops(TEXT("FGrid::Loop 3D"), GetLoops3D(), true, false, false);
+		DisplayInnerDomainPoints(TEXT("FGrid::PointCloud 3D"), GetInner3DPoints());
 
-	DisplayGridNormal();
+		DisplayGridNormal();
 
-	DisplayInnerDomainPoints(TEXT("FGrid::PointCloud before FindInnerFacePoints"), GetInner2DPoints(EGridSpace::UniformScaled));
+		DisplayInnerDomainPoints(TEXT("FGrid::PointCloud before FindInnerFacePoints"), GetInner2DPoints(EGridSpace::UniformScaled));
 
-	Wait(bDisplay);
+		Wait(false);
+	}
 #endif
 
 	FindInnerFacePoints();
@@ -82,12 +82,15 @@ void FGrid::ProcessPointCloud()
 #endif
 
 	// Removed of Thin zone boundary (the last boundaries). In case of thin zone, the number of 2d boundary will be biggest than 3d boundary one.
-	// Only EGridSpace::Default2D is needed.
-	FaceLoops2D[EGridSpace::Default2D].SetNum(FaceLoops3D.Num());
+	// Only EGridSpace::UniformScaled is needed.
+	FaceLoops2D[EGridSpace::UniformScaled].SetNum(FaceLoops3D.Num());
 
 #ifdef DEBUG_GRID
-	DisplayGridLoop(TEXT("FGrid::Final Loop 2D"), GetLoops2D(DisplaySpace), true, false);
-	//DisplayLoop(TEXT("FGrid::Final Loop 3D"), GetLoops3D(), true, false);
+	if (bDisplay)
+	{
+		DisplayGridLoops(TEXT("FGrid::Final Loop 2D"), GetLoops2D(EGridSpace::UniformScaled), true, true, false);
+		Wait(false);
+	}
 #endif
 
 	Chronos.ProcessPointCloudDuration = FChrono::Elapse(StartTime);
@@ -212,7 +215,7 @@ void FGrid::GetPreferredUVCuttingParametersFromLoops(FCuttingGrid& CuttingParame
 	{
 		for (const FOrientedEdge& Edge : Loop->GetEdges())
 		{
-			nbPoints += Edge.Entity->GetOrCreateMesh(MeshModel)->GetNodeCoordinates().Num() + 1;
+			nbPoints += Edge.Entity->GetOrCreateMesh(MeshModel).GetNodeCoordinates().Num() + 1;
 		}
 	}
 
@@ -237,7 +240,11 @@ void FGrid::GetPreferredUVCuttingParametersFromLoops(FCuttingGrid& CuttingParame
 			TSharedRef<FTopologicalEdge> ActiveEdge = Edge->GetLinkActiveEdge();
 			if (ActiveEdge->IsMeshed())
 			{
-				const TSharedRef<FMesh> EdgeMesh = ActiveEdge->GetMesh();
+				const FMesh* EdgeMesh = ActiveEdge->GetMesh();
+				if (!EdgeMesh)
+				{
+					continue;
+				}
 
 				const TArray<FPoint>& EdgeMeshNodes = EdgeMesh->GetNodeCoordinates();
 				if (EdgeMeshNodes.Num() == 0)
@@ -320,7 +327,7 @@ bool FGrid::GeneratePointCloud()
 {
 	FTimePoint StartTime = FChrono::Now();
 
-	if(CheckIf2DGridIsDegenerate())
+	if (CheckIf2DGridIsDegenerate())
 	{
 		return false;
 	}
@@ -358,16 +365,16 @@ void FGrid::FindPointsCloseToLoop()
 	}
 
 	int32 CellIndex = 0;
-	bool bWaitCell = false;
+	bool bWaitCell = true;
 	TFunction<void()> DisplayCell = [&]()
 	{
 		if (bDisplay)
 		{
 			F3DDebugSession _(*FString::Printf(TEXT("Cell %d"), CellIndex++));
-			DisplayPoint2DWithScale(Points2D[EGridSpace::UniformScaled][GlobalIndex] );
-			DisplayPoint2DWithScale(Points2D[EGridSpace::UniformScaled][GlobalIndex - 1] );
-			DisplayPoint2DWithScale(Points2D[EGridSpace::UniformScaled][GlobalIndex - 1 - CuttingCount[EIso::IsoU]] );
-			DisplayPoint2DWithScale(Points2D[EGridSpace::UniformScaled][GlobalIndex - CuttingCount[EIso::IsoU]] );
+			DisplayPoint2DWithScale(Points2D[EGridSpace::UniformScaled][GlobalIndex]);
+			DisplayPoint2DWithScale(Points2D[EGridSpace::UniformScaled][GlobalIndex - 1]);
+			DisplayPoint2DWithScale(Points2D[EGridSpace::UniformScaled][GlobalIndex - 1 - CuttingCount[EIso::IsoU]]);
+			DisplayPoint2DWithScale(Points2D[EGridSpace::UniformScaled][GlobalIndex - CuttingCount[EIso::IsoU]]);
 			Wait(bWaitCell);
 		}
 	};
@@ -528,7 +535,7 @@ void FGrid::FindPointsCloseToLoop()
 			if (bDisplay)
 			{
 				F3DDebugSession _(TEXT("SEG"));
-				DisplaySegmentWithScale(*PointB , *PointA );
+				DisplaySegmentWithScale(*PointB, *PointA);
 			}
 #endif
 
@@ -913,11 +920,20 @@ void FGrid::GetMeshOfLoop(const FTopologicalLoop& Loop)
 		const TSharedPtr<FTopologicalEdge>& Edge = OrientedEdge.Entity;
 		const TSharedRef<FTopologicalEdge>& ActiveEdge = Edge->GetLinkActiveEdge();
 
+#ifdef DEBUG_GET_MESH_OF_LOOP
+		if (bDisplay)
+		{
+			F3DDebugSession _(bDisplay, FString::Printf(TEXT("Edge %d"), Edge->GetId()));
+			Display2D(*Edge);
+			Wait();
+		}
+#endif
+
 		bool bSameDirection = Edge->IsSameDirection(*ActiveEdge);
 
 		TArray<double> EdgeCuttingPointCoordinates;
 		{
-			TArray<FCuttingPoint>& CuttingPoints = Edge->GetCuttingPoints();
+			const TArray<FCuttingPoint>& CuttingPoints = Edge->GetCuttingPoints();
 			if (!CuttingPoints.IsEmpty())
 			{
 				GetCuttingPointCoordinates(CuttingPoints, EdgeCuttingPointCoordinates);
@@ -957,11 +973,14 @@ void FGrid::GetMeshOfLoop(const FTopologicalLoop& Loop)
 		{
 			if (EdgeCuttingPointCoordinates.IsEmpty())
 			{
-				TArray<FPoint>& MeshVertex3D = ActiveEdge->GetMesh()->GetNodeCoordinates();
+				const TArray<FPoint>& MeshVertex3D = ActiveEdge->GetOrCreateMesh(MeshModel).GetNodeCoordinates();
 				TArray<double> ProjectedPointCoords;
 
 				CuttingPolyline.Coordinates.Reserve(MeshVertex3D.Num() + 2);
-				Edge->ProjectTwinEdgePoints(MeshVertex3D, bSameDirection, CuttingPolyline.Coordinates);
+				if(MeshVertex3D.Num())
+				{
+					Edge->ProjectTwinEdgePoints(MeshVertex3D, bSameDirection, CuttingPolyline.Coordinates);
+				}
 				CuttingPolyline.Coordinates.Insert(Edge->GetBoundary().GetMin(), 0);
 				CuttingPolyline.Coordinates.Add(Edge->GetBoundary().GetMax());
 
@@ -993,7 +1012,7 @@ void FGrid::GetMeshOfLoop(const FTopologicalLoop& Loop)
 		TArray<int32> EdgeVerticesIndex;
 		if (Edge->IsDegenerated())
 		{
-			EdgeVerticesIndex.Init(ActiveEdge->GetStartVertex()->GetLinkActiveEntity()->GetOrCreateMesh(MeshModel)->GetMesh(), CuttingPolyline.Coordinates.Num());
+			EdgeVerticesIndex.Init(ActiveEdge->GetStartVertex()->GetLinkActiveEntity()->GetOrCreateMesh(MeshModel).GetMesh(), CuttingPolyline.Coordinates.Num());
 		}
 		else if (Edge->IsVirtuallyMeshed())
 		{
@@ -1003,14 +1022,14 @@ void FGrid::GetMeshOfLoop(const FTopologicalLoop& Loop)
 			int32 MiddleNodeIndex = NodeCount / 2;
 			int32 Index = 0;
 			{
-				int32 StartVertexMeshId = Edge->GetStartVertex()->GetLinkActiveEntity()->GetOrCreateMesh(MeshModel)->GetMesh();
+				int32 StartVertexMeshId = Edge->GetStartVertex()->GetLinkActiveEntity()->GetOrCreateMesh(MeshModel).GetMesh();
 				for (; Index < MiddleNodeIndex; ++Index)
 				{
 					EdgeVerticesIndex.Add(StartVertexMeshId);
 				}
 			}
 			{
-				int32 EndVertexMeshId = Edge->GetEndVertex()->GetLinkActiveEntity()->GetOrCreateMesh(MeshModel)->GetMesh();
+				int32 EndVertexMeshId = Edge->GetEndVertex()->GetLinkActiveEntity()->GetOrCreateMesh(MeshModel).GetMesh();
 				for (; Index < NodeCount; ++Index)
 				{
 					EdgeVerticesIndex.Add(EndVertexMeshId);
@@ -1019,10 +1038,10 @@ void FGrid::GetMeshOfLoop(const FTopologicalLoop& Loop)
 		}
 		else
 		{
-			EdgeVerticesIndex = ((TSharedRef<FEdgeMesh>)ActiveEdge->GetOrCreateMesh(MeshModel))->EdgeVerticesIndex;
+			EdgeVerticesIndex = ActiveEdge->GetOrCreateMesh(MeshModel).EdgeVerticesIndex;
 		}
 
-#ifdef DEBUG_GET_BOUNDARY_MESH
+#ifdef DEBUG_GET_MESH_OF_LOOP
 		if (bDisplay)
 		{
 			F3DDebugSession _(bDisplay, FString::Printf(TEXT("Edge %d cutting points on surface"), ActiveEdge->GetId()));
@@ -1030,7 +1049,8 @@ void FGrid::GetMeshOfLoop(const FTopologicalLoop& Loop)
 			{
 				DisplayPoint(Point2D);
 			}
-			//Wait();
+			DisplayPolyline(CuttingPolyline.Points2D, EVisuProperty::BlueCurve);
+			Wait();
 		}
 #endif
 
@@ -1046,17 +1066,53 @@ void FGrid::GetMeshOfLoop(const FTopologicalLoop& Loop)
 
 		ensureCADKernel(CuttingPolyline.Size() > 1);
 
-		Loop2D.Append(CuttingPolyline.Points2D);
+		Loop2D.Append(MoveTemp(CuttingPolyline.Points2D));
 		Loop2D.Pop();
 
-		Loop3D.Emplace(ActiveEdge->GetStartVertex((OrientedEdge.Direction == EOrientation::Front) == bSameDirection)->GetLinkActiveEntity()->GetBarycenter());
-		Loop3D.Append(CuttingPolyline.Points3D.GetData() + 1, CuttingPolyline.Points3D.Num() - 2);
+#ifdef DEBUG_GET_MESH_OF_LOOP
+		if (bDisplay)
+		{
+			F3DDebugSession _(bDisplay, FString::Printf(TEXT("Loop with Edge %d"), Edge->GetId()));
+			DisplayPolyline(Loop2D, EVisuProperty::BlueCurve);
+			Wait();
+		}
+#endif
+
+#ifdef CADKERNEL_DEBUG_
+		const FPoint2D* PrevPoint = &Loop2D.Last();
+		if(Loop2D.Num() > 1)
+		{
+			for (const FPoint2D& Point : Loop2D)
+			{
+				double Dist = PrevPoint->Distance(Point);
+				if(Dist > 0.0001)
+				{
+					ensureCADKernel(false);
+				}
+				PrevPoint = &Point;
+			}
+		}
+#endif
+
+		int32 LastIndex = Loop3D.Num();
+		Loop3D.Append(MoveTemp(CuttingPolyline.Points3D));
+		Loop3D[LastIndex] = (ActiveEdge->GetStartVertex((OrientedEdge.Direction == EOrientation::Front) == bSameDirection)->GetLinkActiveEntity()->GetBarycenter());
+		Loop3D.Pop();
 
-		LoopNormals.Append(CuttingPolyline.Normals);
+		LoopNormals.Append(MoveTemp(CuttingPolyline.Normals));
 		LoopNormals.Pop();
 
-		LoopIds.Append(EdgeVerticesIndex);
+		LoopIds.Append(MoveTemp(EdgeVerticesIndex));
 		LoopIds.Pop();
+
+#ifdef DEBUG_GET_MESH_OF_LOOP
+		static int32 GetMeshOfLoopIter = 0;
+		GetMeshOfLoopIter++;
+		if (LoopIds.Num() != Loop2D.Num() || GetMeshOfLoopIter == 2991)
+		{
+			printf("");
+		}
+#endif
 	}
 
 	if (Loop2D.Num() < 3) // degenerated loop
@@ -1068,17 +1124,79 @@ void FGrid::GetMeshOfLoop(const FTopologicalLoop& Loop)
 	}
 }
 
-//#define DEBUG_GET_BOUNDARY_MESH
-bool FGrid::GetMeshOfLoops()
+void FGrid::GetMeshOfThinZone(const FThinZone2D& ThinZone)
 {
-	int32 ThinZoneNum = 0;
-	if (Face.HasThinZone())
+	TFunction<void(const FThinZoneSide&, TArray<FPoint2D>&)> GetThinZoneSideMesh = [this](const FThinZoneSide& Side, TArray<FPoint2D>& ThinZoneSideMesh)
 	{
-		ThinZoneNum = Face.GetThinZones().Num();
+		TArray<int32> MeshIndices;
+		FAddMeshNodeFunc AddMeshNode = [&Mesh = ThinZoneSideMesh, &MeshIndices](int32 NodeIndice, const FPoint2D& MeshNode2D, double MeshingTolerance, const FEdgeSegment& EdgeSegment, const FPairOfIndex& OppositeNodeIndices)
+		{
+			if(MeshIndices.IsEmpty() || MeshIndices.Last() != NodeIndice)
+			{
+				MeshIndices.Add(NodeIndice);
+				Mesh.Emplace(MeshNode2D);
+			}
+		};
+
+		FReserveContainerFunc Reserve = [&Mesh = ThinZoneSideMesh, &MeshIndices](int32 MeshVertexCount)
+		{
+			Mesh.Reserve(Mesh.Num() + MeshVertexCount);
+			MeshIndices.Reserve(Mesh.Num() + MeshVertexCount);
+		};
+
+		const bool bWithTolerance = false;
+		Side.GetExistingMeshNodes(GetFace(), MeshModel, Reserve, AddMeshNode, bWithTolerance);
+	};
+
+	TFunction<void(TArray<FPoint2D>&)> RemoveDuplicatedNode = [](TArray<FPoint2D>& ThinZoneMesh)
+	{
+		if (ThinZoneMesh.Num())
+		{
+			// Remove Duplicated Node
+			constexpr double Tolerance = DOUBLE_SMALL_NUMBER;
+			FPoint2D* Next = &ThinZoneMesh[0];
+			for (int32 Index = ThinZoneMesh.Num() - 1; Index >= 0; --Index)
+			{
+				double SquareDist = Next->SquareDistance(ThinZoneMesh[Index]);
+				if (SquareDist < Tolerance)
+				{
+					ThinZoneMesh.RemoveAt(Index);
+				}
+				else
+				{
+					Next = &ThinZoneMesh[Index];
+				}
+			}
+		}
+	};
+
+	TArray<FPoint2D>* ThinZoneMesh = &FaceLoops2D[EGridSpace::UniformScaled].Emplace_GetRef();
+
+	GetThinZoneSideMesh(ThinZone.GetFirstSide(), *ThinZoneMesh);
+	if (ThinZone.GetCategory() == EThinZone2DType::BetweenLoops)
+	{
+		RemoveDuplicatedNode(*ThinZoneMesh);
+		if (ThinZoneMesh->Num() < 3)
+		{
+			FaceLoops2D[EGridSpace::UniformScaled].Pop();
+		}
+
+		ThinZoneMesh = &FaceLoops2D[EGridSpace::UniformScaled].Emplace_GetRef();
+	}
+
+	GetThinZoneSideMesh(ThinZone.GetSecondSide(), *ThinZoneMesh);
+	RemoveDuplicatedNode(*ThinZoneMesh);
+
+	if (ThinZoneMesh->Num() < 3)
+	{
+		FaceLoops2D[EGridSpace::UniformScaled].Pop();
 	}
+}
 
+bool FGrid::GetMeshOfLoops()
+{
 	int32 LoopCount = Face.GetLoops().Num();
-	FaceLoops2D[EGridSpace::Default2D].Reserve(LoopCount + ThinZoneNum);
+	FaceLoops2D[EGridSpace::Default2D].Reserve(LoopCount);
 
 	FaceLoops3D.Reserve(LoopCount);
 	NormalsOfFaceLoops.Reserve(LoopCount);
@@ -1111,31 +1229,30 @@ bool FGrid::GetMeshOfLoops()
 		return false;
 	}
 
-	TFunction<void(TArray<FPoint2D>&, const TArray<FEdgeSegment>&)> AddThinZoneInLoopPoints = [](TArray<FPoint2D>& LoopPoints, const TArray<FEdgeSegment>& Segments)
-	{
-		const FEdgeSegment& Segment = Segments[0];
-		LoopPoints.Emplace(Segment.GetEdge()->Approximate2DPoint(Segment.GetCoordinate(ELimit::Start)));
-		for (const FEdgeSegment& SegmentIter : Segments)
-		{
-			LoopPoints.Emplace(SegmentIter.GetEdge()->Approximate2DPoint(SegmentIter.GetCoordinate(ELimit::End)));
-		}
-	};
-
-	if (ThinZoneNum)
+	ScaleLoops();
+ 
+	if (Face.HasThinZone())
 	{
+#ifdef DEBUG_GET_BOUNDARY_MESH
+		F3DDebugSession _(bDisplay, ("Thin zone loop"));
+#endif
 		for (const FThinZone2D& ThinZone : Face.GetThinZones())
 		{
-			int32 PointNum = ThinZone.GetFirstSide().GetSegments().Num();
-			PointNum += ThinZone.GetSecondSide().GetSegments().Num();
-			TArray<FPoint2D>& LoopPoints = FaceLoops2D[EGridSpace::Default2D].Emplace_GetRef();
-			LoopPoints.Reserve(PointNum + 4);
-
-			AddThinZoneInLoopPoints(LoopPoints, ThinZone.GetFirstSide().GetSegments());
-			AddThinZoneInLoopPoints(LoopPoints, ThinZone.GetSecondSide().GetSegments());
+			GetMeshOfThinZone(ThinZone);
+#ifdef DEBUG_GET_BOUNDARY_MESH
+			if(bDisplay)
+			{
+				DisplayGridLoop(TEXT("ThinZone"), FaceLoops2D[EGridSpace::UniformScaled][FaceLoops2D[EGridSpace::UniformScaled].Num() - 1], true, true, false);
+				Wait();
+			}
+#endif
 		}
+#ifdef DEBUG_GET_BOUNDARY_MESH
+		//Wait();
+#endif
 	}
 
-	// Fit boundaries to Surface bounds.
+	// Fit loops to Surface bounds.
 	const FSurfacicBoundary& Bounds = Face.GetBoundary();
 	for (TArray<FPoint2D>& Loop : FaceLoops2D[EGridSpace::Default2D])
 	{
@@ -1144,12 +1261,33 @@ bool FGrid::GetMeshOfLoops()
 			Bounds.MoveInsideIfNot(Point);
 		}
 	}
+
+	FSurfacicBoundary UniformScaleBounds;
+	UniformScaleBounds[IsoU].Set(UniformCuttingCoordinates[IsoU][0], UniformCuttingCoordinates[IsoU].Last());
+	UniformScaleBounds[IsoV].Set(UniformCuttingCoordinates[IsoV][0], UniformCuttingCoordinates[IsoV].Last());
+
+	// Fit loops in UniformScaled to UniformScale bounds.
+	for (TArray<FPoint2D>& Loop : FaceLoops2D[EGridSpace::UniformScaled])
+	{
+		for (FPoint2D& Point : Loop)
+		{
+			UniformScaleBounds.MoveInsideIfNot(Point);
+		}
+	}
+
 	return true;
 }
 
 void FGrid::ScaleLoops()
 {
+	int32 ThinZoneNum = 0;
+	if (Face.HasThinZone())
+	{
+		ThinZoneNum = Face.GetThinZones().Num();
+	}
+
 	FaceLoops2D[EGridSpace::Scaled].SetNum(FaceLoops2D[EGridSpace::Default2D].Num());
+	FaceLoops2D[EGridSpace::UniformScaled].Reserve(FaceLoops2D[EGridSpace::Default2D].Num() + ThinZoneNum);
 	FaceLoops2D[EGridSpace::UniformScaled].SetNum(FaceLoops2D[EGridSpace::Default2D].Num());
 
 	for (int32 IndexBoudnary = 0; IndexBoudnary < FaceLoops2D[EGridSpace::Default2D].Num(); ++IndexBoudnary)
@@ -1173,7 +1311,6 @@ void FGrid::ScaleLoops()
 			ComputeNewCoordinate(Points2D[EGridSpace::Scaled], IndexU, IndexV, Point, ScaledLoop[Index]);
 			ComputeNewCoordinate(Points2D[EGridSpace::UniformScaled], IndexU, IndexV, Point, UniformScaledLoop[Index]);
 		}
-
 	}
 }
 
@@ -1185,7 +1322,7 @@ bool FGrid::CheckIf2DGridIsDegenerate() const
 		for (int32 Index = 1; Index < CoordinateGrid[Iso].Num(); ++Index)
 		{
 			double Delta = CoordinateGrid[Iso][Index] - CoordinateGrid[Iso][Index - 1];
-			if(Delta > MaxDelta)
+			if (Delta > MaxDelta)
 			{
 				MaxDelta = Delta;
 			}
@@ -1307,7 +1444,7 @@ void FGrid::FindInnerFacePoints()
 	}
 
 #ifdef DEBUG_FIND_INNER_FACE_POINTS
-	DisplayGridLoop(TEXT("FGrid::Loop 2D After move according tol"), GetLoops2D(EGridSpace::UniformScaled), true, false);
+	DisplayGridLoops(TEXT("FGrid::Loop 2D After move according tol"), GetLoops2D(EGridSpace::UniformScaled), true, false);
 #endif
 
 	// Intersection along U axis
diff --git a/Engine/Source/Runtime/Datasmith/CADKernel/Private/CADKernel/Mesh/Structure/Grid.h b/Engine/Source/Runtime/Datasmith/CADKernel/Private/CADKernel/Mesh/Structure/Grid.h
--- a/Engine/Source/Runtime/Datasmith/CADKernel/Private/CADKernel/Mesh/Structure/Grid.h
+++ b/Engine/Source/Runtime/Datasmith/CADKernel/Private/CADKernel/Mesh/Structure/Grid.h
@@ -7,6 +7,7 @@
 #include "CADKernel/Math/Point.h"
 #include "CADKernel/Mesh/MeshEnum.h"
 #include "CADKernel/Mesh/Structure/GridBase.h"
+#include "CADKernel/Mesh/Structure/ThinZone2D.h"
 #include "CADKernel/UI/Display.h"
 
 #ifdef CADKERNEL_DEV
@@ -60,7 +61,6 @@ protected:
 	const FSurfacicTolerance FaceTolerance;
 	const double MinimumElementSize;
 
-
 	FModelMesh& MeshModel;
 
 
@@ -95,15 +95,19 @@ protected:
 	TArray<ENodeMarker> NodeMarkers;
 
 public:
+#ifdef CADKERNEL_DEV
+	FGridChronos Chronos;
+#endif
+
 	FGrid(FTopologicalFace& InFace, FModelMesh& InShellMesh);
 
-#ifndef CADKERNEL_DEBUG
-	virtual ~FGrid() = default;
-#else
+#ifdef CADKERNEL_DEBUG
 	virtual ~FGrid()
 	{
 		Close3DDebugSession(bDisplay);
 	}
+#else
+	virtual ~FGrid() = default;
 #endif
 
 	// ======================================================================================================================================================================================================================
@@ -179,6 +183,7 @@ protected:
 	 */
 	bool GetMeshOfLoops();
 	void GetMeshOfLoop(const FTopologicalLoop& Loop);
+	void GetMeshOfThinZone(const FThinZone2D& ThinZone);
 
 	// Meshing tools =========================================================================================================================================================================================================
 	// ======================================================================================================================================================================================================================
@@ -421,12 +426,16 @@ public:
 	// ======================================================================================================================================================================================================================
 	// Display Methodes   ================================================================================================================================================================================================
 	// ======================================================================================================================================================================================================================
+#ifdef CADKERNEL_DEV
+	void PrintTimeElapse() const;
+#endif
+
 #ifdef CADKERNEL_DEBUG
-	bool bDisplay = false;
 
 	void DisplayIsoNode(EGridSpace Space, const int32 PointIndex, FIdent Ident = 0, EVisuProperty Property = EVisuProperty::BluePoint) const;
 	void DisplayIsoNode(EGridSpace Space, const FIsoNode& Node, FIdent Ident = 0, EVisuProperty Property = EVisuProperty::BluePoint) const;
 	void DisplayIsoNodes(EGridSpace Space, const TArray<const FIsoNode*>& Nodes, EVisuProperty Property = EVisuProperty::BluePoint) const;
+
 	void DisplayIsoNodes(const FString& Message, EGridSpace Space, const TArray<const FIsoNode*>& Nodes, EVisuProperty Property = EVisuProperty::BluePoint) const
 	{
 		if (!bDisplay)
@@ -437,6 +446,22 @@ public:
 		DisplayIsoNodes(Space, Nodes, Property);
 	}
 
+	void DisplayIsoPolyline(const FString& Message, EGridSpace Space, const TArray<const FIsoNode*>& Nodes, EVisuProperty Property = EVisuProperty::BluePoint) const
+	{
+		if (!bDisplay)
+		{
+			return;
+		}
+		F3DDebugSession _(Message);
+		const FIsoNode* LastNode = Nodes.Last();
+		for (const FIsoNode* Node : Nodes)
+		{
+			DisplayIsoSegment(Space, *LastNode, *Node, 0, (EVisuProperty)(Property + 1));
+			LastNode = Node;
+		}
+		DisplayIsoNodes(Space, Nodes, Property);
+	}
+
 	void DisplayIsoSegment(EGridSpace Space, const FIsoSegment& Segment, FIdent Ident = 0, EVisuProperty Property = EVisuProperty::BlueCurve, bool bDisplayOrientation = false) const;
 	void DisplayIsoSegment(EGridSpace Space, const FIsoNode& NodeA, const FIsoNode& NodeB, FIdent Ident = 0, EVisuProperty Property = EVisuProperty::BlueCurve) const;
 
@@ -452,6 +477,7 @@ public:
 	}
 
 	void DisplayTriangle(EGridSpace Space, const FIsoNode& NodeA, const FIsoNode& NodeB, const FIsoNode& NodeC) const;
+	void DisplayTriangle3D(const FIsoNode& NodeA, const FIsoNode& NodeB, const FIsoNode& NodeC) const;
 
 	void DisplayGridPolyline(EGridSpace Space, const FLoopNode& StartNode, bool bDisplayNode, EVisuProperty Property = EVisuProperty::BlueCurve) const;
 	void DisplayGridPolyline(EGridSpace Space, const TArray<FIsoNode*>& Nodes, bool bDisplayNode, EVisuProperty Property = EVisuProperty::BlueCurve) const;
@@ -515,49 +541,60 @@ public:
 	}
 
 	template<typename TPoint>
-	void DisplayGridLoop(FString Message, const TArray<TArray<TPoint>>& Loops, bool bDisplayNodes, bool bMakeGroup) const
+	void DisplayGridLoop(FString Message, const TArray<TPoint>& Loop, bool bDisplayNodes, bool bMakeGroup, bool bDisplaySegments) const
 	{
 		if (!bDisplay)
 		{
 			return;
 		}
 
-		F3DDebugSession _(*Message);
-		int32 LoopIndex = 0;
-		for (const TArray<TPoint>& Loop : Loops)
+		F3DDebugSession _(bMakeGroup, Message);
+
+		const TPoint* FirstSegmentPoint = &Loop[0];
+		if (bDisplayNodes)
 		{
-			if (bMakeGroup)
-			{
-				Open3DDebugSession(FString::Printf(TEXT("Loop %d"), LoopIndex++));
-			}
+			F3DDebugSession Q(bDisplaySegments, FString::Printf(TEXT("FirstNode")));
+			DisplayPoint(*FirstSegmentPoint * DisplayScale, EVisuProperty::BluePoint, 0);
+		}
 
-			const TPoint* FirstSegmentPoint = &Loop[0];
+		for (int32 Index = 1; Index < Loop.Num(); ++Index)
+		{
+			const TPoint* SecondSegmentPoint = &Loop[Index];
+			F3DDebugSession B(bDisplaySegments, FString::Printf(TEXT("Segment %d"), Index));
+			DisplaySegment<TPoint>(*FirstSegmentPoint * DisplayScale, *SecondSegmentPoint * DisplayScale);
 			if (bDisplayNodes)
 			{
-				DisplayPoint(*FirstSegmentPoint * DisplayScale, EVisuProperty::BluePoint, 0);
+				DisplayPoint<TPoint>(*SecondSegmentPoint * DisplayScale, EVisuProperty::BluePoint, Index);
 			}
+			FirstSegmentPoint = SecondSegmentPoint;
+		}
 
-			for (int32 Index = 1; Index < Loop.Num(); ++Index)
-			{
-				const TPoint* SecondSegmentPoint = &Loop[Index];
-				DisplaySegment(*FirstSegmentPoint * DisplayScale, *SecondSegmentPoint * DisplayScale);
-				if (bDisplayNodes)
-				{
-					DisplayPoint(*SecondSegmentPoint * DisplayScale, EVisuProperty::BluePoint, Index);
-				}
-				FirstSegmentPoint = SecondSegmentPoint;
-			}
+		{
+			F3DDebugSession C(bDisplaySegments, FString::Printf(TEXT("Segment %d"), Loop.Num()));
 			DisplaySegment(*FirstSegmentPoint * DisplayScale, *Loop.begin() * DisplayScale);
-			if (bMakeGroup)
-			{
-				Close3DDebugSession();
-			}
 		}
 	}
 
-	void DisplayGridLoops(FString Message, EGridSpace DisplaySpace, bool bDisplayNodes, bool bMakeGroup) const
+	template<typename TPoint>
+	void DisplayGridLoops(FString Message, const TArray<TArray<TPoint>>& Loops, bool bDisplayNodes, bool bMakeGroup, bool bDisplaySegments) const
+	{
+		if (!bDisplay)
+		{
+			return;
+		}
+
+		F3DDebugSession _(*Message);
+		int32 LoopIndex = 0;
+		for (const TArray<TPoint>& Loop : Loops)
+		{
+			FString LoopName = FString::Printf(TEXT("Loop %d"), LoopIndex++);
+			DisplayGridLoop(LoopName, Loop, bDisplayNodes, bMakeGroup, bDisplaySegments);
+		}
+	}
+
+	void DisplayGridLoops(FString Message, EGridSpace DisplaySpace, bool bDisplayNodes, bool bMakeGroup, bool bDisplaySegments) const
 	{
-		DisplayGridLoop(Message, FaceLoops2D[DisplaySpace], bDisplayNodes, bMakeGroup);
+		DisplayGridLoops(Message, FaceLoops2D[DisplaySpace], bDisplayNodes, bMakeGroup, bDisplaySegments);
 	}
 
 #endif
diff --git a/Engine/Source/Runtime/Datasmith/CADKernel/Private/CADKernel/Mesh/Structure/GridBase.cpp b/Engine/Source/Runtime/Datasmith/CADKernel/Private/CADKernel/Mesh/Structure/GridBase.cpp
--- a/Engine/Source/Runtime/Datasmith/CADKernel/Private/CADKernel/Mesh/Structure/GridBase.cpp
+++ b/Engine/Source/Runtime/Datasmith/CADKernel/Private/CADKernel/Mesh/Structure/GridBase.cpp
@@ -3,6 +3,7 @@
 #include "CADKernel/Mesh/Structure/GridBase.h"
 
 #include "CADKernel/Geo/Sampling/SurfacicSampling.h"
+#include "CADKernel/Mesh/Structure/Grid.h"
 #include "CADKernel/Topo/TopologicalFace.h"
 #include "CADKernel/Utils/ArrayUtils.h"
 
diff --git a/Engine/Source/Runtime/Datasmith/CADKernel/Private/CADKernel/Mesh/Structure/GridBase.h b/Engine/Source/Runtime/Datasmith/CADKernel/Private/CADKernel/Mesh/Structure/GridBase.h
--- a/Engine/Source/Runtime/Datasmith/CADKernel/Private/CADKernel/Mesh/Structure/GridBase.h
+++ b/Engine/Source/Runtime/Datasmith/CADKernel/Private/CADKernel/Mesh/Structure/GridBase.h
@@ -104,8 +104,6 @@ protected:
 
 public:
 	FGridChronos Chronos;
-#ifdef CADKERNEL_DEV
-#endif
 
 public:
 	FGridBase(FTopologicalFace& InFace);
@@ -156,9 +154,9 @@ protected:
 		const FPoint2D& PointU1V1 = NewGrid[(IndexV + 1) * CuttingCount[EIso::IsoU] + (IndexU + 1)];
 
 		const double U1MinusU0 = CoordinateGrid[EIso::IsoU][IndexU + 1] - CoordinateGrid[EIso::IsoU][IndexU];
-		const double U0MinusU = CoordinateGrid[EIso::IsoU][IndexU] - InPoint.U;
+		const double U0MinusU  = CoordinateGrid[EIso::IsoU][IndexU] - InPoint.U;
 		const double V1MinusV0 = CoordinateGrid[EIso::IsoV][IndexV + 1] - CoordinateGrid[EIso::IsoV][IndexV];
-		const double V0MinusV = CoordinateGrid[EIso::IsoV][IndexV] - InPoint.V;
+		const double V0MinusV  = CoordinateGrid[EIso::IsoV][IndexV] - InPoint.V;
 
 		OutNewScaledPoint =
 			PointU0V0 +
@@ -229,7 +227,7 @@ public:
 		return Points2D[(int32)Space];
 	}
 
-const FTopologicalFace& GetFace() const
+	const FTopologicalFace& GetFace() const
 	{
 		return Face;
 	}
@@ -263,7 +261,7 @@ const FTopologicalFace& GetFace() const
 	// ======================================================================================================================================================================================================================
 	// Display Methodes   ================================================================================================================================================================================================
 	// ======================================================================================================================================================================================================================
-	bool bDisplay = true;
+	mutable bool bDisplay = true;
 
 	void DisplayIsoNode(EGridSpace Space, const int32 PointIndex, FIdent Ident = 0, EVisuProperty Property = EVisuProperty::BluePoint) const;
 	virtual void DisplayGridPoints(EGridSpace DisplaySpace) const;
diff --git a/Engine/Source/Runtime/Datasmith/CADKernel/Private/CADKernel/Mesh/Structure/LoopCleaner.cpp b/Engine/Source/Runtime/Datasmith/CADKernel/Private/CADKernel/Mesh/Structure/LoopCleaner.cpp
--- a/Engine/Source/Runtime/Datasmith/CADKernel/Private/CADKernel/Mesh/Structure/LoopCleaner.cpp
+++ b/Engine/Source/Runtime/Datasmith/CADKernel/Private/CADKernel/Mesh/Structure/LoopCleaner.cpp
@@ -8,44 +8,58 @@
 #include "CADKernel/Mesh/Structure/Grid.h"
 #include "CADKernel/Topo/TopologicalFace.h"
 
+#ifdef CADKERNEL_DEV
+#include "CADKernel/Mesh/Meshers/MesherReport.h"
+#endif
+
 namespace UE::CADKernel
 {
 
 FLoopCleaner::FLoopCleaner(FIsoTriangulator& Triangulator)
 	: Grid(Triangulator.Grid)
+	, Tolerances(Triangulator.Tolerances)
 	, LoopNodes(Triangulator.LoopNodes)
 	, LoopSegments(Triangulator.LoopSegments)
 	, IsoSegmentFactory(Triangulator.IsoSegmentFactory)
 	, bDisplay(Triangulator.bDisplay)
 	, LoopSegmentsIntersectionTool(Grid)
-	, GeometricTolerance(Triangulator.GeometricTolerance)
-	, SquareGeometricTolerance(FMath::Square(GeometricTolerance))
-	, SquareGeometricTolerance2(2. * SquareGeometricTolerance)
 	, GetNext(LoopCleanerImpl::GetNextNodeImpl)
 	, GetPrevious(LoopCleanerImpl::GetPreviousNodeImpl)
 	, GetFirst(LoopCleanerImpl::GetFirstNode)
 	, GetSecond(LoopCleanerImpl::GetSecondNode)
 {
+}
+
+bool FLoopCleaner::Run()
+{
+	if (!CleanLoops())
+	{
+#ifdef CADKERNEL_DEV
+		FMesherReport::GetLogs().AddRemoveSelfIntersectionFailure();
+#endif
+		return false;
+	}
+
+	if (!UncrossLoops())
+	{
 #ifdef CADKERNEL_DEV
-	SetMesherReport(*Triangulator.MesherReport);
+		FMesherReport::GetLogs().AddRemoveCrossingLoopsFailure();
 #endif
+		return false;
+	}
+
+	return true;
 }
 
 bool FLoopCleaner::CleanLoops()
 {
-
-	//{
-	//	FPoint2D Point = LoopNodes[36].Get2DPoint(EGridSpace::UniformScaled, Grid);
-	//	LoopNodes[56].Set2DPoint(EGridSpace::UniformScaled, Grid, Point);
-	//}
-
 	FindBestLoopExtremity();
 
 #ifdef DEBUG_CLEAN_LOOPS		
 	if (bDisplay)
 	{
-		Grid.DisplayIsoSegments(TEXT("Loops Orientation"), EGridSpace::UniformScaled, LoopSegments, false, true, EVisuProperty::BlueCurve);
-		Grid.DisplayIsoSegments(TEXT("Loops Init"), EGridSpace::UniformScaled, LoopSegments, true, false, EVisuProperty::BlueCurve);
+		Grid.DisplayIsoSegments(TEXT("Loops Orientation"), EGridSpace::UniformScaled, LoopSegments, false, true, false, EVisuProperty::BlueCurve);
+		Grid.DisplayIsoSegments(TEXT("Loops Init"), EGridSpace::UniformScaled, LoopSegments, true, false, false, EVisuProperty::BlueCurve);
 		Grid.DisplayIsoNodes(TEXT("BestStartNodeOfLoop"), EGridSpace::UniformScaled, (const TArray<const FIsoNode*>) BestStartNodeOfLoops, EVisuProperty::BluePoint);
 		Wait(bDisplay);
 	}
@@ -62,9 +76,20 @@ bool FLoopCleaner::CleanLoops()
 	for (FLoopNode* StartNode : BestStartNodeOfLoops)
 	{
 		LoopIndex++;
+
 		StartSegmentIndex = NextLoopFirstSegmentIndex;
 		UpdateNextLoopFirstSegmentIndex(LoopIndex);
 
+		if (!FindAndRemoveCoincidence(StartNode))
+		{
+			if (StartSegmentIndex == 0)
+			{
+				FMessage::Printf(Log, TEXT("The outer loop of the surface %d is degenerated. The mesh of this surface is canceled.\n"), Grid.GetFace().GetId());
+				return false;
+			}
+			continue;
+		}
+
 		Intersections.Empty(5);
 
 		GetLoopNodeStartingFrom(StartNode, NodesOfLoop);
@@ -81,25 +106,35 @@ bool FLoopCleaner::CleanLoops()
 
 #ifdef DEBUG_CLEAN_LOOPS		
 		Grid.DisplayGridPolyline(TEXT("LoopIntersections: remove loop's picks"), EGridSpace::UniformScaled, NodesOfLoop, true, EVisuProperty::YellowCurve);
-		//Wait(bDisplay);
+		Wait(bDisplay);
 #endif
 		
 		FindLoopIntersections();
 
-		if (!RemoveSelfIntersectionsOfLoop())
+		if(Intersections.Num())
 		{
-			FMessage::Printf(Log, TEXT("Loop intersections of the surface %d cannot be fixed. The mesh of this surface is canceled.\n"), Grid.GetFace().GetId());
-			return false;
-		}
+			if (!RemoveSelfIntersectionsOfLoop())
+			{
+				FMessage::Printf(Log, TEXT("Loop intersections of the surface %d cannot be fixed. The mesh of this surface is canceled.\n"), Grid.GetFace().GetId());
+				return false;
+			}
 
+			Intersections.SetNum(0);
+			FindLoopIntersections();
+			if (Intersections.Num())
+			{
+				FMessage::Printf(Log, TEXT("Loop intersections of the surface %d cannot be fixed. The mesh of this surface is canceled.\n"), Grid.GetFace().GetId());
+				return false;
+			}
 #ifdef DEBUG_CLEAN_LOOPS		
-		Grid.DisplayGridPolyline(TEXT("LoopIntersections: remove loop's self intersections"), EGridSpace::UniformScaled, NodesOfLoop, true, EVisuProperty::YellowCurve);
-		//Wait(bDisplay);
+			Grid.DisplayGridPolyline(TEXT("LoopIntersections: remove loop's self intersections"), EGridSpace::UniformScaled, NodesOfLoop, true, EVisuProperty::YellowCurve);
+			//Wait(bDisplay);
 #endif
+		}
 
 		if (!RemoveLoopPicks())
 		{
-			FMessage::Printf(Log, TEXT("Loop intersections of the surface %d cannot be fixed. The mesh of this surface is canceled.\n"), Grid.GetFace().GetId());
+			FMessage::Printf(Log, TEXT("Loop picks of the surface %d cannot be fixed. The mesh of this surface is canceled.\n"), Grid.GetFace().GetId());
 			return false;
 		}
 
@@ -160,7 +195,7 @@ bool FLoopCleaner::UncrossLoops()
 #ifdef DEBUG_UNCROSS_LOOPS
 		if (bDisplay)
 		{
-			LoopSegmentsIntersectionTool.Display(TEXT("IntersectionTool"));
+			LoopSegmentsIntersectionTool.Display(bDisplay, TEXT("IntersectionTool"), EVisuProperty::BlueCurve);
 			F3DDebugSession _(*FString::Printf(TEXT("Segment to be processed %d %d"), Index, Iteration++));
 			Grid.DisplayIsoSegment(EGridSpace::UniformScaled, Segment, 0, EVisuProperty::BlueCurve);
 		}
@@ -172,7 +207,7 @@ bool FLoopCleaner::UncrossLoops()
 			if (bDisplay)
 			{
 				{
-					LoopSegmentsIntersectionTool.Display(TEXT("IntersectionTool"));
+					LoopSegmentsIntersectionTool.Display(bDisplay, TEXT("IntersectionTool"), EVisuProperty::BlueCurve);
 					F3DDebugSession _(*FString::Printf(TEXT("Segment to be processed %d %d"), Index, Iteration++));
 					Grid.DisplayIsoSegment(EGridSpace::UniformScaled, Segment, 0, EVisuProperty::BlueCurve);
 				}
@@ -218,7 +253,7 @@ bool FLoopCleaner::UncrossLoops()
 			}
 
 #ifdef DEBUG_UNCROSS_LOOPS
-			if (true)
+			if (bDisplay)
 			{
 				Grid.DisplayGridPolyline(TEXT("After fix"), EGridSpace::UniformScaled, LoopNodes, true, EVisuProperty::YellowCurve);
 				Wait(true);
@@ -234,7 +269,7 @@ bool FLoopCleaner::UncrossLoops()
 		{
 			RemovePickRecursively((FLoopNode*) &Segment.GetFirstNode(), (FLoopNode*) &Segment.GetSecondNode());
 #ifdef DEBUG_UNCROSS_LOOPS
-			if (true)
+			if (bDisplay)
 			{
 				Grid.DisplayGridPolyline(TEXT("After RemovePickRecursively"), EGridSpace::UniformScaled, LoopNodes, true, EVisuProperty::YellowCurve);
 				Wait(true);
@@ -277,7 +312,7 @@ bool FLoopCleaner::TryToRemoveIntersectionOfTwoConsecutiveIntersectingSegments(c
 	{
 		double StartPointSquareDistance = SquareDistanceOfPointToSegment(Segment2D.Point0, IntersectingSegment2D.Point0, IntersectingSegment2D.Point1);
 		double EndPointSquareDistance = SquareDistanceOfPointToSegment(Segment2D.Point1, IntersectingSegment2D.Point0, IntersectingSegment2D.Point1);
-		if (StartPointSquareDistance < SquareGeometricTolerance && EndPointSquareDistance < SquareGeometricTolerance)
+		if (StartPointSquareDistance < Tolerances.SquareGeometricTolerance && EndPointSquareDistance < Tolerances.SquareGeometricTolerance)
 		{
 			OffsetSegment(Segment, Segment2D, IntersectingSegment2D);
 			return true;
@@ -328,7 +363,7 @@ bool FLoopCleaner::TryToRemoveIntersectionOfTwoConsecutiveIntersectingSegments(c
 #endif
 
 	TSegment<FPoint2D> NextSegment2D(Node->Get2DPoint(EGridSpace::UniformScaled, Grid), NextNode->Get2DPoint(EGridSpace::UniformScaled, Grid));
-	if (!FastIntersectSegments2D(NextSegment2D, IntersectingSegment2D))
+	if (!FastDoIntersect(NextSegment2D, IntersectingSegment2D))
 	{
 		return false;
 	}
@@ -353,7 +388,7 @@ bool FLoopCleaner::TryToRemoveIntersectionOfTwoConsecutiveIntersectingSegments(c
 	SegmentTangent.Normalize();
 	FPoint2D MoveDirection = SegmentTangent.GetPerpendicularVector();
 
-	MoveDirection *= GeometricTolerance;
+	MoveDirection *= Tolerances.GeometricTolerance;
 
 	FPoint2D NewCoordinate = ProjectedPoint + MoveDirection;
 
@@ -365,8 +400,8 @@ bool FLoopCleaner::TryToRemoveIntersectionOfTwoConsecutiveIntersectingSegments(c
 	}
 #endif
 
-	if ((PreviousNode->Get2DPoint(EGridSpace::UniformScaled, Grid).SquareDistance(NewCoordinate) < SquareGeometricTolerance2) ||
-		(NextNode->Get2DPoint(EGridSpace::UniformScaled, Grid).SquareDistance(NewCoordinate) < SquareGeometricTolerance2))
+	if ((PreviousNode->Get2DPoint(EGridSpace::UniformScaled, Grid).SquareDistance(NewCoordinate) < Tolerances.SquareGeometricTolerance2) ||
+		(NextNode->Get2DPoint(EGridSpace::UniformScaled, Grid).SquareDistance(NewCoordinate) < Tolerances.SquareGeometricTolerance2))
 	{
 		if(RemoveNodeOfLoop(*Node))
 		{
@@ -402,7 +437,7 @@ void FLoopCleaner::OffsetSegment(FIsoSegment& Segment, TSegment<FPoint2D>& Segme
 	FPoint2D SegmentTangent = IntersectingSegment2D.Point1 - IntersectingSegment2D.Point0;
 	SegmentTangent.Normalize();
 	FPoint2D MoveDirection = SegmentTangent.GetPerpendicularVector();
-	MoveDirection *= GeometricTolerance;
+	MoveDirection *= Tolerances.GeometricTolerance;
 
 	FPoint2D NewPoint0 = Segment2D.Point0 + MoveDirection;
 	FPoint2D NewPoint1 = Segment2D.Point1 + MoveDirection;
@@ -416,7 +451,7 @@ void FLoopCleaner::OffsetNode(FLoopNode& Node, TSegment<FPoint2D>& IntersectingS
 	FPoint2D SegmentTangent = IntersectingSegment2D.Point1 - IntersectingSegment2D.Point0;
 	SegmentTangent.Normalize();
 	FPoint2D MoveDirection = SegmentTangent.GetPerpendicularVector();
-	MoveDirection *= GeometricTolerance;
+	MoveDirection *= Tolerances.GeometricTolerance;
 
 	FPoint2D NewPoint = Node.Get2DPoint(EGridSpace::UniformScaled, Grid) + MoveDirection;
 	Node.Set2DPoint(EGridSpace::UniformScaled, Grid, NewPoint);
@@ -478,7 +513,7 @@ void FLoopCleaner::RemoveIntersectionByMovingOutsideSegmentNodeInside(const FIso
 	FPoint2D MoveDirection = IntersectingPoints[1] - IntersectingPoints[0];
 	MoveDirection.Normalize();
 	MoveDirection = MoveDirection.GetPerpendicularVector();
-	MoveDirection *= bIsSameInnerLoop ? -GeometricTolerance : GeometricTolerance;
+	MoveDirection *= bIsSameInnerLoop ? -Tolerances.GeometricTolerance : Tolerances.GeometricTolerance;
 
 	FPoint2D NewCoordinate = ProjectedPoint + MoveDirection;
 
@@ -486,8 +521,8 @@ void FLoopCleaner::RemoveIntersectionByMovingOutsideSegmentNodeInside(const FIso
 	FLoopNode& PreviousNode = Node.GetPreviousNode();
 	FLoopNode& NextNode = Node.GetNextNode();
 
-	if ((PreviousNode.Get2DPoint(EGridSpace::UniformScaled, Grid).SquareDistance(NewCoordinate) < SquareGeometricTolerance2) ||
-		(NextNode.Get2DPoint(EGridSpace::UniformScaled, Grid).SquareDistance(NewCoordinate) < SquareGeometricTolerance2))
+	if ((PreviousNode.Get2DPoint(EGridSpace::UniformScaled, Grid).SquareDistance(NewCoordinate) < Tolerances.SquareGeometricTolerance2) ||
+		(NextNode.Get2DPoint(EGridSpace::UniformScaled, Grid).SquareDistance(NewCoordinate) < Tolerances.SquareGeometricTolerance2))
 	{
 		RemoveNodeOfLoop(Node);
 		return;
@@ -519,6 +554,7 @@ void FLoopCleaner::FixLoopOrientation()
 	{
 #ifdef DEBUG_LOOP_ORIENTATION
 		Grid.DisplayIsoSegments(TEXT("Before orientation"), EGridSpace::UniformScaled, LoopSegments, true, true);
+		Wait(bDisplay);
 #endif
 
 		SwapSubLoopOrientation(StartSegmentIndex, NextLoopFirstSegmentIndex);
@@ -727,7 +763,7 @@ bool FLoopCleaner::RemoveOutgoingLoop(const TPair<double, double>& Intersection,
 
 	MoveDirection /= Length;
 	MoveDirection = MoveDirection.GetPerpendicularVector();
-	MoveDirection *= GeometricTolerance;
+	MoveDirection *= Tolerances.GeometricTolerance;
 	MiddlePoint += MoveDirection;
 
 	if (!EndNode->IsDelete())
@@ -746,7 +782,7 @@ bool FLoopCleaner::RemoveIntersectionsOfSubLoop(int32 IntersectionIndex, int32 I
 	{
 		FPoint2D MoveDirection = ProjectedPoint - PointToProject;
 		MoveDirection.Normalize();
-		MoveDirection *= GeometricTolerance;
+		MoveDirection *= Tolerances.GeometricTolerance;
 		ProjectedPoint += MoveDirection;
 		NodeToProject->Set2DPoint(EGridSpace::UniformScaled, Grid, ProjectedPoint);
 	};
@@ -812,7 +848,7 @@ bool FLoopCleaner::RemoveIntersectionsOfSubLoop(int32 IntersectionIndex, int32 I
 		FPoint2D MoveDirection = Point1 - Point0;
 		MoveDirection.Normalize();
 		MoveDirection = MoveDirection.GetPerpendicularVector();
-		MoveDirection *= GeometricTolerance;
+		MoveDirection *= Tolerances.GeometricTolerance;
 
 		FPoint2D NewCoordinate = Point1 + MoveDirection;
 		NodeToMove->Set2DPoint(EGridSpace::UniformScaled, Grid, NewCoordinate);
@@ -1270,8 +1306,8 @@ bool FLoopCleaner::DisconnectCoincidentNodes(const LoopCleanerImpl::FPinchInters
 
 	FPoint2D P1 = SlopeToVector(MediumSlope1);
 	FPoint2D P0 = SlopeToVector(MediumSlope0);
-	P1.Normalize()*= GeometricTolerance*30;
-	P0.Normalize()*= GeometricTolerance*30;
+	P1.Normalize()*= Tolerances.GeometricTolerance*30;
+	P0.Normalize()*= Tolerances.GeometricTolerance*30;
 
 	P0 += *Context.Points[0][1];
 	P1 += *Context.Points[1][1];
@@ -1356,7 +1392,7 @@ bool FLoopCleaner::DisconnectCrossingSegments(LoopCleanerImpl::FPinchIntersectio
 	int32 Segment10_11Index = LoopSegments.IndexOfByKey(Segment10_11);
 
 #ifdef DEBUG_REMOVE_UNIQUE_INTERSECTION
-	Grid.DisplayIsoSegments(TEXT("Loops Orientation"), EGridSpace::UniformScaled, LoopSegments, false, true, EVisuProperty::BlueCurve);
+	Grid.DisplayIsoSegments(TEXT("Loops Orientation"), EGridSpace::UniformScaled, LoopSegments, false, true, false, EVisuProperty::BlueCurve);
 	Grid.DisplayGridPolyline(TEXT("RemoveIntersection end"), EGridSpace::UniformScaled, NodesOfLoop, true);
 	Wait(bDisplay);
 #endif
@@ -1364,7 +1400,7 @@ bool FLoopCleaner::DisconnectCrossingSegments(LoopCleanerImpl::FPinchIntersectio
 	SwapSubLoopOrientation(Segment01_02Index + 1, Segment10_11Index);
 
 #ifdef DEBUG_REMOVE_UNIQUE_INTERSECTION
-	Grid.DisplayIsoSegments(TEXT("Loops Orientation"), EGridSpace::UniformScaled, LoopSegments, false, true, EVisuProperty::BlueCurve);
+	Grid.DisplayIsoSegments(TEXT("Loops Orientation"), EGridSpace::UniformScaled, LoopSegments, false, true, false, EVisuProperty::BlueCurve);
 	Grid.DisplayGridPolyline(TEXT("RemoveIntersection end"), EGridSpace::UniformScaled, NodesOfLoop, true);
 	Wait(bDisplay);
 #endif
@@ -1569,15 +1605,15 @@ void FLoopCleaner::MoveNode(FLoopNode& NodeToMove, FPoint2D& NewPostion)
 
 	FPoint2D ProjectedDirection = NewPostion - PointToMove;
 	ProjectedDirection.Normalize();
-	ProjectedDirection *= GeometricTolerance;
+	ProjectedDirection *= Tolerances.GeometricTolerance;
 
 	FPoint2D NewCoordinate = NewPostion + ProjectedDirection;
 
 	FLoopNode& PreviousNode = NodeToMove.GetPreviousNode();
 	FLoopNode& NextNode = NodeToMove.GetNextNode();
 
-	if ((PreviousNode.Get2DPoint(EGridSpace::UniformScaled, Grid).SquareDistance(NewCoordinate) < SquareGeometricTolerance2) ||
-		(NextNode.Get2DPoint(EGridSpace::UniformScaled, Grid).SquareDistance(NewCoordinate) < SquareGeometricTolerance2))
+	if ((PreviousNode.Get2DPoint(EGridSpace::UniformScaled, Grid).SquareDistance(NewCoordinate) < Tolerances.SquareGeometricTolerance2) ||
+		(NextNode.Get2DPoint(EGridSpace::UniformScaled, Grid).SquareDistance(NewCoordinate) < Tolerances.SquareGeometricTolerance2))
 	{
 		RemoveNodeOfLoop(NodeToMove);
 	}
@@ -1606,13 +1642,12 @@ void FLoopCleaner::MoveNode(FLoopNode& NodeToMove, FPoint2D& NewPostion)
 #endif
 }
 
-//#define DEBUG_FIND_LOOP_INTERSECTIONS
 void FLoopCleaner::FindLoopIntersections()
 {
 	LoopSegmentsIntersectionTool.Empty(NodesOfLoopCount);
 
 	FLoopNode* const* const StartNodePtr = NodesOfLoop.FindByPredicate([](const FLoopNode* Node) { return !Node->IsDelete(); });
-	if (*StartNodePtr == nullptr || (*StartNodePtr)->IsDelete())
+	if (StartNodePtr == nullptr || *StartNodePtr == nullptr || (*StartNodePtr)->IsDelete())
 	{
 		return;
 	}
@@ -1799,6 +1834,63 @@ bool FLoopCleaner::RemoveLoopPicks()
 	return UpdateNodesOfLoop();
 }
 
+bool FLoopCleaner::FindAndRemoveCoincidence(FLoopNode*& StartNode)
+{
+#ifdef DEBUG_FIND_AND_REMOVE_COINCIDENCE
+	F3DDebugSession A(bDisplay, ("FindAndRemoveCoincidence"));
+#endif
+
+	FLoopNode* PreviousNode = StartNode;
+	const FPoint2D* PreviousPoint = &StartNode->Get2DPoint(EGridSpace::UniformScaled, Grid);
+	for (FLoopNode* Node = GetNext(StartNode); ; Node = GetNext(Node))
+	{
+		const FPoint2D& Point = Node->Get2DPoint(EGridSpace::UniformScaled, Grid);
+
+		if (CheckAndRemoveCoincidence(*PreviousPoint, Point, *PreviousNode))
+		{
+			if (Node->IsDelete())
+			{
+				return false;
+			}
+
+			if (PreviousNode == StartNode)
+			{
+				StartNode = Node;
+#ifdef DEBUG_FIND_AND_REMOVE_COINCIDENCE
+				if(bDisplay)
+				{
+					F3DDebugSession B(TEXT("RemoveCoincidence"));
+					DisplayPoint2DWithScale(*PreviousPoint, EVisuProperty::RedPoint);
+				}
+#endif
+				continue;
+			}
+#ifdef DEBUG_FIND_AND_REMOVE_COINCIDENCE
+			if (bDisplay)
+			{
+				F3DDebugSession B(TEXT("RemoveCoincidence"));
+				DisplayPoint2DWithScale(*PreviousPoint, EVisuProperty::RedPoint);
+			}
+		}
+		else if(bDisplay)
+		{
+			F3DDebugSession B(TEXT("Seg"));
+			DisplaySegmentWithScale(*PreviousPoint, Point, 1, EVisuProperty::BlueCurve);
+#endif
+		}
+	
+		if (Node == StartNode)
+		{
+			break;
+		}
+
+		PreviousNode = Node;
+		PreviousPoint = &Point;
+	}
+
+	return true;
+}
+
 bool FLoopCleaner::RemovePickRecursively(FLoopNode* Node0, FLoopNode* Node1)
 {
 	if (Node0->IsDelete() || Node1->IsDelete())
diff --git a/Engine/Source/Runtime/Datasmith/CADKernel/Private/CADKernel/Mesh/Structure/LoopCleaner.h b/Engine/Source/Runtime/Datasmith/CADKernel/Private/CADKernel/Mesh/Structure/LoopCleaner.h
--- a/Engine/Source/Runtime/Datasmith/CADKernel/Private/CADKernel/Mesh/Structure/LoopCleaner.h
+++ b/Engine/Source/Runtime/Datasmith/CADKernel/Private/CADKernel/Mesh/Structure/LoopCleaner.h
@@ -7,10 +7,8 @@
 #include "CADKernel/Mesh/Meshers/IsoTriangulator/IntersectionSegmentTool.h"
 #include "CADKernel/Mesh/Meshers/IsoTriangulator/IsoNode.h"
 #include "CADKernel/Mesh/Meshers/IsoTriangulator/IsoSegment.h"
+#include "CADKernel/Mesh/Meshers/ParametricMesherConstantes.h"
 
-#ifdef CADKERNEL_DEV
-#include "CADKernel/Mesh/Meshers/MesherReport.h"
-#endif
 
 
 namespace UE::CADKernel
@@ -19,8 +17,7 @@ class FGrid;
 class FIsoSegment;
 class FIsoTriangulator;
 class FLoopNode;
-
-struct FMesherReport;
+class FMeshingTolerances;
 
 namespace LoopCleanerImpl
 {
@@ -97,14 +94,12 @@ class FLoopCleaner
 {
 private:
 	FGrid& Grid;
+	const FMeshingTolerances& Tolerances;
+
 	TArray<FLoopNode>& LoopNodes;
 	TArray<FIsoSegment*>& LoopSegments;
 	TFactory<FIsoSegment>& IsoSegmentFactory;
 
-#ifdef CADKERNEL_DEV
-	FMesherReport* MesherReport;
-#endif
-
 	bool bDisplay;
 
 	TArray<FLoopNode*> BestStartNodeOfLoops;
@@ -141,12 +136,6 @@ private:
 	TArray<TPair<double, double>> Intersections;
 	FIntersectionSegmentTool LoopSegmentsIntersectionTool;
 
-	const double GeometricTolerance;
-	const double SquareGeometricTolerance;
-	const double SquareGeometricTolerance2;
-	//const double MeshingTolerance;
-	//const double SquareMeshingTolerance;
-
 	LoopCleanerImpl::GetNextNodeMethod GetNext;
 	LoopCleanerImpl::GetNextNodeMethod GetPrevious;
 	LoopCleanerImpl::GetSegmentToNodeMethod GetFirst;
@@ -156,36 +145,10 @@ public:
 
 	FLoopCleaner(FIsoTriangulator& Triangulator);
 
-	bool Run()
-	{
-		if (!CleanLoops())
-		{
-#ifdef CADKERNEL_DEV
-			MesherReport->Logs.AddRemoveSelfIntersectionFailure();
-#endif
-			return false;
-		}
-		
-		if (!UncrossLoops())
-		{
-#ifdef CADKERNEL_DEV
-			MesherReport->Logs.AddRemoveCrossingLoopsFailure();
-#endif
-			return false;
-		}
-
-		return true;
-	}
+	bool Run();
 
 private:
 
-#ifdef CADKERNEL_DEV
-	void SetMesherReport(FMesherReport& InMesherReport)
-	{
-		MesherReport = &InMesherReport;
-	}
-#endif
-
 	bool CleanLoops();
 	bool UncrossLoops();
 
@@ -204,6 +167,7 @@ private:
 	 */
 	bool RemoveLoopPicks();
 	bool RemovePickRecursively(FLoopNode* Node0, FLoopNode* Node1);
+	bool FindAndRemoveCoincidence(FLoopNode*& StartNode);
 
 	/**
 	 * @return true if the node has been deleted
@@ -217,7 +181,7 @@ private:
 			double SquareDistance1 = SquareDistanceOfPointToSegment(PreviousPoint, NodeToRemovePoint, NextPoint);
 			double SquareDistance2 = SquareDistanceOfPointToSegment(NextPoint, NodeToRemovePoint, PreviousPoint);
 			double MinSquareDistance = FMath::Min(SquareDistance1, SquareDistance2);
-			if (MinSquareDistance < SquareGeometricTolerance2)
+			if (MinSquareDistance < Tolerances.SquareGeometricTolerance2)
 			{
 				bRemoveNode = true;
 			}
@@ -236,7 +200,7 @@ private:
 	bool CheckAndRemoveCoincidence(const FPoint2D& Point0, const FPoint2D& Point1, FLoopNode& NodeToRemove)
 	{
 		double SquareDistance = Point0.SquareDistance(Point1);
-		if (SquareDistance < SquareGeometricTolerance2)
+		if (SquareDistance < Tolerances.SquareGeometricTolerance2)
 		{
 			return RemoveNodeOfLoop(NodeToRemove);
 		}
diff --git a/Engine/Source/Runtime/Datasmith/CADKernel/Private/CADKernel/Mesh/Structure/Mesh.cpp b/Engine/Source/Runtime/Datasmith/CADKernel/Private/CADKernel/Mesh/Structure/Mesh.cpp
--- a/Engine/Source/Runtime/Datasmith/CADKernel/Private/CADKernel/Mesh/Structure/Mesh.cpp
+++ b/Engine/Source/Runtime/Datasmith/CADKernel/Private/CADKernel/Mesh/Structure/Mesh.cpp
@@ -2,12 +2,25 @@
 
 #include "CADKernel/Mesh/Structure/Mesh.h"
 
+#include "CADKernel/Geo/Sampling/PolylineTools.h"
+#include "CADKernel/Mesh/Structure/EdgeMesh.h"
 #include "CADKernel/Mesh/Structure/ModelMesh.h"
 #include "CADKernel/Topo/TopologicalEntity.h"
+#include "CADKernel/Topo/TopologicalEdge.h"
 
 namespace UE::CADKernel
 {
 
+TArray<double> FEdgeMesh::GetElementLengths() const
+{
+	const FTopologicalEdge& Edge = static_cast<const FTopologicalEdge&>(TopologicalEntity);
+	const TArray<FPoint>& MeshInnerNodes = GetNodeCoordinates();
+
+	const FPoint& StartNode = Edge.GetStartVertex()->GetCoordinates();
+	const FPoint& EndNode = Edge.GetEndVertex()->GetCoordinates();
+	return PolylineTools::ComputePolylineSegmentLengths(StartNode, MeshInnerNodes, EndNode);
+}
+
 int32 FMesh::RegisterCoordinates()
 {
 	ModelMesh.RegisterCoordinates(NodeCoordinates, StartNodeId, MeshModelIndex);
diff --git a/Engine/Source/Runtime/Datasmith/CADKernel/Private/CADKernel/Mesh/Structure/ModelMesh.cpp b/Engine/Source/Runtime/Datasmith/CADKernel/Private/CADKernel/Mesh/Structure/ModelMesh.cpp
--- a/Engine/Source/Runtime/Datasmith/CADKernel/Private/CADKernel/Mesh/Structure/ModelMesh.cpp
+++ b/Engine/Source/Runtime/Datasmith/CADKernel/Private/CADKernel/Mesh/Structure/ModelMesh.cpp
@@ -45,33 +45,34 @@ FInfoEntity& FModelMesh::GetInfo(FInfoEntity& Info) const
 }
 #endif
 
-const int32 FModelMesh::GetIndexOfVertexFromId(const int32 Ident) const
+const FVertexMesh* FModelMesh::GetMeshOfVertexNodeId(const int32 Ident) const
 {
-	for (const TSharedPtr<FVertexMesh>& VertexMesh : VertexMeshes)
+	for (const FVertexMesh* VertexMesh : VertexMeshes)
 	{
 		if (Ident == VertexMesh->GetStartVertexId())
 		{
-			return VertexMesh->GetIndexInMeshModel();
+			return VertexMesh;
 		}
 	}
-	return -1;
+	return nullptr;
 }
 
-const TSharedPtr<FVertexMesh> FModelMesh::GetMeshOfVertexNodeId(const int32 Ident) const
+#ifdef OLD
+const int32 FModelMesh::GetIndexOfVertexFromId(const int32 Ident) const
 {
-	for (const TSharedPtr<FVertexMesh>& VertexMesh : VertexMeshes)
+	for (const FVertexMesh* VertexMesh : VertexMeshes)
 	{
 		if (Ident == VertexMesh->GetStartVertexId())
 		{
-			return VertexMesh;
+			return VertexMesh->GetIndexInMeshModel();
 		}
 	}
-	return nullptr;
+	return -1;
 }
 
 const int32 FModelMesh::GetIndexOfEdgeFromId(const int32 Ident) const
 {
-	for (const TSharedPtr<FEdgeMesh>& EdgeMesh : EdgeMeshes)
+	for (const FEdgeMesh* EdgeMesh : EdgeMeshes)
 	{
 		if (Ident >= EdgeMesh->GetStartVertexId())
 		{
@@ -98,6 +99,7 @@ const int32 FModelMesh::GetIndexOfSurfaceFromId(const int32 Ident) const
 	}
 	return -1;
 }
+#endif
 
 void FModelMesh::GetNodeCoordinates(TArray<FPoint>& NodeCoordinates) const
 {
@@ -122,23 +124,23 @@ void FModelMesh::GetNodeCoordinates(TArray<FVector3f>& NodeCoordinates) const
 	}
 }
 
-const TArray<TSharedPtr<FMesh>>& FModelMesh::GetMeshes() const
+const TArray<FMesh*>& FModelMesh::GetMeshes() const
 {
 	if (FaceMeshes.Num())
 	{
-		return (TArray<TSharedPtr<FMesh>>&) FaceMeshes;
+		return (TArray<FMesh*>&) FaceMeshes;
 	}
 	if (EdgeMeshes.Num())
 	{
-		return (TArray<TSharedPtr<FMesh>>&) EdgeMeshes;
+		return (TArray<FMesh*>&) EdgeMeshes;
 	}
-	return (TArray<TSharedPtr<FMesh>>&) VertexMeshes;
+	return (TArray<FMesh*>&) VertexMeshes;
 }
 
 int32 FModelMesh::GetTriangleCount() const
 {
 	int32 TriangleCount = 0;
-	for (const TSharedPtr<FFaceMesh>& FaceMesh : FaceMeshes)
+	for (const FFaceMesh* FaceMesh : FaceMeshes)
 	{
 		TriangleCount += FaceMesh->TrianglesVerticesIndex.Num() / 3;
 	}
diff --git a/Engine/Source/Runtime/Datasmith/CADKernel/Private/CADKernel/Mesh/Structure/ModelMeshAnalyzer.cpp b/Engine/Source/Runtime/Datasmith/CADKernel/Private/CADKernel/Mesh/Structure/ModelMeshAnalyzer.cpp
--- a/Engine/Source/Runtime/Datasmith/CADKernel/Private/CADKernel/Mesh/Structure/ModelMeshAnalyzer.cpp
+++ b/Engine/Source/Runtime/Datasmith/CADKernel/Private/CADKernel/Mesh/Structure/ModelMeshAnalyzer.cpp
@@ -21,7 +21,7 @@ void FModelMeshAnalyzer::BuildMesh()
 	ModelMesh.GetNodeCoordinates(NodeCoordinates);
 
 	int32 TriangleCount = 0;
-	for (const TSharedPtr<FFaceMesh>& FaceMesh : ModelMesh.GetFaceMeshes())
+	for (const FFaceMesh* FaceMesh : ModelMesh.GetFaceMeshes())
 	{
 		TriangleCount += FaceMesh->TrianglesVerticesIndex.Num() / 3;
 	}
@@ -49,11 +49,11 @@ void FModelMeshAnalyzer::BuildMesh()
 		return &Edge;
 	};
 
-	for (const TSharedPtr<FFaceMesh>& FaceMesh : ModelMesh.GetFaceMeshes())
+	for (const FFaceMesh* FaceMesh : ModelMesh.GetFaceMeshes())
 	{
-		TArray<int32>& TrianglesVerticesIndex = FaceMesh->TrianglesVerticesIndex;
-		TArray<int32>& VerticesGlobalIndex = FaceMesh->VerticesGlobalIndex;
-		TArray<FVector3f>& FaceNormals = FaceMesh->Normals;
+		const TArray<int32>& TrianglesVerticesIndex = FaceMesh->TrianglesVerticesIndex;
+		const TArray<int32>& VerticesGlobalIndex = FaceMesh->VerticesGlobalIndex;
+		const TArray<FVector3f>& FaceNormals = FaceMesh->Normals;
 
 		for (int32 Index = 0; Index < FaceMesh->TrianglesVerticesIndex.Num(); Index += 3)
 		{
@@ -73,7 +73,7 @@ void FModelMeshAnalyzer::BuildMesh()
 				Normals[NormalI] = FaceNormals[TrianglesVerticesIndex[Index + NormalI]];
 			}
 
-			Analyzer::FTriangle& Triangle = Triangles.Emplace_GetRef(Vertices, TriangleEdges, Normals, FaceMesh.Get());
+			Analyzer::FTriangle& Triangle = Triangles.Emplace_GetRef(Vertices, TriangleEdges, Normals, FaceMesh);
 
 			for (uint32 EdgeI = 0; EdgeI < 3; ++EdgeI)
 			{
