CADKernel update: Meshing of thin faces

diff --git a/Engine/Source/Runtime/Datasmith/CADKernel/Private/CADKernel/Mesh/Structure/ThinZone2DFinder.cpp b/Engine/Source/Runtime/Datasmith/CADKernel/Private/CADKernel/Mesh/Structure/ThinZone2DFinder.cpp
--- a/Engine/Source/Runtime/Datasmith/CADKernel/Private/CADKernel/Mesh/Structure/ThinZone2DFinder.cpp
+++ b/Engine/Source/Runtime/Datasmith/CADKernel/Private/CADKernel/Mesh/Structure/ThinZone2DFinder.cpp
@@ -5,6 +5,8 @@
 #include "CADKernel/Math/Point.h"
 #include "CADKernel/Mesh/Structure/Grid.h"
 #include "CADKernel/Mesh/Structure/EdgeSegment.h"
+#include "CADKernel/Mesh/Structure/EdgeMesh.h"
+#include "CADKernel/Mesh/Structure/VertexMesh.h"
 #include "CADKernel/Topo/TopologicalEdge.h"
 #include "CADKernel/Topo/TopologicalFace.h"
 #include "CADKernel/Topo/TopologicalLoop.h"
@@ -12,7 +14,6 @@
 
 namespace UE::CADKernel
 {
-
 FIdent FEdgeSegment::LastId = 0;
 
 void FThinZone2DFinder::FindCloseSegments()
@@ -46,7 +47,7 @@ void FThinZone2DFinder::FindCloseSegments()
 
 #ifdef DEBUG_FIND_CLOSE_SEGMENTS
 	F3DDebugSession _(bDisplay, ("Close Segments"));
-	int32 SegmentID = 1291;
+	int32 SegmentID = 102;
 	int32 SegmentIndex = 0;
 #endif
 
@@ -67,7 +68,7 @@ void FThinZone2DFinder::FindCloseSegments()
 
 		const FPoint2D SegmentMiddle = Segment->GetCenter();
 
-		FPoint2D ClosePoint, ClosePoint2;
+		FPoint2D ClosePoint;
 		FEdgeSegment* CloseSegment = nullptr;
 
 		//     DistanceMin between Candidate and Segment with Candidate is before 
@@ -99,7 +100,7 @@ void FThinZone2DFinder::FindCloseSegments()
 		const double SegmentMiddleMax = SegmentMiddle.DiagonalAxisCoordinate() + MaxSpacePlusLength;
 		double MinSquareThickness = HUGE_VALUE;
 
-#ifdef DEBUG_FIND_CLOSE_SEGMENTS
+#ifdef DEBUG_FIND_CLOSE_SEGMENTS_
 		if (SegmentIndex == SegmentID)
 		{
 			F3DDebugSession _(bDisplay, FString::Printf(TEXT("Segment %d"), SegmentIndex));
@@ -149,27 +150,26 @@ void FThinZone2DFinder::FindCloseSegments()
 				}
 			}
 
-			// check the angle between segments. As they are opposite, the cosAngle as to be close to -1
+			// check the angle between segments. As they are opposite, the angle has to be close to PI i.e. slope has to be close to 4 => slope > Pi/2			
 			{
 				const double Slope = Segment->ComputeUnorientedSlopeOf(Candidate);
-				if (Slope < 3.) // Angle < 3Pi/4 (135 deg)
+				if (Slope < Slope::ThreeQuaterPiSlope)
 				{
 					continue;
 				}
 			}
 
-#ifdef DEBUG_FIND_CLOSE_SEGMENTS
+			const FPoint2D& SecondPointCandidate = Candidate->GetExtemity(ELimit::End);
+			double Coordinate;
+			const FPoint2D Projection = ProjectPointOnSegment(SegmentMiddle, FirstPointCandidate, SecondPointCandidate, Coordinate, true);
+
+#ifdef DEBUG_FIND_CLOSE_SEGMENTS_
 			if (SegmentIndex == SegmentID)
 			{
-				F3DDebugSession _(bDisplay, FString::Printf(TEXT("Candidate %d %d"), SegmentIndex, CandidateIndex));
+				F3DDebugSession _(bDisplay, FString::Printf(TEXT("Candidate %d %d Dist2"), SegmentIndex, CandidateIndex));
 				ThinZone::DisplayEdgeSegmentAndProjection(Segment, Candidate, EVisuProperty::BlueCurve, EVisuProperty::BlueCurve, EVisuProperty::RedCurve);
-				Wait();
-  			}
-			Wait(SegmentIndex == SegmentID + 1);
+			}
 #endif
-			const FPoint2D& SecondPointCandidate = Candidate->GetExtemity(ELimit::End);
-			double Coordinate;
-			const FPoint2D Projection = ProjectPointOnSegment(SegmentMiddle, FirstPointCandidate, SecondPointCandidate, Coordinate, true);
 
 			const double SquareDistance = SegmentMiddle.SquareDistance(Projection);
 			if (SquareDistance > SquareFinderTolerance)
@@ -179,10 +179,14 @@ void FThinZone2DFinder::FindCloseSegments()
 
 			if (MinSquareThickness > SquareDistance)
 			{
-				// check the angle between segment and Middle-Projection. As they are opposite, the cosAngle as to be close to 0
-				const double Slope = Segment->ComputeUnorientedSlopeOf(SegmentMiddle, Projection);
-				if (Slope < 1. || Slope > 3.)
+				// check the angle between segment and Middle-Projection. If Slope is positive, Middle-Projection is on the outer side 
+				const double Slope = Candidate->ComputeOrientedSlopeOf(SegmentMiddle, Projection);
+				if (Slope < 0)
 				{
+#ifdef DEBUG_FIND_CLOSE_SEGMENTS_
+					F3DDebugSession _(bDisplay, FString::Printf(TEXT("Candidate %f outer side criteria"), Slope));
+					ThinZone::DisplayEdgeSegmentAndProjection(Segment, Candidate, EVisuProperty::RedCurve, EVisuProperty::RedCurve, EVisuProperty::YellowCurve);
+#endif
 					continue;
 				}
 
@@ -206,16 +210,16 @@ void FThinZone2DFinder::FindCloseSegments()
 
 				F3DDebugSession _(bDisplay, FString::Printf(TEXT("Close Segment %d"), SegmentIndex));
 				ThinZone::DisplayEdgeSegmentAndProjection(Segment, CloseSegment, EVisuProperty::BlueCurve, EVisuProperty::BlueCurve, EVisuProperty::RedCurve);
-				SegmentIndex++;
-				//Wait(SegmentIndex > 22);
 			}
 #endif
 		}
+#ifdef DEBUG_FIND_CLOSE_SEGMENTS
+		SegmentIndex++;
+#endif
 	}
 }
 
-//#define DEBUG_FIND_CLOSE_SEGMENTS
-void FThinZone2DFinder::FindCloseSegments(const TArray<FEdgeSegment*> Segments, const TArray< const TArray<FEdgeSegment*>*> OppositeSides)
+void FThinZone2DFinder::FindCloseSegments(const TArray<FEdgeSegment*>& Segments, const TArray< const TArray<FEdgeSegment*>*>& OppositeSides)
 {
 #ifdef DEBUG_FIND_CLOSE_SEGMENTS
 	F3DDebugSession _(TEXT("FindCloseSegments"));
@@ -228,49 +232,28 @@ void FThinZone2DFinder::FindCloseSegments(const TArray<FEdgeSegment*> Segments,
 	}
 #endif
 
-	int32 SegmentToBeginIndex = 0;
-	for (FEdgeSegment* Segment : Segments)
+	double MinSquareThickness;
+	TFunction<FEdgeSegment* (const FEdgeSegment*, bool)> FindCloseSegment = [&OppositeSides, &MinSquareThickness](const FEdgeSegment* Segment, const bool bForce)
 	{
+		const FPoint2D SegmentMiddle = Segment->GetCenter();
 
-#ifdef DEBUG_FIND_CLOSE_SEGMENTS
-		{
-			F3DDebugSession _(TEXT("Segment"));
-			ThinZone::DisplayEdgeSegment(Segment, Segment->HasMarker1() ? EVisuProperty::RedCurve : EVisuProperty::BlueCurve);
-		}
-#endif
-
-		if (!Segment->HasMarker1())
-		{
-			continue;
-		}
-		Segment->ResetMarker1();
-
-		FPoint2D SegmentMiddle = Segment->GetCenter();
-
-		FPoint2D ClosePoint, ClosePoint2;
+		FPoint2D ClosePoint;
 		FEdgeSegment* CloseSegment = nullptr;
 
-		double MinSquareThickness = HUGE_VALUE;
-
-#ifdef DEBUG_FIND_CLOSE_SEGMENTS_2
-		//if (SegmentIndex == SegmentID)
-		F3DDebugSession _(bDisplay, TEXT("Segment"));
-		{
-			F3DDebugSession _(bDisplay, TEXT("Segment"));
-			DisplaySegmentWithScale(Segment->GetExtemity(ELimit::Start), Segment->GetExtemity(ELimit::End), 0, EVisuProperty::YellowCurve);
-		}
-#endif
-
+		MinSquareThickness = HUGE_VALUE;
 		for (const TArray<FEdgeSegment*>* OppositeSidePtr : OppositeSides)
 		{
 			const TArray<FEdgeSegment*>& OppositeSide = *OppositeSidePtr;
 			for (FEdgeSegment* Candidate : OppositeSide)
 			{
-				// check the angle between segments. As they are opposite, the Angle as to be close to Pi i.e. Slop ~= 4
-				double Slope = Segment->ComputeUnorientedSlopeOf(Candidate);
-				if (Slope < 3.) // Angle < 3Pi/4 (135 deg)
+				if (!bForce)
 				{
-					continue;
+					// check the angle between segments. As they are opposite, the Angle as to be close to Pi i.e. Slop ~= 4
+					double Slope = Segment->ComputeOrientedSlopeOf(Candidate);
+					if (Slope > -Slope::RightSlope && Slope < Slope::RightSlope)
+					{
+						continue;
+					}
 				}
 
 				const FPoint2D& FirstPointCandidate = Candidate->GetExtemity(ELimit::Start);
@@ -282,35 +265,36 @@ void FThinZone2DFinder::FindCloseSegments(const TArray<FEdgeSegment*> Segments,
 				const double SquareDistance = SegmentMiddle.SquareDistance(Projection);
 				if (MinSquareThickness > SquareDistance)
 				{
-					// check the angle between segment and Middle-Projection. As they are opposite, the cosAngle as to be close to 0
-					Slope = Segment->ComputeUnorientedSlopeOf(SegmentMiddle, Projection);
-					if (Slope < 1. || Slope > 3.)
-					{
-						continue;
-					}
-
 					MinSquareThickness = SquareDistance;
 					ClosePoint = Projection;
 					CloseSegment = Candidate;
 				}
 			}
 		}
+		return CloseSegment;
+	};
 
+	int32 SegmentToBeginIndex = 0;
+	for (FEdgeSegment* Segment : Segments)
+	{
+		if (Segment->GetCloseSegment())
+		{
+			continue;
+		}
+
+		constexpr bool bForceFind = true;
+		FEdgeSegment* CloseSegment = FindCloseSegment(Segment, !bForceFind);
+		if (!CloseSegment)
+		{
+			CloseSegment = FindCloseSegment(Segment, bForceFind);
+		}
 		if (CloseSegment)
 		{
 			Segment->SetCloseSegment(CloseSegment, MinSquareThickness);
-#ifdef DEBUG_FIND_CLOSE_SEGMENTS
-			if (bDisplay)
-			{
-				F3DDebugSession _(bDisplay, TEXT("Close Segment"));
-				ThinZone::DisplayEdgeSegmentAndProjection(Segment, CloseSegment, EVisuProperty::BlueCurve, EVisuProperty::BlueCurve, EVisuProperty::RedCurve);
-			}
-#endif
 		}
 	}
 }
 
-//#define DEBUG_LINK_CLOSE_SEGMENTS
 void FThinZone2DFinder::LinkCloseSegments()
 {
 	TArray<FEdgeSegment*> ThinZoneSegments;
@@ -326,7 +310,7 @@ void FThinZone2DFinder::LinkCloseSegments()
 		if (Loop != EdgeSegment->GetEdge()->GetLoop())
 		{
 			Loop = EdgeSegment->GetEdge()->GetLoop();
-			if(bThinZone)
+			if (bThinZone)
 			{
 				ChainIndex++;
 				bThinZone = false;
@@ -404,9 +388,206 @@ void FThinZone2DFinder::LinkCloseSegments()
 	}
 }
 
+void FThinZone2DFinder::CheckIfCloseSideOfThinSidesAreNotDegenerated()
+{
+#ifdef DEBUG_CHECK_THIN_SIDE
+	F3DDebugSession A(bDisplay, TEXT("Check Thin Sides"));
+#endif
+
+	for (TArray<FEdgeSegment*>& Side : ThinZoneSides)
+	{
+		if (Side.IsEmpty())
+		{
+			continue;
+		}
+
+		if (!CheckIfCloseSideOfThinSideIsNotDegenerated(Side))
+		{
+#ifdef DEBUG_CHECK_THIN_SIDE
+			if(bDisplay)
+			{
+				F3DDebugSession _(TEXT("Delete Thin Zone"));
+				ThinZone::DisplayThinZoneSide(Side, 0, EVisuProperty::RedCurve);
+			}
+#endif
+			for (FEdgeSegment* Segment : Side)
+			{
+				Segment->ResetCloseData();
+				Segment->SetChainIndex(Ident::Undefined);
+			}
+ 			Side.Empty();
+		}
+	}
+#ifdef DEBUG_CHECK_THIN_SIDE
+	Wait(bDisplay);
+#endif
+}
+
+/**
+ * The main idea of this algorithm is to check that the Side closed area is not limited to a point (Tol = 1/3 of the close segment)
+ * If the close side is defined by at least 3 segments, it's not degenerated.
+ * If the (2) segments of the close side are not in the same side -> it's not degenerated.
+ * Each extremities of the segments of the side is projected on the close segment(s)
+ * If the max distance between the projected points is smaller than the tolerance, the side is degenerated
+ */
+bool FThinZone2DFinder::CheckIfCloseSideOfThinSideIsNotDegenerated(TArray<FEdgeSegment*>& Side)
+{
+#ifdef DEBUG_CHECK_THIN_SIDE_
+	F3DDebugSession _(bDisplay, ("Check Thin Zone"));
+#endif
+
+	const FEdgeSegment* CloseSegments[3] = { nullptr, nullptr,nullptr };
+	TFunction<FPoint2D(const FPoint2D&)> GetOppositeSegment = [&CloseSegments](const FPoint2D& Candidate) -> FPoint2D
+	{
+		FPoint2D OutOppositePoint;
+		double MinSquareThickness = HUGE_VALUE;
+		for (int32 Index = 0; CloseSegments[Index] != nullptr; ++Index)
+		{
+			const FEdgeSegment* CloseSegment = CloseSegments[Index];
+			const FPoint2D& FirstPointCandidate = CloseSegment->GetExtemity(ELimit::Start);
+			const FPoint2D& SecondPointCandidate = CloseSegment->GetExtemity(ELimit::End);
+
+			double Coordinate;
+			const FPoint2D Projection = ProjectPointOnSegment(Candidate, FirstPointCandidate, SecondPointCandidate, Coordinate, true);
+
+			const double SquareDistance = Candidate.SquareDistance(Projection);
+			if (MinSquareThickness > SquareDistance)
+			{
+				MinSquareThickness = SquareDistance;
+				OutOppositePoint = Projection;
+			}
+		}
+		return OutOppositePoint;
+	};
+
+	if (Side.IsEmpty())
+	{
+		// side is degenerated
+		return false;
+	}
+
+	// find the close segments
+	for (const FEdgeSegment* Segment : Side)
+	{
+		const FEdgeSegment* CloseSegment = Segment->GetCloseSegment();
+		if (!CloseSegments[0])
+		{
+			CloseSegments[0] = CloseSegment;
+			continue;
+		}
+		if (CloseSegments[0] != CloseSegment)
+		{
+			if (!CloseSegments[1])
+			{
+				CloseSegments[1] = CloseSegment;
+			}
+			else if (CloseSegments[1] != CloseSegment)
+			{
+				CloseSegments[2] = CloseSegment;
+				break;
+			}
+		}
+	}
+
+	// More than two segments -> not degenerated
+	if (CloseSegments[2] != nullptr)
+	{
+		return true;
+	}
+
+	if (CloseSegments[1] != nullptr)
+	{
+		// Case two chain -> not degenerated 
+		// 
+		//          Chain 0     Side     Chain 1
+		//                \      \ /     /
+		//     Close[0] -> \      #     / <- Close[1]
+		//                  \          /
+		//                   
+		if (CloseSegments[0]->GetChainIndex() != CloseSegments[1]->GetChainIndex())
+		{
+			return true;
+		}
+
+		// Case the close segments are not connected -> not degenerated 
+		//                       Side     Chain
+		//                \      \ /     /
+		//     Close[0] -> \      #     / <- Close[1]
+		//                  \          /
+		//                   #--------#
+		if (CloseSegments[0]->GetPrevious() != CloseSegments[1] && CloseSegments[0]->GetNext() != CloseSegments[1])
+		{
+			return true;
+		}
+	}
+
+	TFunction<double()> ComputeTolerance = [&CloseSegments]() -> double
+	{
+		double CloseLength = CloseSegments[0]->GetLength();
+		const double Ratio = CloseSegments[1] ? 12. : 6.;
+		if (CloseSegments[1])
+		{
+			CloseLength += CloseSegments[1]->GetLength();
+		}
+		return FMath::Square(CloseLength / Ratio);
+	};
+	const double ToleranceCloseSquare = ComputeTolerance();
+
+	const FPoint2D& StartPoint = Side[0]->GetExtemity(Start);
+	FPoint2D FirstOppositePoint = GetOppositeSegment(StartPoint);
+
+#ifdef DEBUG_CHECK_THIN_SIDE_
+	{
+		F3DDebugSession _(bDisplay, ("Seg"));
+		DisplaySegmentWithScale(StartPoint, FirstOppositePoint, 0, EVisuProperty::BlueCurve);
+		DisplayPoint2DWithScale(StartPoint, EVisuProperty::BluePoint);
+		DisplayPoint2DWithScale(StartPoint, EVisuProperty::RedPoint);
+	}
+#endif
+
+	for (const FEdgeSegment* Segment : Side)
+	{
+		const FPoint2D& NextPoint = Segment->GetExtemity(End);
+		FPoint2D OppositePoint = GetOppositeSegment(NextPoint);
+		const double SquarewDistance = OppositePoint.SquareDistance(FirstOppositePoint);
+
+#ifdef DEBUG_CHECK_THIN_SIDE_
+		{
+			F3DDebugSession _(bDisplay, TEXT("Seg"));
+			DisplaySegmentWithScale(NextPoint, OppositePoint, 0, EVisuProperty::BlueCurve);
+			DisplayPoint2DWithScale(NextPoint, EVisuProperty::BluePoint);
+			DisplayPoint2DWithScale(OppositePoint, EVisuProperty::RedPoint);
+			Wait();
+		}
+#endif
+
+		if (SquarewDistance > ToleranceCloseSquare)
+		{
+			return true;
+		}
+	}
+
+	// Degenerated... 
+	// Check if at least one segment of the side is the close segment of the other side
+	const FEdgeSegment* CloseOfCloseSegments[2] = { CloseSegments[0]->GetCloseSegment(), CloseSegments[1] ? CloseSegments[1]->GetCloseSegment() : nullptr };
+	for (const FEdgeSegment* Segment : Side)
+	{
+		if (Segment == CloseOfCloseSegments[0] || Segment == CloseOfCloseSegments[1])
+		{
+			return true;
+		}
+	}
+
+	return false;
+}
+
 void FThinZone2DFinder::ImproveThinSide()
 {
-	TFunction<TArray<FEdgeSegment*>(const FEdgeSegment*)> GetComplementary = [](const FEdgeSegment* EdgeSegment) ->TArray<FEdgeSegment*>
+#ifdef DEBUG_THIN_ZONES_IMPROVE
+	F3DDebugSession _(bDisplay, ("Improve Thin Zone"));
+#endif
+
+	TFunction<TArray<FEdgeSegment*>(const FEdgeSegment*)> GetComplementary = [this](const FEdgeSegment* EdgeSegment) ->TArray<FEdgeSegment*>
 	{
 		// We allow to extend a thin zone up to 4 times the local thickness if it lets to link another ThinZone
 		constexpr double ComplementaryFactor = 4.;
@@ -415,7 +596,7 @@ void FThinZone2DFinder::ImproveThinSide()
 		ComplementaryEdges.Reserve(100);
 
 		double Length = 0;
-		double MaxLength = FMath::Sqrt(EdgeSegment->GetCloseSquareDistance()) * ComplementaryFactor;
+		double MaxLength = this->FinderTolerance;
 
 		for (FEdgeSegment* PreviousSegment = EdgeSegment->GetPrevious();; PreviousSegment = PreviousSegment->GetPrevious())
 		{
@@ -437,7 +618,7 @@ void FThinZone2DFinder::ImproveThinSide()
 		return MoveTemp(ComplementaryEdges);
 	};
 
-	TFunction<void(TArray<FEdgeSegment*>&, TArray<FEdgeSegment*>&)> MergeChains = [&](TArray<FEdgeSegment*>& ThinZoneSide, TArray<FEdgeSegment*>& ComplementaryEdges)
+	TFunction<void(TArray<FEdgeSegment*>&, TArray<FEdgeSegment*>&)> MergeChains = [&ThinZoneSides = ThinZoneSides](TArray<FEdgeSegment*>& ThinZoneSide, TArray<FEdgeSegment*>& ComplementaryEdges)
 	{
 		FEdgeSegment* PreviousSegment = ComplementaryEdges[0]->GetPrevious();
 		const FIdent ChainId = ThinZoneSide[0]->GetChainIndex();
@@ -454,7 +635,7 @@ void FThinZone2DFinder::ImproveThinSide()
 		}
 	};
 
-	TFunction<void(TArray<FEdgeSegment*>&)> MergeWithPreviousChain = [&](TArray<FEdgeSegment*>& ThinZoneSide)
+	TFunction<void(TArray<FEdgeSegment*>&)> MergeWithPreviousChain = [&ThinZoneSides = ThinZoneSides](TArray<FEdgeSegment*>& ThinZoneSide)
 	{
 		const FEdgeSegment* PreviousSegment = ThinZoneSide[0]->GetPrevious();
 		const FIdent ChainId = ThinZoneSide[0]->GetChainIndex();
@@ -470,6 +651,30 @@ void FThinZone2DFinder::ImproveThinSide()
 		}
 	};
 
+	TFunction<void(const FEdgeSegment*, TArray<FEdgeSegment*>&)> SetChainIndex = [](const FEdgeSegment* ThinZoneSegment, TArray<FEdgeSegment*>& Complementary)
+	{
+		const FIdent ChainId = ThinZoneSegment->GetChainIndex();
+		for (FEdgeSegment* Segment : Complementary)
+		{
+			Segment->SetChainIndex(ChainId);
+		}
+	};
+
+	TFunction<void(TArray<FEdgeSegment*>&, TArray<FEdgeSegment*>&)> MergeSideWithComplementaryAtStart = [SetChainIndex](TArray<FEdgeSegment*>& ThinZoneSide, TArray<FEdgeSegment*>& Complementary)
+	{
+		SetChainIndex(ThinZoneSide[0], Complementary);
+		Complementary.Reserve(Complementary.Num() + ThinZoneSide.Num());
+		Complementary.Append(MoveTemp(ThinZoneSide));
+		ThinZoneSide = MoveTemp(Complementary);
+	};
+
+	TFunction<void(TArray<FEdgeSegment*>&, TArray<FEdgeSegment*>&)> MergeSideWithComplementaryAtEnd = [SetChainIndex](TArray<FEdgeSegment*>& ThinZoneSide, TArray<FEdgeSegment*>& Complementary)
+	{
+		SetChainIndex(ThinZoneSide[0], Complementary);
+		ThinZoneSide.Reserve(Complementary.Num() + ThinZoneSide.Num());
+		ThinZoneSide.Append(MoveTemp(Complementary));
+	};
+
 	TFunction<void(TArray<FEdgeSegment*>&)> SetMarker1ToComplementarySegments = [&](TArray<FEdgeSegment*>& Complementary)
 	{
 		for (FEdgeSegment* Segment : Complementary)
@@ -488,99 +693,181 @@ void FThinZone2DFinder::ImproveThinSide()
 			continue;
 		}
 
+#ifdef DEBUG_THIN_ZONES_IMPROVE
+		static int32 IndexSide = 0;
+		if (bDisplay)
+		{
+			ThinZone::DisplayThinZoneSide(Side, IndexSide++, EVisuProperty::BlueCurve, false);
+			Wait();
+		}
+#endif
+
 		const FEdgeSegment* Segment0Side = Side[0];
 		if (Segment0Side->GetPrevious()->GetCloseSegment() != nullptr)
 		{
 			// Connected case
-			// 
-			//                  Segment0Side->GetPrevious()            
-			//                  | Segment0Side            
-			//       Side(n-1)  | |  Side(n)       Side(n-1) = ThinZoneSides[n-1]
-			//    #--------------#-----------#
-			//
-			//    #--------------------------#
-			//                  | |
-			//                  | Segment0Side->GetCloseSegment()
-			//                  Segment0Side->GetPrevious()->GetCloseSegment()
+			// 	             
+			//                             Segment0Side->GetPrevious()            
+			//                Segment0Side |           
+			//                  Side(n)  | | Side(n-1)       Side(n-1) = ThinZoneSides[n-1]
+			//             #--------------#-----------#
+			//	          
+			//             #--------------------------#
+			//                           | |
+			// Seg0Side->GetCloseSegment() | 
+			//                             Segment0Side->GetPrevious()->GetCloseSegment()
 
 			FIdent CloseOfPreviousSideIndex = Segment0Side->GetPrevious()->GetCloseSegment()->GetChainIndex();
 			if (CloseOfPreviousSideIndex != Side[0]->GetChainIndex() && CloseOfPreviousSideIndex == Segment0Side->GetCloseSegment()->GetChainIndex())
 			{
 				MergeWithPreviousChain(Side);
+#ifdef DEBUG_THIN_ZONES_IMPROVE
+				if (bDisplay)
+				{
+					ThinZone::DisplayThinZoneSide(Side, 10000, EVisuProperty::BlueCurve, false);
+					Wait();
+				}
+#endif
+				continue;
 			}
+			continue;
 		}
-		else
+
+		TArray<FEdgeSegment*> Complementary = GetComplementary(Segment0Side);
+		if (Complementary.Num())
 		{
-			TArray<FEdgeSegment*> Complementary = GetComplementary(Segment0Side);
-			if (Complementary.Num())
+#ifdef DEBUG_THIN_ZONES_IMPROVE
+			if (bDisplay)
 			{
-				// Case2 : separated by a small chain slightly too far to be considered as close
-
-				//                  PreviousComplementary[0]
-				//                  | Complementary[0]     Seg0Side            Side(n)   = ThinZoneSides[n]
-				//       Side(n-1)  | |   Complementary    |  Side(n)          Side(n-1) = ThinZoneSides[n-1]
-				//    #--------------#--------------------#-----------#
-				//
-				//    #-----------------------------------------------#
-				//                  |                      |
-				//                  |                      Seg0Side->GetCloseSegment()
-				//                  ClosePreviousComplementary[0]
-				const FEdgeSegment* PreviousComplementary = Complementary[0]->GetPrevious();
-				const FEdgeSegment* ClosePreviousComplementary = PreviousComplementary->GetCloseSegment();
-
-
-				// Case Side(n-1) and Side(n) are closed together
-				// 
-				//                               PreviousComplementary[0]
-				//       Side(n-1)               | 
-				//    #---------------------------#
-				//                                |
-				//                                |  Complementary
-				//    #---------------------------#
-				//                Side(n)        |       
-				//                               ClosePreviousComplementary 
-				if (ClosePreviousComplementary->GetChainIndex() == Segment0Side->GetChainIndex())
+				ThinZone::DisplayThinZoneSide(Complementary, -1, EVisuProperty::RedCurve, false);
+				Wait();
+			}
+#endif
+
+			// Case2 : separated by a small chain slightly too far to be considered as close
+			//
+			//                  Seg0Side             Complementary[0]
+			//                  |                    | PreviousComplementary[0]            Side(n)   = ThinZoneSides[n]
+			//       Side(n  )  |     Complementary  | |  Side(n-1)                        Side(n-1) = ThinZoneSides[n-1]
+			//    #--------------#--------------------#-----------#
+			//
+			//    #-----------------------------------------------#
+			//                  |                      |
+			//                  |                      ClosePreviousComplementary[0]
+			//                  Seg0Side->GetCloseSegment()
+			const FEdgeSegment* PreviousComplementary = Complementary[0]->GetPrevious();
+			const FEdgeSegment* ClosePreviousComplementary = PreviousComplementary->GetCloseSegment();
+
+			if (ClosePreviousComplementary->GetChainIndex() == Ident::Undefined)
+			{
+				continue;
+			}
+
+			// Case Side(n-1) and Side(n) are closed together
+			// 
+			//                               PreviousComplementary = Complementary[0]->Previous = PreviousSide.Last 
+			//          Side(n)              | 
+			//    #---------------------------#
+			//                                |
+			//                                |  Complementary
+			//    #---------------------------#
+			//              Side(n-1)        |       
+			//                               ClosePreviousComplementary
+			if (ClosePreviousComplementary->GetChainIndex() == Segment0Side->GetChainIndex())
+			{
+				// Check the slope
+				double SlopeAtPreviousComplementary = Complementary[0]->ComputeUnorientedSlopeOf(PreviousComplementary);
+				if (SlopeAtPreviousComplementary < 1)
 				{
-					// Do nothing
+					TArray<FEdgeSegment*>& PreviousSide = ThinZoneSides[PreviousComplementary->GetChainIndex()];
+					MergeSideWithComplementaryAtEnd(PreviousSide, Complementary);
+					SetMarker1ToComplementarySegments(Complementary);
+#ifdef DEBUG_THIN_ZONES_IMPROVE
+					if (bDisplay)
+					{
+						ThinZone::DisplayThinZoneSide(PreviousSide, 80000, EVisuProperty::BlueCurve, false);
+						Wait();
+					}
+#endif
 					continue;
 				}
 
-				FIdent CloseSideIndex = Segment0Side->GetCloseSegment()->GetChainIndex();
-				if (CloseSideIndex == ClosePreviousComplementary->GetChainIndex())
+				const FEdgeSegment* LastComplementary = Complementary.Last();
+				const FEdgeSegment* NextComplementary = LastComplementary->GetNext();
+				double SlopeAtNextComplementary = LastComplementary->ComputeUnorientedSlopeOf(NextComplementary);
+				if (SlopeAtNextComplementary < 1)
 				{
-					MergeChains(Side, Complementary);
+					MergeSideWithComplementaryAtStart(Side, Complementary);
 					SetMarker1ToComplementarySegments(Complementary);
+
+#ifdef DEBUG_THIN_ZONES_IMPROVE
+					if (bDisplay)
+					{
+						ThinZone::DisplayThinZoneSide(Side, 70000, EVisuProperty::BlueCurve, false);
+						Wait();
+					}
+#endif
+					continue;
 				}
-				else
+
+				continue;
+			}
+
+			FIdent CloseSideIndex = Segment0Side->GetCloseSegment()->GetChainIndex();
+			if (CloseSideIndex == ClosePreviousComplementary->GetChainIndex())
+			{
+				MergeChains(Side, Complementary);
+				SetMarker1ToComplementarySegments(Complementary);
+
+#ifdef DEBUG_THIN_ZONES_IMPROVE
+				if (bDisplay)
 				{
-					//                  PreviousComplementary[0]
-					//                  | Complementary[0]     Seg0Side            Side(n)   = ThinZoneSides[n]
-					//       Side(n-1)  | |   Complementary    |  Side(n)          Side(n-1) = ThinZoneSides[n-1]
-					//    #--------------#--------------------#-----------#  ->
-					//
-					//    #--------------#--------------------#-----------#  <-
-					//                  |  ComplementaryClose  |
-					//                  |                      Seg0Side->GetCloseSegment() == PreviousComplementaryClose[0]
-					//                  ClosePreviousComplementary[0]
-					TArray<FEdgeSegment*> ComplementaryClose = GetComplementary(ClosePreviousComplementary);
-					if (ComplementaryClose.Num())
-					{
-						const FEdgeSegment* PreviousComplementaryClose = ComplementaryClose[0]->GetPrevious(); // == Seg0Side->GetCloseSegment()
-						const FEdgeSegment* CloseOfPreviousComplementaryClose = PreviousComplementaryClose->GetCloseSegment(); // == Seg0Side
+					ThinZone::DisplayThinZoneSide(Side, 20000, EVisuProperty::BlueCurve, false);
+					Wait();
+				}
+#endif
 
-						if (CloseOfPreviousComplementaryClose->GetChainIndex() == Segment0Side->GetChainIndex())
-						{
-							MergeChains(Side, Complementary);
-							SetMarker1ToComplementarySegments(Complementary);
+				continue;
+			}
 
-							TArray<FEdgeSegment*>& SideOfClosePreviousComplementary = ThinZoneSides[ClosePreviousComplementary->GetChainIndex()];
-							MergeChains(SideOfClosePreviousComplementary, ComplementaryClose);
-							SetMarker1ToComplementarySegments(ComplementaryClose);
-						}
+			//                  PreviousComplementary[0]
+			//                  | Complementary[0]     Seg0Side            Side(n)   = ThinZoneSides[n]
+			//       Side(n-1)  | |   Complementary    |  Side(n)          Side(n-1) = ThinZoneSides[n-1]
+			//    #--------------#--------------------#-----------#  ->
+			//
+			//    #--------------#--------------------#-----------#  <-
+			//                  |  ComplementaryClose  |
+			//                  |                      Seg0Side->GetCloseSegment() == PreviousComplementaryClose[0]
+			//                  ClosePreviousComplementary[0]
+			TArray<FEdgeSegment*> ComplementaryClose = GetComplementary(ClosePreviousComplementary);
+			if (ComplementaryClose.Num())
+			{
+				const FEdgeSegment* PreviousComplementaryClose = ComplementaryClose[0]->GetPrevious(); // == Seg0Side->GetCloseSegment()
+				const FEdgeSegment* CloseOfPreviousComplementaryClose = PreviousComplementaryClose->GetCloseSegment(); // == Seg0Side
+
+				if (CloseOfPreviousComplementaryClose->GetChainIndex() == Segment0Side->GetChainIndex())
+				{
+					MergeChains(Side, Complementary);
+					SetMarker1ToComplementarySegments(Complementary);
+
+					TArray<FEdgeSegment*>& SideOfClosePreviousComplementary = ThinZoneSides[ClosePreviousComplementary->GetChainIndex()];
+					MergeChains(SideOfClosePreviousComplementary, ComplementaryClose);
+					SetMarker1ToComplementarySegments(ComplementaryClose);
+
+#ifdef DEBUG_THIN_ZONES_IMPROVE			
+					if (bDisplay)
+					{
+
+						ThinZone::DisplayThinZoneSide(Side, 30000, EVisuProperty::BlueCurve, false);
+						Wait();
 					}
+#endif
+
+					continue;
 				}
 			}
 		}
+		continue;
 	}
 
 	TFunction<const TArray<const TArray<FEdgeSegment*>*>(const TArray<FEdgeSegment*>&)> FindOppositeSides = [&](const TArray<FEdgeSegment*>& Side) -> const TArray<const TArray<FEdgeSegment*>*>
@@ -606,25 +893,46 @@ void FThinZone2DFinder::ImproveThinSide()
 		OppositeSides.Reserve(OppositeSideIndexes.Num());
 		for (int32 Index : OppositeSideIndexes)
 		{
+			if (Index < 0)
+			{
+				continue;
+			}
 			OppositeSides.Add(&ThinZoneSides[Index]);
 		}
 
 		return MoveTemp(OppositeSides);
 	};
 
+#ifdef DEBUG_SEARCH_THIN_ZONES_
+	if (bDisplay)
+	{
+		ThinZone::DisplayThinZoneSides2(ThinZoneSides);
+		Wait();
+	}
+#endif
+
 	// Finalize:
 	// Find close segments to the added segments
 	for (TArray<FEdgeSegment*>& Side : ThinZoneSides)
 	{
-		if (!Algo::AllOf(Side, [](const FEdgeSegment* Seg) { return !Seg->HasMarker1(); }))
+		if (!Algo::AllOf(Side, [](const FEdgeSegment* Seg) { return Seg->GetCloseSegment(); }))
 		{
 			const TArray<const TArray<FEdgeSegment*>*> OppositeSides = FindOppositeSides(Side);
 			FindCloseSegments(Side, OppositeSides);
 		}
 	}
 
-#ifdef DEBUG_THIN_ZONES_IMPROVE
-	DisplayThinZoneSides(ThinZoneSides);
+#ifdef DEBUG_SEARCH_THIN_ZONES_
+	if (bDisplay)
+	{
+		ThinZone::DisplayThinZoneSides2(ThinZoneSides);
+		Wait();
+	}
+#endif
+
+
+#ifdef DEBUG_THIN_ZONES_IMPROVE_
+	ThinZone::DisplayThinZoneSides(ThinZoneSides);
 	Wait();
 #endif
 }
@@ -642,12 +950,19 @@ void FThinZone2DFinder::SplitThinSide()
 		}
 
 		FIdent CloseSideIndex = ThinZoneSide[0]->GetCloseSegment()->GetChainIndex();
-		if (!Algo::AllOf(ThinZoneSide, [&](const FEdgeSegment* EdgeSegment) { return EdgeSegment->GetCloseSegment()->GetChainIndex() == CloseSideIndex; }))
+		if (!Algo::AllOf(ThinZoneSide, [&](const FEdgeSegment* EdgeSegment) 
+			{ 
+				if (!EdgeSegment->GetCloseSegment())
+				{
+					return false;
+				}
+				return EdgeSegment->GetCloseSegment()->GetChainIndex() == CloseSideIndex; 
+			}))
 		{
 			int32 Index = 0;
 			for (; Index < ThinZoneSide.Num(); ++Index)
 			{
-				if (ThinZoneSide[Index]->GetCloseSegment()->GetChainIndex() != CloseSideIndex)
+				if (!ThinZoneSide[Index]->GetCloseSegment() || ThinZoneSide[Index]->GetCloseSegment()->GetChainIndex() != CloseSideIndex)
 				{
 					break;
 				}
@@ -656,12 +971,17 @@ void FThinZone2DFinder::SplitThinSide()
 			const int32 LastIndexFirstSide = Index;
 			while (Index < ThinZoneSide.Num())
 			{
+				if (!ThinZoneSide[Index]->GetCloseSegment())
+				{
+					continue;
+				}
+
 				CloseSideIndex = ThinZoneSide[Index]->GetCloseSegment()->GetChainIndex();
 				TArray<FEdgeSegment*>& NewThinZoneSide = NewThinZoneSides.Emplace_GetRef();
 				for (; Index < ThinZoneSide.Num(); ++Index)
 				{
 					FEdgeSegment* Segment = ThinZoneSide[Index];
-					if (Segment->GetCloseSegment()->GetChainIndex() != CloseSideIndex)
+					if (!Segment->GetCloseSegment() || Segment->GetCloseSegment()->GetChainIndex() != CloseSideIndex)
 					{
 						break;
 					}
@@ -684,11 +1004,6 @@ void FThinZone2DFinder::SplitThinSide()
 
 void FThinZone2DFinder::BuildThinZone()
 {
-#ifdef DEBUG_BUILD_THIN_ZONES
-	ThinZone::DisplayThinZoneSides(ThinZoneSides);
-	Wait();
-#endif
-
 	// The number of ThinZone should be less than the number of ThinZoneSides, 
 	ThinZones.Reserve(ThinZoneSides.Num());
 
@@ -700,6 +1015,12 @@ void FThinZone2DFinder::BuildThinZone()
 		}
 
 		const FEdgeSegment* FirstSideSegment = FirstSide[0];
+		const FIdent OppositeChainIndex = FirstSideSegment->GetCloseSegment()->GetChainIndex();
+		if (OppositeChainIndex == Ident::Undefined)
+		{
+			continue;
+		}
+
 		TArray<FEdgeSegment*>& SecondSide = ThinZoneSides[FirstSideSegment->GetCloseSegment()->GetChainIndex()];
 		if (SecondSide.IsEmpty())
 		{
@@ -715,8 +1036,10 @@ void FThinZone2DFinder::BuildThinZone()
 
 void FThinZone2DFinder::GetThinZoneSideConnectionsLength(const TArray<FEdgeSegment*>& FirstSide, const TArray<FEdgeSegment*>& SecondSide, double InMaxLength, double* OutLengthBetweenExtremities, TArray<const FTopologicalEdge*>* OutPeakEdges)
 {
-	TFunction<double(const FEdgeSegment*, const FEdgeSegment*, TFunction<const FEdgeSegment* (const FEdgeSegment*)>, TArray<const FTopologicalEdge*>&)> ComputeLengthBetweenExtremities = 
-		[&](const FEdgeSegment* Start, const FEdgeSegment* End, TFunction<const FEdgeSegment* (const FEdgeSegment*)> GetNext, TArray<const FTopologicalEdge*>& PeakEdges) -> double
+	using FGetNeighborFunc = TFunction<const FEdgeSegment* (const FEdgeSegment*)>;
+	using FComputeLengthBetweenExtremitiesFunc = TFunction<double(const FEdgeSegment*, const FEdgeSegment*, FGetNeighborFunc, TArray<const FTopologicalEdge*>&)>;
+
+	FComputeLengthBetweenExtremitiesFunc ComputeLengthBetweenExtremities = [&InMaxLength](const FEdgeSegment* Start, const FEdgeSegment* End, FGetNeighborFunc GetNext, TArray<const FTopologicalEdge*>& PeakEdges) -> double
 	{
 		PeakEdges.Reserve(10);
 
@@ -749,13 +1072,41 @@ void FThinZone2DFinder::GetThinZoneSideConnectionsLength(const TArray<FEdgeSegme
 
 void FThinZone2DFinder::BuildThinZone(const TArray<FEdgeSegment*>& FirstSide, const TArray<FEdgeSegment*>& SecondSide)
 {
-	FThinZone2D& Zone = ThinZones.Emplace_GetRef(FirstSide, SecondSide);
+	TFunction<void(const TArray<FEdgeSegment*>&, double&, double&)> ComputeThicknessAndLength = [](const TArray<FEdgeSegment*>& Side, double& MaxThickness, double& SideLength)
+	{
+		SideLength = 0;
+		double SquareMaxThickness = 0;
+
+		for (const FEdgeSegment* Segment : Side)
+		{
+			const double SquareThickness = Segment->GetCloseSquareDistance();
+			const double SegmentLength = Segment->GetLength();
+			SideLength += SegmentLength;
+			if (SquareMaxThickness < SquareThickness)
+			{
+				SquareMaxThickness = SquareThickness;
+			}
+		}
+		MaxThickness = FMath::Sqrt(SquareMaxThickness);
+	};
+
+	EThinZone2DType Zone2DType = EThinZone2DType::BetweenLoops;
+	TArray<const FTopologicalEdge*> PeakEdges[2];
 
 	if (!FirstSide[0]->IsInner() && !SecondSide[0]->IsInner())
 	{
-		const double MaxThickness = Zone.GetMaxThickness();
+		double MaxThicknessSide1 = 0;
+		double MaxThicknessSide2 = 0;
+		double Side1Length;
+		double Side2Length;
+
+		ComputeThicknessAndLength(FirstSide, MaxThicknessSide1, Side1Length);
+		ComputeThicknessAndLength(SecondSide, MaxThicknessSide2, Side2Length);
+
+		const double MaxThickness = FMath::Max(MaxThicknessSide1, MaxThicknessSide2);
+		const double ThinZoneLength = Side1Length + Side2Length;
+
 		double LengthBetweenExtremity[2] = { HUGE_VALUE, HUGE_VALUE };
-		TArray<const FTopologicalEdge*> PeakEdges[2];
 		GetThinZoneSideConnectionsLength(FirstSide, SecondSide, 3. * MaxThickness, LengthBetweenExtremity, PeakEdges);
 
 		//                     Side 0 
@@ -770,59 +1121,63 @@ void FThinZone2DFinder::BuildThinZone(const TArray<FEdgeSegment*>& FirstSide, co
 		// So the MaxLengthBetweendSidesToBePeak is theoretically MaxThickness x Sqrt(2)...
 		// It's simplified to MaxThickness x 2.
 		const double MaxLengthBetweendSidesToBeAPeak = MaxThickness * 2.;
-		const double ThinZoneLength = Zone.Length();
 		const double MinThinZoneLengthToBeGlobal = ExternalLoopLength - 2. * MaxLengthBetweendSidesToBeAPeak;
 
 		if (ThinZoneLength > MinThinZoneLengthToBeGlobal || (LengthBetweenExtremity[0] < MaxLengthBetweendSidesToBeAPeak && LengthBetweenExtremity[1] < MaxLengthBetweendSidesToBeAPeak))
 		{
-			Zone.SetCategory(EThinZone2DType::Global);
+			Zone2DType = EThinZone2DType::Global;
 		}
 		else if (LengthBetweenExtremity[0] < MaxLengthBetweendSidesToBeAPeak)
 		{
-			if (ThinZoneLength < MaxThickness * 5.)
+			if (ThinZoneLength < MaxThickness * 3.)
 			{
-				// The thin zone is too small, it's deleted
-				ThinZones.Pop();
+				// The thin zone is too small => void
 				return;
-
 			}
 			else
 			{
-				Zone.SetCategory(EThinZone2DType::PeakStart);
-				Zone.SetPeakEdgesMarker(PeakEdges[0]);
+				Zone2DType = EThinZone2DType::PeakStart;
 			}
 		}
 		else if (LengthBetweenExtremity[1] < MaxLengthBetweendSidesToBeAPeak)
 		{
-			if (ThinZoneLength < MaxThickness * 5.)
+			if (ThinZoneLength < MaxThickness * 3.)
 			{
-				// The thin zone is too small, it's deleted
-				ThinZones.Pop();
+				// The thin zone is too small => void
 				return;
-
 			}
 			else
 			{
-				Zone.SetCategory(EThinZone2DType::PeakEnd);
-				Zone.SetPeakEdgesMarker(PeakEdges[1]);
+				Zone2DType = EThinZone2DType::PeakEnd;
 			}
 		}
 		else
 		{
-			Zone.SetCategory(EThinZone2DType::Butterfly);
+			Zone2DType = EThinZone2DType::Butterfly;
 		}
+	}
 
+	FThinZone2D& Zone = ThinZones.Emplace_GetRef(FirstSide, SecondSide);
+
+	Zone.SetCategory(Zone2DType);
+	if (Zone2DType == EThinZone2DType::PeakStart)
+	{
+		Zone.SetPeakEdgesMarker(PeakEdges[0]);
 	}
-	else
+	else if (Zone2DType == EThinZone2DType::PeakEnd)
 	{
-		Zone.SetCategory(EThinZone2DType::BetweenLoops);
+		Zone.SetPeakEdgesMarker(PeakEdges[1]);
 	}
+	Zone.AddToEdge();
 }
 
 bool FThinZone2DFinder::SearchThinZones(double InTolerance)
 {
 #ifdef DEBUG_THIN_ZONES
 	F3DDebugSession A(bDisplay, FString::Printf(TEXT("ThinZone Finder %d"), Grid.GetFace().GetId()));
+#endif
+
+#ifdef DEBUG_SEARCH_THIN_ZONES
 	if (bDisplay)
 	{
 		F3DDebugSession _(TEXT("Thin surface grid"));
@@ -841,7 +1196,7 @@ bool FThinZone2DFinder::SearchThinZones(double InTolerance)
 	}
 	Chronos.BuildLoopSegmentsTime = FChrono::Elapse(StartTime);
 
-#ifdef DEBUG_THIN_ZONES
+#ifdef DEBUG_SEARCH_THIN_ZONES
 	DisplayLoopSegments();
 #endif
 
@@ -853,7 +1208,7 @@ bool FThinZone2DFinder::SearchThinZones(double InTolerance)
 	if (bDisplay)
 	{
 		DisplayCloseSegments();
-		Wait();
+		//Wait();
 	}
 #endif
 
@@ -861,19 +1216,29 @@ bool FThinZone2DFinder::SearchThinZones(double InTolerance)
 	LinkCloseSegments();
 
 #ifdef DEBUG_SEARCH_THIN_ZONES
-	if(bDisplay)
+	if (bDisplay)
 	{
 		ThinZone::DisplayThinZoneSides(ThinZoneSides);
+		//Wait();
 	}
 #endif
 
-	ImproveThinSide();
+	CheckIfCloseSideOfThinSidesAreNotDegenerated();
 
 #ifdef DEBUG_SEARCH_THIN_ZONES
 	if (bDisplay)
 	{
 		ThinZone::DisplayThinZoneSides(ThinZoneSides);
-		Wait();
+		//Wait();
+	}
+#endif
+
+	ImproveThinSide();
+
+#ifdef DEBUG_SEARCH_THIN_ZONES
+	if (bDisplay)
+	{
+		ThinZone::DisplayThinZoneSides2(ThinZoneSides);
 	}
 #endif
 
@@ -883,9 +1248,10 @@ bool FThinZone2DFinder::SearchThinZones(double InTolerance)
 	if (bDisplay)
 	{
 		ThinZone::DisplayThinZoneSides(ThinZoneSides);
-		Wait();
+		//Wait();
 	}
 #endif
+
 	Chronos.LinkCloseSegmentTime = FChrono::Elapse(StartTime);
 
 	StartTime = FChrono::Now();
@@ -893,10 +1259,10 @@ bool FThinZone2DFinder::SearchThinZones(double InTolerance)
 	Chronos.BuildThinZoneTime = FChrono::Elapse(StartTime);
 
 #ifdef DEBUG_THIN_ZONES
-	if(bDisplay)
+	if (bDisplay)
 	{
 		ThinZone::DisplayThinZones(ThinZones);
-		Wait(false);
+		//Wait();
 	}
 #endif
 
@@ -939,7 +1305,10 @@ bool FThinZone2DFinder::BuildLoopSegments()
 		for (const FOrientedEdge& Edge : Edges)
 		{
 			TArray<double> Coordinates;
-			Edge.Entity->Sample(WishedSegmentLength, Coordinates);
+			// limits the count of segment to 1000 per edge
+			const double EdgeLength = Edge.Entity->Length();
+			const double SegmentLength = FMath::Min(FMath::Max(EdgeLength / 1000., WishedSegmentLength), EdgeLength);
+			Edge.Entity->Sample(SegmentLength, Coordinates);
 
 			TArray<FPoint2D> ScaledPoints;
 			{
@@ -958,7 +1327,7 @@ bool FThinZone2DFinder::BuildLoopSegments()
 					Coordinates.RemoveAt(EndPointIndex);
 				}
 
-				if (ScaledPoints.IsEmpty())
+				if (ScaledPoints.Num() == 1)
 				{
 					continue;
 				}
@@ -968,12 +1337,17 @@ bool FThinZone2DFinder::BuildLoopSegments()
 			{
 				for (int32 PointIndex = 1; PointIndex < ScaledPoints.Num(); PointIndex++)
 				{
-					while (ScaledPoints[PointIndex - 1].Distance(ScaledPoints[PointIndex]) < GeometricTolerance)
+					while (ScaledPoints.Num() > 1 && ScaledPoints[PointIndex - 1].Distance(ScaledPoints[PointIndex]) < GeometricTolerance)
 					{
 						ScaledPoints.RemoveAt(PointIndex);
 						Coordinates.RemoveAt(PointIndex);
 					}
 				}
+
+				if (ScaledPoints.Num() == 1)
+				{
+					continue;
+				}
 			}
 
 			const int32 LastPointIndex = ScaledPoints.Num() - 1;
@@ -1000,7 +1374,7 @@ bool FThinZone2DFinder::BuildLoopSegments()
 			}
 		}
 
-		if(PrecedingSegment)
+		if (PrecedingSegment)
 		{
 			PrecedingSegment->SetNext(FirstSegment);
 		}
@@ -1048,19 +1422,58 @@ void FThinZone2D::Finalize()
 
 	TMap<int32, FEdgeSegment*> NewSegmentMap;
 
-	TArray<FEdgeSegment>& FirstSideSegments = FirstSide.GetSegments();
-	TArray<FEdgeSegment>& SecondSideSegments = SecondSide.GetSegments();
-	BuildMap(FirstSideSegments, NewSegmentMap);
-	BuildMap(SecondSideSegments, NewSegmentMap);
+	TArray<FEdgeSegment>& SideASegments = SideA.GetSegments();
+	TArray<FEdgeSegment>& SideBSegments = SideB.GetSegments();
+	BuildMap(SideASegments, NewSegmentMap);
+	BuildMap(SideBSegments, NewSegmentMap);
 
-	UpdateReference(FirstSideSegments, NewSegmentMap);
-	UpdateReference(SecondSideSegments, NewSegmentMap);
+	UpdateReference(SideASegments, NewSegmentMap);
+	UpdateReference(SideBSegments, NewSegmentMap);
 
-	FirstSide.ComputeThicknessAndLength();
-	SecondSide.ComputeThicknessAndLength();
+	SideA.ComputeThicknessAndLength();
+	SideB.ComputeThicknessAndLength();
 
-	Thickness = (FirstSide.GetThickness() + SecondSide.GetThickness()) * 0.5;
-	MaxThickness = FMath::Max(FirstSide.GetMaxThickness(), SecondSide.GetMaxThickness());
+	Thickness = (SideA.GetThickness() + SideB.GetThickness()) * 0.5;
+	MaxThickness = FMath::Max(SideA.GetMaxThickness(), SideB.GetMaxThickness());
+}
+
+int32 FThinZoneSide::GetImposedPointCount()
+{
+	TMap<FTopologicalEdge*, FLinearBoundary> ThinZoneOfEdges;
+
+	FLinearBoundary* Boundary = nullptr;
+	for (FEdgeSegment& Segment : Segments)
+	{
+		FTopologicalEdge* Edge = nullptr;
+		if (Edge != Segment.GetEdge())
+		{
+			Edge = Segment.GetEdge();
+			ThinZoneOfEdges.Emplace(Edge);
+			Boundary = ThinZoneOfEdges.Find(Edge);
+			Boundary->Set(Segment.GetCoordinate(Start), Segment.GetCoordinate(End));
+		}
+		else if (Boundary)
+		{
+			Boundary->ExtendTo(Segment.GetCoordinate(Start), Segment.GetCoordinate(End));
+		}
+	}
+
+	int32 ImposedPointCount = 0;
+	for (TPair<FTopologicalEdge*, FLinearBoundary> ThinZoneOfEdge : ThinZoneOfEdges)
+	{
+		FTopologicalEdge* Edge = ThinZoneOfEdge.Key;
+		FLinearBoundary& ThinZoneBoundary = ThinZoneOfEdge.Value;
+
+		const TArray<FImposedCuttingPoint>& ImposedCuttingPoints = Edge->GetImposedCuttingPoints();
+		for (const FImposedCuttingPoint& CuttingPoint : ImposedCuttingPoints)
+		{
+			if (ThinZoneBoundary.Contains(CuttingPoint.Coordinate))
+			{
+				ImposedPointCount++;
+			}
+		}
+	}
+	return ImposedPointCount;
 }
 
 void FThinZoneSide::ComputeThicknessAndLength()
@@ -1086,15 +1499,96 @@ void FThinZoneSide::ComputeThicknessAndLength()
 	MaxThickness = FMath::Sqrt(SquareMaxThickness);
 }
 
-void FThinZoneSide::SetEdgesAsThinZone()
+void FThinZoneSide::GetEdges(TArray<FTopologicalEdge*>& OutEdges) const
+{
+	for (FTopologicalEdge* Edge : Edges)
+	{
+		if (!Edge->IsProcessed())
+		{
+			OutEdges.Add(Edge);
+			Edge->SetProcessedMarker();
+		}
+	}
+}
+
+void FThinZoneSide::AddToEdge()
 {
 	FTopologicalEdge* Edge = nullptr;
-	for (FEdgeSegment& Segment : Segments)
+	FLinearBoundary SideEdgeCoordinate;
+
+	TFunction<void(FTopologicalEdge*, FLinearBoundary&) > SetEdge = [this](FTopologicalEdge* Edge, FLinearBoundary& SideEdgeCoordinate)
 	{
-		if (Edge != Segment.GetEdge())
+		Edge->AddThinZone(this, SideEdgeCoordinate);
+		Edge->SetThinZoneMarker();
+		Edges.AddUnique(Edge);
+	};
+
+	for (FEdgeSegment& EdgeSegment : GetSegments())
+	{
+		double UMin = EdgeSegment.GetCoordinate(ELimit::Start);
+		double UMax = EdgeSegment.GetCoordinate(ELimit::End);
+		GetMinMax(UMin, UMax);
+
+		if (Edge != EdgeSegment.GetEdge())
 		{
-			Edge = Segment.GetEdge();
-			Edge->SetThinZoneMarker();
+			if (Edge)
+			{
+				SetEdge(Edge, SideEdgeCoordinate);
+			}
+
+			Edge = EdgeSegment.GetEdge();
+			SideEdgeCoordinate.Set(UMin, UMax);
+		}
+		else
+		{
+			SideEdgeCoordinate.ExtendTo(UMin, UMax);
+		}
+	};
+	SetEdge(Edge, SideEdgeCoordinate);
+}
+
+void FThinZoneSide::CheckEdgesZoneSide()
+{
+	ResetMarker1();
+	ResetMarker2();
+	for (FTopologicalEdge* Edge : Edges)
+	{
+		if (Edge->HasMarker1())
+		{
+			SetMarker1();
+		}
+		if (Edge->HasMarker2())
+		{
+			SetMarker2();
+		}
+	}
+}
+
+void FThinZoneSide::SetEdgesZoneSide(ESide Side)
+{
+	ResetMarkers();
+	if (Side == ESide::First)
+	{
+		Algo::ForEach(Edges, [](FTopologicalEdge* Edge) { Edge->SetMarker1(); });
+	}
+	else
+	{
+		Algo::ForEach(Edges, [](FTopologicalEdge* Edge) { Edge->SetMarker2(); });
+	}
+}
+
+void FThinZoneSide::CleanMesh()
+{
+	FTopologicalEdge* Edge = nullptr;
+	for (FEdgeSegment& EdgeSegment : GetSegments())
+	{
+		if (Edge != EdgeSegment.GetEdge())
+		{
+			Edge = EdgeSegment.GetEdge();
+			if (!Edge->IsMeshed())
+			{
+				Edge->RemovePreMesh();
+			}
 		}
 	}
 }
@@ -1138,11 +1632,276 @@ EMeshingState FThinZoneSide::GetMeshingState() const
 	return FullyMeshed;
 }
 
-void FThinZone2D::SetEdgesAsThinZone()
+void FThinZone2D::AddToEdge()
 {
-	FirstSide.SetEdgesAsThinZone();
-	SecondSide.SetEdgesAsThinZone();
+	SideA.AddToEdge();
+	SideB.AddToEdge();
 }
 
+void FThinZoneSide::GetExistingMeshNodes(const FTopologicalFace& Face, FModelMesh& MeshModel, FReserveContainerFunc& Reserve, FAddMeshNodeFunc& AddMeshNode, const bool bWithTolerance) const
+{
+	bool bEdgeIsForward = true;
+	int32 Index = 0;
+	int32 Increment = 1;
+	double SegmentUMin = 0.;
+	double SegmentUMax = 0.;
+	int32 EdgeMeshNodeCount = 0;
+
+	TArray<double> CoordinatesOfMesh;
+	TArray<int32> NodeIndices;
+	TArray<FPairOfIndex> OppositeNodeIndices;
+	TArray<double> EdgeElementLength;
+
+	typedef TFunction<bool(double, double)> FCompareMethode;
+	TFunction<void(double, FCompareMethode)> FindFirstIndex = [&CoordinatesOfMesh, &Index, &Increment](double ULimit, FCompareMethode Compare)
+	{
+		for (; Index >= 0 && Index < CoordinatesOfMesh.Num(); Index += Increment)
+		{
+			if (Compare(ULimit, CoordinatesOfMesh[Index]))
+			{
+				break;
+			}
+		}
+	};
+
+	TFunction<void(const FTopologicalEdge*, const FEdgeSegment&, ELimit)> AddEdgeVertexIfInclude = [&](const FTopologicalEdge* Edge, const FEdgeSegment& EdgeSegment, ELimit VertexExtremity)
+	{
+		const bool bIsStartVertex = (bEdgeIsForward == (VertexExtremity == ELimit::Start));
+
+		double VertexCoordinate = bIsStartVertex ? Edge->GetStartCurvilinearCoordinates() : Edge->GetEndCurvilinearCoordinates();
+
+		if (VertexCoordinate > SegmentUMin - DOUBLE_SMALL_NUMBER && VertexCoordinate < SegmentUMax + DOUBLE_SMALL_NUMBER)
+		{
+			FTopologicalVertex& Vertex = bIsStartVertex ? *Edge->GetStartVertex() : *Edge->GetEndVertex();
+			int32 MeshVertexId = Vertex.GetOrCreateMesh(MeshModel).GetMesh();
+			FPoint2D Vertex2DCoordinate = EdgeSegment.ComputeEdgePoint(VertexCoordinate);
+
+			FPairOfIndex OppositeNodeIndices = FPairOfIndex::Undefined;
+
+			const TArray<FCuttingPoint>& CuttingPoints = Edge->GetCuttingPoints();
+			if (CuttingPoints.Num())
+			{
+				const FCuttingPoint& CuttingPoint = bIsStartVertex ? CuttingPoints[0] : CuttingPoints.Last();
+				if (FMath::IsNearlyEqual(CuttingPoint.Coordinate, VertexCoordinate))
+				{
+					OppositeNodeIndices = CuttingPoint.OppositNodeIndices;
+				}
+			}
+
+			const double MeshingTolerance = FMath::Max(DOUBLE_KINDA_SMALL_NUMBER, bWithTolerance ? (bIsStartVertex ? EdgeElementLength[0] : EdgeElementLength.Last()) * ASixth : DOUBLE_KINDA_SMALL_NUMBER);
+			AddMeshNode(MeshVertexId, Vertex2DCoordinate, MeshingTolerance, EdgeSegment, OppositeNodeIndices);
+		}
+	};
+
+	TFunction<void(const FEdgeSegment&)> AddEdgeMeshNodeIfInclude = [&](const FEdgeSegment& EdgeSegment)
+	{
+		for (; Index >= 0 && Index < EdgeMeshNodeCount; Index += Increment)
+		{
+			const double MeshNodeCoordinate = CoordinatesOfMesh[Index];
+			if (MeshNodeCoordinate < SegmentUMin || MeshNodeCoordinate > SegmentUMax)
+			{
+				break;
+			}
+
+			const double MeshingTolerance = bWithTolerance ? FMath::Min(EdgeElementLength[Index], EdgeElementLength[Index + 1]) * AThird : 0;
+			const FPoint2D MeshNode2D = EdgeSegment.ComputeEdgePoint(MeshNodeCoordinate);
+			AddMeshNode(NodeIndices[Index], MeshNode2D, MeshingTolerance, EdgeSegment, OppositeNodeIndices[Index]);
+		}
+	};
+
+	// Count the existing mesh vertices on the Side
+	{
+		int32 MeshVertexCount = 0;
+		const FTopologicalEdge* Edge = nullptr;
+		for (const FEdgeSegment& EdgeSegment : GetSegments())
+		{
+			if (Edge != EdgeSegment.GetEdge())
+			{
+				Edge = EdgeSegment.GetEdge();
+				MeshVertexCount++; // first vertex
+
+				if (Edge->GetLinkActiveEdge()->IsMeshed())
+				{
+					const TArray<FPoint>& EdgeMeshNodes = Edge->GetMesh()->GetNodeCoordinates();
+					MeshVertexCount += EdgeMeshNodes.Num();
+				}
+				else if (Edge->IsPreMeshed())
+				{
+					const TArray<FCuttingPoint>& EdgeCuttingPoints = Edge->GetCuttingPoints();
+					MeshVertexCount += EdgeCuttingPoints.Num();
+				}
+			}
+		}
+		MeshVertexCount++; // last vertex
+
+		Reserve(MeshVertexCount);
+	}
+
+	// Get the existing mesh vertices and fill output data
+	{
+		FTopologicalEdge* Edge = nullptr;
+		const FTopologicalEdge* ActiveEdge = nullptr;
+
+		EOrientation EdgeOrientation = EOrientation::Front;
+
+		for (const FEdgeSegment& EdgeSegment : GetSegments())
+		{
+			SegmentUMin = EdgeSegment.GetCoordinate(ELimit::Start);
+			SegmentUMax = EdgeSegment.GetCoordinate(ELimit::End);
+			GetMinMax(SegmentUMin, SegmentUMax);
+
+			if (Edge != EdgeSegment.GetEdge())
+			{
+				Edge = EdgeSegment.GetEdge();
+				ActiveEdge = &Edge->GetLinkActiveEdge().Get();
+				bEdgeIsForward = EdgeSegment.IsForward();
+				EdgeMeshNodeCount = 0;
+
+				if (bWithTolerance)
+				{
+					if (ActiveEdge->IsMeshed())
+					{
+						EdgeElementLength = Edge->GetMesh()->GetElementLengths();
+					}
+					else
+					{
+						EdgeElementLength = Edge->GetPreElementLengths();
+					}
+				}
+
+				AddEdgeVertexIfInclude(Edge, EdgeSegment, ELimit::Start);
+
+				if (ActiveEdge->IsPreMeshed() && !Edge->IsDegenerated())
+				{
+					{
+						FAddCuttingPointFunc AddCuttingPoint = [&Edge](const double Coordinate, const ECoordinateType Type, const FPairOfIndex OppositNodeIndices, const double DeltaU)
+						{
+							Edge->AddTwinsCuttingPoint(Coordinate, DeltaU);
+						};
+
+						const bool bOnlyWithOppositeNode = false;
+						Edge->TransferCuttingPointFromMeshedEdge(bOnlyWithOppositeNode, AddCuttingPoint);
+					}
+
+					const TArray<FCuttingPoint>& CuttingPoints = Edge->GetCuttingPoints();
+					EdgeMeshNodeCount = CuttingPoints.Num();
+					if (EdgeMeshNodeCount > 2)
+					{
+						CoordinatesOfMesh.Empty(EdgeMeshNodeCount);
+						OppositeNodeIndices.Empty(EdgeMeshNodeCount);
+						EdgeMeshNodeCount -= 2;
+						for (int32 Cndex = 1; Cndex <= EdgeMeshNodeCount; ++Cndex)
+						{
+							const FCuttingPoint& CuttingPoint = CuttingPoints[Cndex];
+							CoordinatesOfMesh.Add(CuttingPoint.Coordinate);
+							OppositeNodeIndices.Emplace(CuttingPoint.OppositNodeIndices);
+						}
+
+						if (ActiveEdge->IsMeshed())
+						{
+							NodeIndices.Empty(EdgeMeshNodeCount);
+							const FEdgeMesh* EdgeMesh = Edge->GetMesh();
+							if (!EdgeMesh)
+							{
+								continue;
+							}
+
+							int32 StartMeshVertexId = EdgeMesh->GetStartVertexId();
+							if (Edge->IsSameDirection(*ActiveEdge))
+							{
+								for (int32 Cndex = 0; Cndex < EdgeMeshNodeCount; ++Cndex)
+								{
+									NodeIndices.Add(StartMeshVertexId++);
+								}
+							}
+							else
+							{
+								StartMeshVertexId += EdgeMeshNodeCount;
+								for (int32 Cndex = 0; Cndex < EdgeMeshNodeCount; ++Cndex)
+								{
+									NodeIndices.Add(--StartMeshVertexId);
+								}
+							}
+						}
+						else
+						{
+							NodeIndices.Init(-1, EdgeMeshNodeCount);
+						}
+					}
+					else if (EdgeMeshNodeCount > 0)
+					{
+						EdgeMeshNodeCount = 0;
+					}
+					else
+					{
+						const FEdgeMesh* EdgeMesh = Edge->GetMesh();
+						if (!EdgeMesh)
+						{
+							continue;
+						}
+
+						const TArray<FPoint>& MeshVertices = EdgeMesh->GetNodeCoordinates();
+						int32 LocIncrement = 1;
+
+						int32 StartMeshVertexId = EdgeMesh->GetStartVertexId();
+						TArray<FPoint> ProjectedPoints;
+						Edge->ProjectPoints(MeshVertices, CoordinatesOfMesh, ProjectedPoints);
+						if (EdgeMeshNodeCount > 1 && CoordinatesOfMesh[0] > CoordinatesOfMesh[1])
+						{
+							Algo::Reverse(CoordinatesOfMesh);
+							LocIncrement = -1;
+							StartMeshVertexId += EdgeMeshNodeCount - 1;
+						}
+
+						NodeIndices.Empty(EdgeMeshNodeCount);
+						OppositeNodeIndices.Init(FPairOfIndex::Undefined, EdgeMeshNodeCount);
+						for (const double& Coordinate : CoordinatesOfMesh)
+						{
+							NodeIndices.Add(StartMeshVertexId);
+							StartMeshVertexId += LocIncrement;
+						}
+					}
+
+					if (EdgeElementLength.Num() != CuttingPoints.Num())
+					{
+						EdgeElementLength = Edge->GetPreElementLengths();
+					}
+
+					if (bEdgeIsForward)
+					{
+						Index = 0;
+						Increment = 1;
+						FindFirstIndex(SegmentUMin, [](double Value1, double Value2) { return (Value1 < Value2); });
+					}
+					else
+					{
+						Index = EdgeMeshNodeCount - 1;
+						Increment = -1;
+						FindFirstIndex(SegmentUMax, [](double Value1, double Value2) { return (Value1 > Value2); });
+					}
+				}
+				else if (bWithTolerance)
+				{
+					EdgeElementLength.Empty(1);
+					EdgeElementLength.Add(Edge->Length());
+				}
+			}
+
+			if (EdgeMeshNodeCount > Index)
+			{
+				AddEdgeMeshNodeIfInclude(EdgeSegment);
+			}
+		}
+
+		// Last vertex, last segment 
+		{
+			const FEdgeSegment& EdgeSegment = GetSegments().Last();
+			AddEdgeVertexIfInclude(Edge, EdgeSegment, ELimit::End);
+		}
+	}
+}
+
+
+
 }
 
diff --git a/Engine/Source/Runtime/Datasmith/CADKernel/Private/CADKernel/Mesh/Structure/ThinZone2DFinder.h b/Engine/Source/Runtime/Datasmith/CADKernel/Private/CADKernel/Mesh/Structure/ThinZone2DFinder.h
--- a/Engine/Source/Runtime/Datasmith/CADKernel/Private/CADKernel/Mesh/Structure/ThinZone2DFinder.h
+++ b/Engine/Source/Runtime/Datasmith/CADKernel/Private/CADKernel/Mesh/Structure/ThinZone2DFinder.h
@@ -11,9 +11,9 @@
 #include "CADKernel/Mesh/Structure/GridBase.h"
 #include "CADKernel/Mesh/Structure/ThinZone2D.h"
 
-//#ifdef CADKERNEL_DEV
-//#include "CADKernel/UI/DefineForDebug.h"
-//#endif
+#ifdef CADKERNEL_DEBUG
+#include "CADKernel/UI/DefineForDebug.h"
+#endif
 
 namespace UE::CADKernel
 {
@@ -93,7 +93,7 @@ namespace UE::CADKernel
 		/**
 		 * For Segment with Marker1 in Segments, find in OppositeSides the closest segment in the respect of the criteria
 		 */
-		void FindCloseSegments(const TArray<FEdgeSegment*> Segments, const TArray< const TArray<FEdgeSegment*>*> OppositeSides);
+		void FindCloseSegments(const TArray<FEdgeSegment*>& Segments, const TArray<const TArray<FEdgeSegment*>*>& OppositeSides);
 
 		/**
 		 * Fill ThinZoneSides (TArray<TArray<FEdgeSegment*>>)
@@ -114,6 +114,20 @@ namespace UE::CADKernel
 		 */
 		void ImproveThinSide();
 
+		/**
+		 * Check that the thin side is not opposite to a unique point
+		 * 
+		 *     #-----------------------#  a Side
+		 * 
+		 * 
+		 *                                      #------------ 
+		 *                                     /
+		 *           The opposite unique point 
+		 * @return false if the opposite is degenerated
+		 */
+		bool CheckIfCloseSideOfThinSideIsNotDegenerated(TArray<FEdgeSegment*>& Side);
+		void CheckIfCloseSideOfThinSidesAreNotDegenerated();
+
 		/**
 		 * SplitThinSide splits chain in front of many other chain eg.
 		 * Side A is close to Side B then Side C
@@ -145,11 +159,14 @@ namespace UE::CADKernel
 #ifdef CADKERNEL_DEBUG
 namespace ThinZone
 {
+static bool bDisplay = false;
 void DisplayEdgeSegmentAndProjection(const FEdgeSegment* Segment, EVisuProperty SegColor, EVisuProperty OppositeColor, EVisuProperty ProjectionColor);
 void DisplayEdgeSegmentAndProjection(const FEdgeSegment* Segment, const FEdgeSegment* Opposite, EVisuProperty SegColor, EVisuProperty OppositeColor, EVisuProperty ProjectionColor);
 void DisplayEdgeSegment(const FEdgeSegment* EdgeSegment, EVisuProperty Color);
 void DisplayEdgeSegment(const FEdgeSegment* EdgeSegment, EVisuProperty Color, int32 Index);
 void DisplayThinZoneSides(const TArray<TArray<FEdgeSegment*>>& ThinZoneSides);
+void DisplayThinZoneSides2(const TArray<TArray<FEdgeSegment*>>& ThinZoneSides);
+void DisplayThinZoneSide2(const TArray<FEdgeSegment*>& Side, int32 Index, bool bSplitBySide = false, bool bSplitBySegment = false);
 void DisplayThinZoneSide(const TArray<FEdgeSegment*>& Side, int32 Index, EVisuProperty Color, bool bSplitBySegment = false);
 void DisplayThinZoneSide(const FThinZoneSide& Side, int32 Index, EVisuProperty Color, bool bSplitBySegment = false);
 void DisplayThinZones(const TArray<FThinZone2D>& ThinZones);
diff --git a/Engine/Source/Runtime/Datasmith/CADKernel/Private/CADKernel/Topo/Shell.cpp b/Engine/Source/Runtime/Datasmith/CADKernel/Private/CADKernel/Topo/Shell.cpp
--- a/Engine/Source/Runtime/Datasmith/CADKernel/Private/CADKernel/Topo/Shell.cpp
+++ b/Engine/Source/Runtime/Datasmith/CADKernel/Private/CADKernel/Topo/Shell.cpp
@@ -64,7 +64,7 @@ void FShell::Empty()
 	FTopologicalShapeEntity::Empty();
 }
 
-void FShell::Add(TArray<FTopologicalFace*> Faces)
+void FShell::Add(TArray<FTopologicalFace*>& Faces)
 {
 	TopologicalFaces.Reserve(TopologicalFaces.Num() + Faces.Num());
 
@@ -126,10 +126,10 @@ void FShell::Merge(TSharedPtr<FShell>& Shell)
 
 void FShell::SpreadBodyOrientation()
 {
-	bool bIsOutter = IsOutter();
+	const bool bIsOuter = IsOuter();
 	for (FOrientedFace& Face : TopologicalFaces)
 	{
-		if (bIsOutter != (Face.Direction == EOrientation::Front))
+		if (bIsOuter != (Face.Direction == EOrientation::Front))
 		{
 			Face.Entity->SetBackOriented();
 		}
@@ -138,10 +138,10 @@ void FShell::SpreadBodyOrientation()
 
 void FShell::UpdateShellOrientation()
 {
-	bool bIsOutter = IsOutter();
+	const bool bIsOuter = IsOuter();
 	for (FOrientedFace& Face : TopologicalFaces)
 	{
-		if (Face.Entity->IsBackOriented())
+		if (bIsOuter != Face.Entity->IsBackOriented())
 		{
 			Face.Direction = EOrientation::Back;
 		}
@@ -304,12 +304,7 @@ int32 FShell::Orient()
 			for (const FOrientedEdge& OrientedEdge : Loop->GetEdges())
 			{
 				const TSharedPtr<FTopologicalEdge>& Edge = OrientedEdge.Entity;
-				if (Edge->HasMarker1())
-				{
-					continue;
-				}
-				Edge->SetMarker1();
-
+	
 				if ((Edge->GetTwinEntityCount() != 2) || Edge->IsDegenerated())
 				{
 					if (Edge->GetTwinEntityCount() == 1)
@@ -319,17 +314,7 @@ int32 FShell::Orient()
 					continue;
 				}
 
-				FTopologicalEdge* NeighboringEdge = nullptr;
-				for (FTopologicalEdge* NextEdge : Edge->GetTwinEntities())
-				{
-					if (NextEdge == Edge.Get())
-					{
-						continue;
-					}
-					NeighboringEdge = NextEdge;
-					break;
-				}
-
+				FTopologicalEdge* NeighboringEdge = Edge->GetTwinEdge();
 				if (NeighboringEdge == nullptr)
 				{
 					continue;
@@ -347,12 +332,12 @@ int32 FShell::Orient()
 					continue;
 				}
 
-				if (NeighboringFace->HasMarker1())
+				if (NeighboringFace->IsNotToOrAlreadyProcess())
 				{
 					continue;
 				}
 
-				NeighboringFace->SetMarker1();
+				NeighboringFace->SetProcessedMarker();
 				Front.Add(NeighboringFace);
 
 				const FOrientedEdge* OrientedNeighboringEdge = NeighboringLoop->GetOrientedEdge(NeighboringEdge);
@@ -400,10 +385,6 @@ int32 FShell::Orient()
 
 		for (FTopologicalFace* Face : Subshell)
 		{
-			if (Face == nullptr || Face->IsDeleted())
-			{
-				continue;
-			}
 			Face->UpdateBBox(3, ApproximationFactor, BBox);
 		}
 
@@ -431,10 +412,6 @@ int32 FShell::Orient()
 	{
 		for (FTopologicalFace* Face : Subshell)
 		{
-			if (Face == nullptr)
-			{
-				continue;
-			}
 			Face->SwapOrientation();
 		}
 	};
@@ -502,19 +479,16 @@ int32 FShell::Orient()
 	for (FOrientedFace& Face : GetFaces())
 	{
 		Face.Entity->ResetMarkers();
-		for (const TSharedPtr<FTopologicalLoop>& Loop : Face.Entity->GetLoops())
+		if(!Face.Entity->IsDeletedOrDegenerated())
 		{
-			for (const FOrientedEdge& OrientedEdge : Loop->GetEdges())
-			{
-				OrientedEdge.Entity->ResetMarkers();
-			}
+			Face.Entity->SetToProcessMarker();
 		}
 	}
 
 	int32 ProcessFaceCount = 0;
 	for (FOrientedFace& Face : GetFaces())
 	{
-		if (Face.Entity->HasMarker1())
+		if (Face.Entity->IsNotToOrAlreadyProcess())
 		{
 			continue;
 		}
@@ -526,7 +500,7 @@ int32 FShell::Orient()
 		Subshell.Empty(TopologicalFaceCount);
 		Front.Empty(TopologicalFaceCount);
 
-		Face.Entity->SetMarker1();
+		Face.Entity->SetProcessedMarker();
 		Front.Add(Face.Entity.Get());
 
 		SpreadFront();
@@ -543,7 +517,10 @@ int32 FShell::Orient()
 
 	UpdateShellOrientation();
 
-	ResetMarkersRecursively();
+	for (FOrientedFace& Face : GetFaces())
+	{
+		Face.Entity->ResetMarkers();
+	}
 
 	return ShellSwappedFaceCount;
 }
diff --git a/Engine/Source/Runtime/Datasmith/CADKernel/Private/CADKernel/Topo/TopologicalEdge.cpp b/Engine/Source/Runtime/Datasmith/CADKernel/Private/CADKernel/Topo/TopologicalEdge.cpp
--- a/Engine/Source/Runtime/Datasmith/CADKernel/Private/CADKernel/Topo/TopologicalEdge.cpp
+++ b/Engine/Source/Runtime/Datasmith/CADKernel/Private/CADKernel/Topo/TopologicalEdge.cpp
@@ -11,6 +11,8 @@
 #include "CADKernel/Math/Point.h"
 #include "CADKernel/Math/SlopeUtils.h"
 #include "CADKernel/Mesh/Structure/EdgeMesh.h"
+#include "CADKernel/Mesh/Structure/ModelMesh.h"
+#include "CADKernel/Mesh/Structure/VertexMesh.h"
 #include "CADKernel/Topo/TopologicalLoop.h"
 #include "CADKernel/Topo/TopologicalLink.h"
 #include "CADKernel/Topo/TopologicalFace.h"
@@ -104,6 +106,8 @@ bool FTopologicalEdge::CheckIfDegenerated() const
 
 	Curve->CheckIfDegenerated(Boundary, bDegeneration2D, bDegeneration3D, Length3D);
 
+	Max2DTolerance = Length3D * FactorToComputeMaxTol;
+
 	if (bDegeneration3D)
 	{
 		SetAsDegenerated();
@@ -482,11 +486,17 @@ void FTopologicalEdge::SetEndVertex(const double NewCoordinate, const FPoint& Ne
 	EndVertex->SetCoordinates(NewPoint3D);
 }
 
+void FTopologicalEdge::ComputeLength()
+{
+	Length3D = Curve->ApproximateLength(Boundary);
+	Max2DTolerance = Length3D * FactorToComputeMaxTol;
+}
+
 double FTopologicalEdge::Length() const
 {
 	if (Length3D < 0)
 	{
-		Length3D = Curve->ApproximateLength(Boundary);
+		const_cast<FTopologicalEdge*>(this)->ComputeLength();
 	}
 	return Length3D;
 }
@@ -529,6 +539,42 @@ int32 FTopologicalEdge::EvaluateCuttingPointNum()
 	return (int32)Num;
 }
 
+double FTopologicalEdge::TransformTwinEdgeCoordinateToLocalCoordinate(const FTopologicalEdge& TwinEdge, const double InTwinCoordinate) const
+{
+	if (this == &TwinEdge)
+	{
+		return InTwinCoordinate;
+	}
+
+	if (IsDegenerated() || TwinEdge.IsDegenerated())
+	{
+		// linear transform
+		const bool bSameDirection = IsSameDirection(TwinEdge);
+		const double Start = bSameDirection ? Boundary.GetMin(): Boundary.GetMax();
+		const double End = bSameDirection ? Boundary.GetMax() : Boundary.GetMin();
+
+		const double Distance = End - Start;
+		if (FMath::IsNearlyZero(Distance))
+		{
+			return Start;
+		}
+
+		const double TwinStart = TwinEdge.GetBoundary().GetMin();
+		const double TwinEnd = TwinEdge.GetBoundary().GetMax();
+		const double TwinDistance = TwinEnd - TwinStart;
+		if (FMath::IsNearlyZero(TwinDistance))
+		{
+			return Start;
+		}
+
+		return Start + (TwinStart - InTwinCoordinate) * Distance / TwinDistance;
+	}
+
+	FPoint PointOnEdge = TwinEdge.GetCurve()->Approximate3DPoint(InTwinCoordinate);
+	FPoint ProjectedPoint;
+	return Curve->GetCoordinateOfProjectedPoint(Boundary, PointOnEdge, ProjectedPoint);
+}
+
 double FTopologicalEdge::TransformLocalCoordinateToActiveEdgeCoordinate(const double InLocalCoordinate) const
 {
 	if (IsActiveEntity())
@@ -537,9 +583,7 @@ double FTopologicalEdge::TransformLocalCoordinateToActiveEdgeCoordinate(const do
 	}
 
 	const FTopologicalEdge& ActiveEdge = *GetLinkActiveEntity();
-	FPoint PointOnEdge = Curve->Approximate3DPoint(InLocalCoordinate);
-	FPoint ProjectedPoint;
-	return ActiveEdge.GetCurve()->GetCoordinateOfProjectedPoint(Boundary, PointOnEdge, ProjectedPoint);
+	return ActiveEdge.TransformTwinEdgeCoordinateToLocalCoordinate(*this, InLocalCoordinate);
 }
 
 double FTopologicalEdge::TransformActiveEdgeCoordinateToLocalCoordinate(const double InActiveEdgeCoordinate) const
@@ -550,44 +594,119 @@ double FTopologicalEdge::TransformActiveEdgeCoordinateToLocalCoordinate(const do
 	}
 
 	const FTopologicalEdge& ActiveEdge = *GetLinkActiveEntity();
-	FPoint PointOnEdge = ActiveEdge.GetCurve()->Approximate3DPoint(InActiveEdgeCoordinate);
-	FPoint ProjectedPoint;
-	return Curve->GetCoordinateOfProjectedPoint(Boundary, PointOnEdge, ProjectedPoint);
+	return TransformTwinEdgeCoordinateToLocalCoordinate(ActiveEdge, InActiveEdgeCoordinate);
 }
 
-void FTopologicalEdge::TransformLocalCoordinatesToActiveEdgeCoordinates(const TArray<double>& InLocalCoordinate, TArray<double>& OutActiveEdgeCoordinates) const
+void FTopologicalEdge::TransformTwinEdgeCoordinatesToLocalCoordinates(const FTopologicalEdge& TwinEdge, const TArray<double>& InTwinCoordinates, TArray<double>& OutLocalCoordinates) const
 {
-	if (IsActiveEntity())
+	if (this == &TwinEdge)
 	{
-		OutActiveEdgeCoordinates = InLocalCoordinate;
+		OutLocalCoordinates = InTwinCoordinates;
+		return;
 	}
-	const FTopologicalEdge& ActiveEdge = *GetLinkActiveEntity();
-	TArray<FPoint> EdgePoints;
-	Curve->Approximate3DPoints(InLocalCoordinate, EdgePoints);
-	TArray<FPoint> ProjectedPoints;
-	ActiveEdge.GetCurve()->ProjectPoints(Boundary, EdgePoints, OutActiveEdgeCoordinates, ProjectedPoints);
-}
 
-void FTopologicalEdge::TransformActiveEdgeCoordinatesToLocalCoordinates(const TArray<double>& InActiveEdgeCoordinate, TArray<double>& OutLocalCoordinates) const
-{
-	if (IsActiveEntity())
+	if (IsDegenerated() || TwinEdge.IsDegenerated())
 	{
-		OutLocalCoordinates = InActiveEdgeCoordinate;
+		// linear transform
+		const bool bSameDirection = IsSameDirection(TwinEdge);
+		const double Start = bSameDirection ? Boundary.GetMin() : Boundary.GetMax();
+		const double End = bSameDirection ? Boundary.GetMax() : Boundary.GetMin();
+
+		const double Distance = End - Start;
+		if (FMath::IsNearlyZero(Distance))
+		{
+			OutLocalCoordinates.Init(Start, InTwinCoordinates.Num());
+			return;
+		}
+
+		const double TwinStart = TwinEdge.GetBoundary().GetMin();
+		const double TwinEnd = TwinEdge.GetBoundary().GetMax();
+		const double TwinDistance = TwinEnd - TwinStart;
+		if (FMath::IsNearlyZero(TwinDistance))
+		{
+			OutLocalCoordinates.Init(Start, InTwinCoordinates.Num());
+			return;
+		}
+
+		OutLocalCoordinates.Empty(InTwinCoordinates.Num());
+		const double Factor = Distance / TwinDistance;
+		for(double TwinCoordinate : InTwinCoordinates)
+		{
+			OutLocalCoordinates.Add(Start + (TwinCoordinate - TwinStart) * Factor);
+		}
 		return;
 	}
 
+	TArray<FPoint> Cutting3DPoints;
+	TwinEdge.ApproximatePoints(InTwinCoordinates, Cutting3DPoints);
+	ProjectTwinEdgePoints(Cutting3DPoints, IsSameDirection(TwinEdge), OutLocalCoordinates);
+}
+
+void FTopologicalEdge::TransformLocalCoordinatesToActiveEdgeCoordinates(const TArray<double>& InLocalCoordinates, TArray<double>& OutActiveEdgeCoordinates) const
+{
+	const FTopologicalEdge& ActiveEdge = *GetLinkActiveEntity();
+	return ActiveEdge.TransformTwinEdgeCoordinatesToLocalCoordinates(*this, InLocalCoordinates, OutActiveEdgeCoordinates);
+}
+
+void FTopologicalEdge::TransformActiveEdgeCoordinatesToLocalCoordinates(const TArray<double>& InActiveEdgeCoordinates, TArray<double>& OutLocalCoordinates) const
+{
 	const FTopologicalEdge& ActiveEdge = *GetLinkActiveEntity();
-	TArray<FPoint> ActiveEdgePoint;
-	ActiveEdge.GetCurve()->Approximate3DPoints(InActiveEdgeCoordinate, ActiveEdgePoint);
-	TArray<FPoint> ProjectedPoints;
-	Curve->ProjectPoints(Boundary, ActiveEdgePoint, OutLocalCoordinates, ProjectedPoints);
+	return TransformTwinEdgeCoordinatesToLocalCoordinates(ActiveEdge, InActiveEdgeCoordinates, OutLocalCoordinates);
 }
 
-void FTopologicalEdge::AddImposedCuttingPointU(const double ImposedCuttingPointU, int32 OppositeNodeIndex)
+void FTopologicalEdge::SortImposedCuttingPoints()
 {
-	ImposedCuttingPointUs.Emplace(ImposedCuttingPointU, OppositeNodeIndex);
+	Algo::Sort(ImposedCuttingPointUs, [](const FImposedCuttingPoint& C1, const FImposedCuttingPoint& C2) { return C1.Coordinate < C2.Coordinate; });
+
+	if (ImposedCuttingPointUs.Num() > 1)
+	{
+		int32 LastIndex = -1;
+		int32 NewIndex = 0;
+		for (int32 Index = 1; Index < ImposedCuttingPointUs.Num(); ++Index)
+		{
+			if (FMath::IsNearlyEqual(ImposedCuttingPointUs[NewIndex].Coordinate, ImposedCuttingPointUs[Index].Coordinate))
+			{
+				if (LastIndex < 0)
+				{
+					LastIndex = NewIndex;
+				}
+				if (ImposedCuttingPointUs[NewIndex].OppositNodeIndex == -1)
+				{
+					Swap(ImposedCuttingPointUs[NewIndex], ImposedCuttingPointUs[Index]);
+				}
+				else if (ImposedCuttingPointUs[Index].OppositNodeIndex == -1)
+				{
+				}
+				else if (ImposedCuttingPointUs[NewIndex].OppositNodeIndex == ImposedCuttingPointUs[Index].OppositNodeIndex)
+				{
+					ImposedCuttingPointUs[NewIndex].DeltaU = FMath::Max(ImposedCuttingPointUs[NewIndex].DeltaU, ImposedCuttingPointUs[Index].DeltaU);
+				}
+				else if ((LastIndex >= 0) && (LastIndex != NewIndex) && (ImposedCuttingPointUs[LastIndex].OppositNodeIndex == ImposedCuttingPointUs[Index].OppositNodeIndex))
+				{
+					ImposedCuttingPointUs[LastIndex].DeltaU = FMath::Max(ImposedCuttingPointUs[LastIndex].DeltaU, ImposedCuttingPointUs[Index].DeltaU);
+				}
+				else
+				{
+					NewIndex++;
+				}
+			}
+			else
+			{
+				LastIndex = -1;
+				NewIndex++;
+				if (NewIndex != Index)
+				{
+					Swap(ImposedCuttingPointUs[NewIndex], ImposedCuttingPointUs[Index]);
+				}
+			}
+		}
+		NewIndex++;
+		ImposedCuttingPointUs.SetNum(NewIndex);
+	}
 }
 
+
+
 void FTopologicalEdge::ProjectTwinEdgePointsOn2DCurve(const TSharedRef<FTopologicalEdge>& InTwinEdge, const TArray<double>& InTwinEdgePointCoords, TArray<FPoint2D>& OutPoints2D)
 {
 	if (&InTwinEdge.Get() == this)
@@ -599,13 +718,60 @@ void FTopologicalEdge::ProjectTwinEdgePointsOn2DCurve(const TSharedRef<FTopologi
 		TArray<FPoint> Points3D;
 		InTwinEdge->ApproximatePoints(InTwinEdgePointCoords, Points3D);
 
-		bool bSameDirection = IsSameDirection(*InTwinEdge);
 		TArray<double> Coordinates;
-		Curve->ProjectTwinCurvePoints(Points3D, bSameDirection, Coordinates);
+
+		bool bSameDirection = IsSameDirection(*InTwinEdge);
+		const double ToleranceOfProjection = Length3D * 0.1;
+		Curve->ProjectTwinCurvePoints(Points3D, bSameDirection, Coordinates, ToleranceOfProjection);
 		Curve->Approximate2DPoints(Coordinates, OutPoints2D);
 	}
 }
 
+void FTopologicalEdge::GenerateMeshElements(FModelMesh& MeshModel)
+{
+	FTopologicalEdge& ActiveEdge = *GetLinkActiveEntity();
+
+	bool bSameDirection = IsSameDirection(ActiveEdge);
+
+	FEdgeMesh& EdgeMesh = ActiveEdge.GetOrCreateMesh(MeshModel);
+
+	int32 StartVertexNodeIndex = ActiveEdge.GetStartVertex()->GetOrCreateMesh(MeshModel).GetMesh();
+	int32 EndVertexNodeIndex = ActiveEdge.GetEndVertex()->GetOrCreateMesh(MeshModel).GetMesh();
+
+	TArray<double> CuttingPointCoordinates = GetCuttingPointCoordinates();
+	ensureCADKernel(CuttingPointCoordinates.Num() > 1);
+	CuttingPointCoordinates.RemoveAt(0);
+	CuttingPointCoordinates.Pop();
+
+	TArray<FPoint>& Coordinates = EdgeMesh.GetNodeCoordinates();
+	ApproximatePoints(CuttingPointCoordinates, Coordinates);
+
+	if (!bSameDirection)
+	{
+		Algo::Reverse(Coordinates);
+	}
+
+#ifdef DEBUG_MESH_EDGE
+	if (bDisplay)
+	{
+		F3DDebugSession _(FString::Printf(TEXT("Edge Mesh %d"), Edge.GetId()));
+		TArray<FPoint2D> Mesh2D;
+		Edge.Approximate2DPoints(CuttingPointCoordinates, Mesh2D);
+		for (const FPoint2D& Vertex : Mesh2D)
+		{
+			DisplayPoint(Vertex, EVisuProperty::RedPoint);
+		}
+		Wait();
+	}
+#endif
+
+	EdgeMesh.RegisterCoordinates();
+	EdgeMesh.Mesh(StartVertexNodeIndex, EndVertexNodeIndex);
+	MeshModel.AddMesh(EdgeMesh);
+	ActiveEdge.SetMeshedMarker();
+	SetMeshedMarker();
+}
+
 bool FTopologicalEdge::IsSameDirection(const FTopologicalEdge& Edge) const
 {
 	if (!TopologicalLink)
@@ -676,7 +842,7 @@ FInfoEntity& FTopologicalEdge::GetInfo(FInfoEntity& Info) const
 }
 #endif
 
-TSharedRef<FEdgeMesh> FTopologicalEdge::GetOrCreateMesh(FModelMesh& ShellMesh)
+FEdgeMesh& FTopologicalEdge::GetOrCreateMesh(FModelMesh& ShellMesh)
 {
 	if (!IsActiveEntity())
 	{
@@ -687,7 +853,118 @@ TSharedRef<FEdgeMesh> FTopologicalEdge::GetOrCreateMesh(FModelMesh& ShellMesh)
 	{
 		Mesh = FEntity::MakeShared<FEdgeMesh>(ShellMesh, *this);
 	}
-	return Mesh.ToSharedRef();
+	return *Mesh;
+}
+
+void FTopologicalEdge::RemovePreMesh()
+{
+	if (IsMeshed())
+	{
+		return;
+	}
+	CuttingPointUs.Empty();
+	ResetPreMeshed();
+}
+
+const FTopologicalEdge* FTopologicalEdge::GetPreMeshedTwin() const
+{
+	for (const FTopologicalEdge* Edge : GetLink()->GetTwinEntities())
+	{
+		if (!Edge->GetCuttingPoints().IsEmpty())
+		{
+			return Edge;
+		}
+	}
+	return nullptr;
+}
+
+void FTopologicalEdge::AddImposedCuttingPointU(const double ImposedCuttingPointU, int32 OppositeNodeIndex, const double DeltaU)
+{
+	ImposedCuttingPointUs.Emplace(ImposedCuttingPointU, OppositeNodeIndex, DeltaU);
+}
+
+void FTopologicalEdge::AddTwinsCuttingPoint(double Coord, double DeltaU)
+{
+	if (FMath::IsNearlyEqual(Coord, Boundary.GetMin(), GetTolerance2DAt(Boundary.GetMin())))
+	{
+		Coord = Boundary.GetMin();
+	}
+	else if (FMath::IsNearlyEqual(Coord, Boundary.GetMax(), GetTolerance2DAt(Boundary.GetMax())))
+	{
+		Coord = Boundary.GetMax();
+	}
+
+	CuttingPointUs.Emplace(Coord, ECoordinateType::ImposedCoordinate, FPairOfIndex::Undefined, DeltaU);
+}
+
+void FTopologicalEdge::TransferCuttingPointFromMeshedEdge(bool bOnlyWithOppositeNode, FAddCuttingPointFunc AddCuttingPoint)
+{
+	const FTopologicalEdge* PreMeshedTwin = GetPreMeshedTwin();
+
+	if ((PreMeshedTwin == nullptr) || (PreMeshedTwin == this) || (PreMeshedTwin->GetFace() == GetFace()))
+	{
+		return;
+	}
+
+	TArray<double> PreMeshEdgeCuttingPointCoords = PreMeshedTwin->GetCuttingPointCoordinates();
+	TArray<double> CuttingPointCoords;
+	TransformTwinEdgeCoordinatesToLocalCoordinates(*PreMeshedTwin, PreMeshEdgeCuttingPointCoords, CuttingPointCoords);
+
+	const TArray<FCuttingPoint>& PreMeshEdgeCuttingPoints = PreMeshedTwin->GetCuttingPoints();
+	CuttingPointUs.Empty(PreMeshEdgeCuttingPoints.Num());
+	for (int32 Index = 0; Index < PreMeshEdgeCuttingPoints.Num(); ++Index)
+	{
+		const FCuttingPoint& PreMeshEdgeCuttingPoint = PreMeshEdgeCuttingPoints[Index];
+
+		if (bOnlyWithOppositeNode && PreMeshEdgeCuttingPoint.OppositNodeIndices[0] < 0)
+		{
+			continue;
+		}
+
+		const double PreMeshEdgeCuttingPointCoord = PreMeshEdgeCuttingPoint.Coordinate;
+		const double PreMeshEdgeCuttingPointDeltaU = PreMeshEdgeCuttingPoint.IsoDeltaU;
+
+		const double CuttingPointCoord = CuttingPointCoords[Index];
+		const double EdgeTol = GetTolerance2DAt(CuttingPointCoord);
+		const double PreMeshEdgeTol = PreMeshedTwin->GetTolerance2DAt(PreMeshEdgeCuttingPointCoord);
+
+		const double CuttingPointDeltaU = PreMeshEdgeCuttingPointDeltaU * EdgeTol / PreMeshEdgeTol;
+
+		AddCuttingPoint(CuttingPointCoord, PreMeshEdgeCuttingPoint.Type, PreMeshEdgeCuttingPoint.OppositNodeIndices, CuttingPointDeltaU);
+	}
+}
+
+TArray<double> FTopologicalEdge::GetCuttingPointCoordinates() const
+{
+	TArray<double> CuttingPointCoordinates;
+	CuttingPointCoordinates.Reserve(GetCuttingPoints().Num());
+	for (const FCuttingPoint& CuttingPoint : GetCuttingPoints())
+	{
+		CuttingPointCoordinates.Add(CuttingPoint.Coordinate);
+	}
+	return MoveTemp(CuttingPointCoordinates);
+}
+
+TArray<double> FTopologicalEdge::GetPreElementLengths() const
+{
+	const TArray<FCuttingPoint>& CuttingPoints = GetCuttingPoints();
+	const int32 CuttingPointCount = CuttingPoints.Num();
+
+	TArray<FPoint> InnerNodes;
+	if (CuttingPointCount > 2)
+	{
+		TArray<double> Coordinates;
+		Coordinates.Reserve(CuttingPointCount);
+		for (int32 Index = 1; Index < CuttingPointCount - 1; ++Index)
+		{
+			Coordinates.Add(CuttingPoints[Index].Coordinate);
+		}
+		ApproximatePoints(Coordinates, InnerNodes);
+	}
+
+	const FPoint& StartNode = GetStartVertex()->GetCoordinates();
+	const FPoint& EndNode = GetEndVertex()->GetCoordinates();
+	return PolylineTools::ComputePolylineSegmentLengths(StartNode, InnerNodes, EndNode);
 }
 
 TSharedPtr<FTopologicalEdge> FTopologicalEdge::CreateEdgeByMergingEdges(TArray<FOrientedEdge>& Edges, const TSharedRef<FTopologicalVertex> StartVertex, const TSharedRef<FTopologicalVertex> EndVertex)
@@ -1123,7 +1400,8 @@ TSharedPtr<FTopologicalVertex> FTopologicalEdge::SplitAt(double SplittingCoordin
 		Boundary.Min = SplittingCoordinate;
 	}
 	MiddelVertex->AddConnectedEdge(*this);
-	Length3D = Curve->ApproximateLength(Boundary);
+
+	ComputeLength();
 
 	Loop->SplitEdge(*this, NewEdge, bKeepStartVertexConnectivity);
 	return MiddelVertex;
diff --git a/Engine/Source/Runtime/Datasmith/CADKernel/Private/CADKernel/Topo/TopologicalFace.cpp b/Engine/Source/Runtime/Datasmith/CADKernel/Private/CADKernel/Topo/TopologicalFace.cpp
--- a/Engine/Source/Runtime/Datasmith/CADKernel/Private/CADKernel/Topo/TopologicalFace.cpp
+++ b/Engine/Source/Runtime/Datasmith/CADKernel/Private/CADKernel/Topo/TopologicalFace.cpp
@@ -191,21 +191,20 @@ void FTopologicalFace::ApplyNaturalLoops(const FSurfacicBoundary& Boundaries)
 
 	// Build 4 bounding edges of the surface
 	StartPoint.Set(Boundaries[EIso::IsoU].Min, Boundaries[EIso::IsoV].Min);
-	EndPoint.Set(Boundaries[EIso::IsoU].Min, Boundaries[EIso::IsoV].Max);
+	EndPoint.Set(Boundaries[EIso::IsoU].Max, Boundaries[EIso::IsoV].Min);
 	BuildEdge(StartPoint, EndPoint);
 
-	StartPoint.Set(Boundaries[EIso::IsoU].Min, Boundaries[EIso::IsoV].Max);
+	StartPoint.Set(Boundaries[EIso::IsoU].Max, Boundaries[EIso::IsoV].Min);
 	EndPoint.Set(Boundaries[EIso::IsoU].Max, Boundaries[EIso::IsoV].Max);
 	BuildEdge(StartPoint, EndPoint);
 
 	StartPoint.Set(Boundaries[EIso::IsoU].Max, Boundaries[EIso::IsoV].Max);
-	EndPoint.Set(Boundaries[EIso::IsoU].Max, Boundaries[EIso::IsoV].Min);
+	EndPoint.Set(Boundaries[EIso::IsoU].Min, Boundaries[EIso::IsoV].Max);
 	BuildEdge(StartPoint, EndPoint);
 
-	StartPoint.Set(Boundaries[EIso::IsoU].Max, Boundaries[EIso::IsoV].Min);
+	StartPoint.Set(Boundaries[EIso::IsoU].Min, Boundaries[EIso::IsoV].Max);
 	EndPoint.Set(Boundaries[EIso::IsoU].Min, Boundaries[EIso::IsoV].Min);
 	BuildEdge(StartPoint, EndPoint);
-
 	if (Edges.IsEmpty())
 	{
 		return;
@@ -462,6 +461,7 @@ void FTopologicalFace::SpawnIdent(FDatabase& Database)
 FInfoEntity& FTopologicalFace::GetInfo(FInfoEntity& Info) const
 {
 	return FTopologicalShapeEntity::GetInfo(Info)
+		.Add(TEXT("Orientation"), IsBackOriented() ? "Back" : "Front")
 		.Add(TEXT("Carrier Surface"), CarrierSurface)
 		.Add(TEXT("Boundary"), (FSurfacicBoundary&) Boundary)
 		.Add(TEXT("Loops"), Loops)
@@ -470,13 +470,13 @@ FInfoEntity& FTopologicalFace::GetInfo(FInfoEntity& Info) const
 }
 #endif
 
-TSharedRef<FFaceMesh> FTopologicalFace::GetOrCreateMesh(FModelMesh& MeshModel)
+FFaceMesh& FTopologicalFace::GetOrCreateMesh(FModelMesh& MeshModel)
 {
 	if (!Mesh.IsValid())
 	{
 		Mesh = FEntity::MakeShared<FFaceMesh>(MeshModel, *this);
 	}
-	return Mesh.ToSharedRef();
+	return *Mesh;
 }
 
 // Meshing parameters ==============================================================================================================================================================================================================================
diff --git a/Engine/Source/Runtime/Datasmith/CADKernel/Private/CADKernel/Topo/TopologicalLoop.cpp b/Engine/Source/Runtime/Datasmith/CADKernel/Private/CADKernel/Topo/TopologicalLoop.cpp
--- a/Engine/Source/Runtime/Datasmith/CADKernel/Private/CADKernel/Topo/TopologicalLoop.cpp
+++ b/Engine/Source/Runtime/Datasmith/CADKernel/Private/CADKernel/Topo/TopologicalLoop.cpp
@@ -1125,6 +1125,21 @@ bool FTopologicalLoop::IsInside(const FTopologicalLoop& OtherLoop) const
 	TArray<FPoint2D> OtherSampling2D;
 	OtherLoop.Get2DSampling(OtherSampling2D);
 
+#ifdef DEBUG_IS_INSIDE
+	bool bDisplayDebug = true; //(GetId() == 1046569);
+	if (bDisplayDebug)
+	{
+		{
+			F3DDebugSession _(*FString::Printf(TEXT("Loop")));
+			DisplayPolylineWithScale(Sampling2D, EVisuProperty::BlueCurve);
+		}
+		{
+			F3DDebugSession _(*FString::Printf(TEXT("External Loop")));
+			DisplayPolylineWithScale(OtherSampling2D, EVisuProperty::RedCurve);
+		}
+		Wait();
+	}
+#endif
 
 	int32 InsidePoint = 0;
 	int32 OutsidePoint = 0;
@@ -1169,7 +1184,7 @@ bool FTopologicalLoop::IsInside(const FTopologicalLoop& OtherLoop) const
 		CountLeftIntersection(TestPoint, EIso::IsoV);
 	}
 
-	return InsidePoint > OutsidePoint;
+	return InsidePoint < OutsidePoint;
 }
 
 double FTopologicalLoop::Length() const
diff --git a/Engine/Source/Runtime/Datasmith/CADKernel/Private/CADKernel/Topo/TopologicalVertex.cpp b/Engine/Source/Runtime/Datasmith/CADKernel/Private/CADKernel/Topo/TopologicalVertex.cpp
--- a/Engine/Source/Runtime/Datasmith/CADKernel/Private/CADKernel/Topo/TopologicalVertex.cpp
+++ b/Engine/Source/Runtime/Datasmith/CADKernel/Private/CADKernel/Topo/TopologicalVertex.cpp
@@ -170,14 +170,11 @@ void FVertexLink::DefineActiveEntity()
 				break;
 			}
 		}
-#ifdef CADKERNEL_DEV
-		ensureCADKernel(Square < 0.01);
-#endif
 	}
 	ActiveEntity = ClosedVertex;
 }
 
-TSharedRef<FVertexMesh> FTopologicalVertex::GetOrCreateMesh(FModelMesh& MeshModel)
+FVertexMesh& FTopologicalVertex::GetOrCreateMesh(FModelMesh& MeshModel)
 {
 	if (!IsActiveEntity())
 	{
@@ -190,10 +187,10 @@ TSharedRef<FVertexMesh> FTopologicalVertex::GetOrCreateMesh(FModelMesh& MeshMode
 
 		Mesh->GetNodeCoordinates().Emplace(GetBarycenter());
 		Mesh->RegisterCoordinates();
-		MeshModel.AddMesh(Mesh.ToSharedRef());
-		SetMeshed();
+		MeshModel.AddMesh(*Mesh);
+		SetMeshedMarker();
 	}
-	return Mesh.ToSharedRef();
+	return *Mesh;
 }
 
 void FTopologicalVertex::SpawnIdent(FDatabase& Database)
