CADKernel update: Meshing of thin faces

diff --git a/Engine/Source/Runtime/Datasmith/CADKernel/Private/CADKernel/Topo/Topomaker.cpp b/Engine/Source/Runtime/Datasmith/CADKernel/Private/CADKernel/Topo/Topomaker.cpp
--- a/Engine/Source/Runtime/Datasmith/CADKernel/Private/CADKernel/Topo/Topomaker.cpp
+++ b/Engine/Source/Runtime/Datasmith/CADKernel/Private/CADKernel/Topo/Topomaker.cpp
@@ -47,7 +47,7 @@ namespace TopomakerTools
  * Merge Border Vertices with other vertices.
  * @param Vertices: the initial array of active vertices to process, this array is updated at the end of the process
  */
-void MergeCoincidentVertices(TArray<TSharedPtr<FTopologicalVertex>>& VerticesToMerge, double Tolerance)
+void MergeCoincidentVertices(TArray<FTopologicalVertex*>& VerticesToMerge, double Tolerance)
 {
 	double SquareTolerance = FMath::Square(Tolerance);
 
@@ -63,7 +63,7 @@ void MergeCoincidentVertices(TArray<TSharedPtr<FTopologicalVertex>>& VerticesToM
 	TArray<int32> SortedVertexIndices;
 	SortedVertexIndices.Reserve(VertexNum);
 
-	for (TSharedPtr<FTopologicalVertex>& Vertex : VerticesToMerge)
+	for (FTopologicalVertex* Vertex : VerticesToMerge)
 	{
 		VerticesWeight.Add(Vertex->GetCoordinates().DiagonalAxisCoordinate());
 	}
@@ -76,7 +76,7 @@ void MergeCoincidentVertices(TArray<TSharedPtr<FTopologicalVertex>>& VerticesToM
 
 	for (int32 IndexI = 0; IndexI < VertexNum; ++IndexI)
 	{
-		TSharedPtr<FTopologicalVertex>& Vertex = VerticesToMerge[SortedVertexIndices[IndexI]];
+		FTopologicalVertex* Vertex = VerticesToMerge[SortedVertexIndices[IndexI]];
 		if (Vertex->HasMarker1())
 		{
 			continue;
@@ -91,7 +91,7 @@ void MergeCoincidentVertices(TArray<TSharedPtr<FTopologicalVertex>>& VerticesToM
 
 		for (int32 IndexJ = IndexI + 1; IndexJ < VertexNum; ++IndexJ)
 		{
-			TSharedPtr<FTopologicalVertex>& OtherVertex = VerticesToMerge[SortedVertexIndices[IndexJ]];
+			FTopologicalVertex* OtherVertex = VerticesToMerge[SortedVertexIndices[IndexJ]];
 			if (OtherVertex->HasMarker1())
 			{
 				continue;
@@ -113,26 +113,26 @@ void MergeCoincidentVertices(TArray<TSharedPtr<FTopologicalVertex>>& VerticesToM
 		}
 	}
 
-	for (TSharedPtr<FTopologicalVertex>& Vertex : VerticesToMerge)
+	for (FTopologicalVertex* Vertex : VerticesToMerge)
 	{
 		Vertex->ResetMarker1();
 	}
 
-	TArray<TSharedPtr<FTopologicalVertex>> ActiveVertices;
+	TArray<FTopologicalVertex*> ActiveVertices;
 	ActiveVertices.Reserve(VertexNum);
 
-	for (TSharedPtr<FTopologicalVertex>& Vertex : VerticesToMerge)
+	for (FTopologicalVertex* Vertex : VerticesToMerge)
 	{
-		TSharedPtr<FTopologicalVertex> ActiveVertex = Vertex->GetLinkActiveEntity();
-		if (ActiveVertex->HasMarker1())
+		FTopologicalVertex& ActiveVertex = *Vertex->GetLinkActiveEntity();
+		if (ActiveVertex.HasMarker1())
 		{
 			continue;
 		}
-		ActiveVertex->SetMarker1();
-		ActiveVertices.Add(ActiveVertex);
+		ActiveVertex.SetMarker1();
+		ActiveVertices.Add(&ActiveVertex);
 	}
 
-	for (TSharedPtr<FTopologicalVertex>& Vertex : ActiveVertices)
+	for (FTopologicalVertex* Vertex : ActiveVertices)
 	{
 		Vertex->ResetMarker1();
 	}
@@ -208,7 +208,7 @@ TSharedPtr<FTopologicalVertex> SplitAndLink(FTopologicalVertex& StartVertex, FTo
 	return NewVertex;
 }
 
-void StitchParallelEdges(TArray<TSharedPtr<FTopologicalVertex>>& VerticesToProcess, double SewTolerance, double MinEdgeLength)
+void StitchParallelEdges(TArray<FTopologicalVertex*>& VerticesToProcess, double SewTolerance, double MinEdgeLength)
 {
 	double SquareTolerance = FMath::Square(SewTolerance);
 	double SquareMinEdgeLength = FMath::Square(MinEdgeLength);
@@ -217,9 +217,9 @@ void StitchParallelEdges(TArray<TSharedPtr<FTopologicalVertex>>& VerticesToProce
 
 	for (int32 VertexI = 0; VertexI < VerticesToProcess.Num(); ++VertexI)
 	{
-		TSharedPtr<FTopologicalVertex> Vertex = VerticesToProcess[VertexI];
+		FTopologicalVertex* Vertex = VerticesToProcess[VertexI];
 
-		if (!Vertex.IsValid() || Vertex->IsDeleted() || !Vertex->IsBorderVertex())
+		if (Vertex->IsDeleted() || !Vertex->IsBorderVertex())
 		{
 			continue;
 		}
@@ -304,11 +304,11 @@ void StitchParallelEdges(TArray<TSharedPtr<FTopologicalVertex>>& VerticesToProce
 
 					if (NewVertex.IsValid() && !Vertex->IsLinkedTo(*NewVertex) && !OtherVertex->IsLinkedTo(*NewVertex))
 					{
-						VerticesToProcess.Add(NewVertex);
+						VerticesToProcess.Add(&*NewVertex);
 					}
 				}
 
-				if (!Vertex.IsValid() || Vertex->IsDeleted() || !Vertex->IsBorderVertex())
+				if (Vertex->IsDeleted() || !Vertex->IsBorderVertex())
 				{
 					break;
 				}
@@ -377,11 +377,11 @@ void MergeCoincidentEdges(TArray<FTopologicalEdge*>& EdgesToProcess, double MinE
 	FChrono::PrintClockElapse(EVerboseLevel::Log, TEXT("    "), TEXT("Merge coincident edges"), StepDuration);
 }
 
-void MergeCoincidentEdges(TArray<TSharedPtr<FTopologicalVertex>>& VerticesToProcess, double MinEdgeLength)
+void MergeCoincidentEdges(TArray<FTopologicalVertex*>& VerticesToProcess, double MinEdgeLength)
 {
 	FTimePoint StartTime = FChrono::Now();
 
-	for (TSharedPtr<FTopologicalVertex>& Vertex : VerticesToProcess)
+	for (FTopologicalVertex* Vertex : VerticesToProcess)
 	{
 		if (Vertex->IsDeleted() || !Vertex->IsActiveEntity() || Vertex->HasMarker1())
 		{
@@ -409,7 +409,7 @@ void MergeCoincidentEdges(TArray<TSharedPtr<FTopologicalVertex>>& VerticesToProc
 		}
 	}
 
-	for (TSharedPtr<FTopologicalVertex>& Vertex : VerticesToProcess)
+	for (FTopologicalVertex* Vertex : VerticesToProcess)
 	{
 		Vertex->ResetMarker1();
 	}
@@ -507,7 +507,7 @@ void FTopomaker::RemoveFacesFromShell()
 
 	for (FShell* Shell : ShellSet)
 	{
-		bool bIsOutter = Shell->IsOutter();
+		const bool bIsOuter = Shell->IsOuter();
 
 		TArray<FOrientedFace> ShellFace;
 		ShellFace.Reserve(Shell->FaceCount());
@@ -515,7 +515,7 @@ void FTopomaker::RemoveFacesFromShell()
 		{
 			if (!Face.Entity->GetHost())
 			{
-				if (bIsOutter != (Face.Direction == EOrientation::Front))
+				if (bIsOuter != (Face.Direction == EOrientation::Front))
 				{
 					Face.Entity->SetBackOriented();
 				}
@@ -541,7 +541,7 @@ void FTopomaker::Sew()
 {
 	FTimePoint StartJoinTime = FChrono::Now();
 
-	TArray<TSharedPtr<FTopologicalVertex>> BorderVertices;
+	TArray<FTopologicalVertex*> BorderVertices;
 	GetBorderVertices(BorderVertices);
 	TopomakerTools::MergeCoincidentVertices(BorderVertices, SewTolerance);
 
@@ -586,13 +586,13 @@ void FTopomaker::Sew()
 #endif
 }
 
-void FTopomaker::GetVertices(TArray<TSharedPtr<FTopologicalVertex>>& Vertices)
+void FTopomaker::GetVertices(TArray<FTopologicalVertex*>& Vertices)
 {
 	Vertices.Empty(10 * Faces.Num());
 
 	for (TSharedPtr<FTopologicalFace> Face : Faces)
 	{
-		if (Face->IsDeleted())
+		if (Face->IsDeletedOrDegenerated())
 		{
 			continue;
 		}
@@ -605,30 +605,30 @@ void FTopomaker::GetVertices(TArray<TSharedPtr<FTopologicalVertex>>& Vertices)
 				if (!Edge->GetStartVertex()->GetLinkActiveEntity()->HasMarker1())
 				{
 					Edge->GetStartVertex()->GetLinkActiveEntity()->SetMarker1();
-					Vertices.Add(Edge->GetStartVertex()->GetLinkActiveEntity());
+					Vertices.Add(&*Edge->GetStartVertex()->GetLinkActiveEntity());
 				}
 				if (!Edge->GetEndVertex()->GetLinkActiveEntity()->HasMarker1())
 				{
 					Edge->GetEndVertex()->GetLinkActiveEntity()->SetMarker1();
-					Vertices.Add(Edge->GetEndVertex()->GetLinkActiveEntity());
+					Vertices.Add(&*Edge->GetEndVertex()->GetLinkActiveEntity());
 				}
 			}
 		}
 	}
-	for (TSharedPtr<FTopologicalVertex>& Vertex : Vertices)
+	for (FTopologicalVertex* Vertex : Vertices)
 	{
 		Vertex->ResetMarker1();
 	}
 }
 
-void FTopomaker::GetBorderVertices(TArray<TSharedPtr<FTopologicalVertex>>& BorderVertices)
+void FTopomaker::GetBorderVertices(TArray<FTopologicalVertex*>& BorderVertices)
 {
-	TArray<TSharedPtr<FTopologicalVertex>> Vertices;
+	TArray<FTopologicalVertex*> Vertices;
 	GetVertices(Vertices);
 
 	BorderVertices.Empty(Vertices.Num());
 
-	for (TSharedPtr<FTopologicalVertex>& Vertex : Vertices)
+	for (FTopologicalVertex* Vertex : Vertices)
 	{
 		if (Vertex->IsBorderVertex())
 		{
@@ -744,10 +744,10 @@ void FTopomaker::RemoveIsolatedEdges()
 
 	TArray<FTopologicalEdge*> IsolatedEdges;
 
-	TArray<TSharedPtr<FTopologicalVertex>> Vertices;
+	TArray<FTopologicalVertex*> Vertices;
 	GetVertices(Vertices);
 
-	for (const TSharedPtr<FTopologicalVertex>& Vertex : Vertices)
+	for (const FTopologicalVertex* Vertex : Vertices)
 	{
 		for (const FTopologicalVertex* TwinVertex : Vertex->GetTwinEntities())
 		{
@@ -768,61 +768,151 @@ void FTopomaker::RemoveIsolatedEdges()
 	FMessage::Printf(EVerboseLevel::Log, TEXT("\n\nIsolatedEdges count %d\n\n\n"), IsolatedEdges.Num());
 }
 
-void FTopomaker::UnlinkNonManifoldVertex()
+void SplitVertexLinkByShell(FTopologicalVertex* Vertex)
 {
-	TArray<TSharedPtr<FTopologicalVertex>> Vertices;
-	GetVertices(Vertices);
+	TMap<const FShell*, TArray<FTopologicalVertex*>> ShellToVertices;
 
-	TMap<FShell*, int32> ShellToVertexCount;
-	for (const TSharedPtr<FTopologicalVertex>& Vertex : Vertices)
+	const TArray<FTopologicalVertex*>& Twins = Vertex->GetTwinEntities();
+	ShellToVertices.Reserve(Twins.Num());
+
+	for (FTopologicalVertex* TwinVertex : Twins)
 	{
-		ShellToVertexCount.Empty(Vertex->GetTwinEntityCount());
-		for (const FTopologicalVertex* TwinVertex : Vertex->GetTwinEntities())
+		const FShell* OtherShell = (const FShell*)TwinVertex->GetDirectConnectedEdges()[0]->GetLoop()->GetFace()->GetHost();
+		TArray<FTopologicalVertex*>& Vertices = ShellToVertices.FindOrAdd(OtherShell);
+		Vertices.Add(TwinVertex);
+	}
+
+	Vertex->UnlinkTwinEntities();
+
+	for (TPair<const FShell*, TArray<FTopologicalVertex*>>& Pair : ShellToVertices)
+	{
+		TArray<FTopologicalVertex*>& Vertices = Pair.Value;
+		if (Vertices.Num() > 1)
 		{
-			ensureCADKernel(!TwinVertex->GetDirectConnectedEdges().IsEmpty());
-			FTopologicalEdge* Edge = TwinVertex->GetDirectConnectedEdges()[0];
-			FShell* Shell = (FShell*)Edge->GetLoop()->GetFace()->GetHost();
-			if (Shell != nullptr)
+			FTopologicalVertex* FirstVertex = Vertices[0];
+			for (FTopologicalVertex* TwinVertex : Vertices)
 			{
-				ShellToVertexCount.FindOrAdd(Shell)++;
+				FirstVertex->Link(*TwinVertex);
 			}
 		}
+	}
+}
+
+void SplitEdgeLinkByShell(FTopologicalEdge& Edge)
+{
+	TMap<const FShell*, TArray<FTopologicalEdge*>> ShellToEdges;
+
+	const TArray<FTopologicalEdge*>& Twins = Edge.GetTwinEntities();
+	ShellToEdges.Reserve(Twins.Num());
 
-		if (ShellToVertexCount.Num() > 1)
+	for (FTopologicalEdge* TwinEdge : Twins)
+	{
+		const FShell* OtherShell = (const FShell*)TwinEdge->GetLoop()->GetFace()->GetHost();
+		TArray<FTopologicalEdge*>& Edges = ShellToEdges.FindOrAdd(OtherShell);
+		Edges.Add(TwinEdge);
+	}
+
+	Edge.UnlinkTwinEntities();
+
+	for (TPair<const FShell*, TArray<FTopologicalEdge*>>& Pair : ShellToEdges)
+	{
+		TArray<FTopologicalEdge*>& Edges = Pair.Value;
+		if (Edges.Num() > 1)
 		{
-			TMap<FShell*, TArray<FTopologicalVertex*>> ShellToVertices;
-			ShellToVertices.Reserve(ShellToVertexCount.Num());
-			for (TPair<FShell*, int32> Pair : ShellToVertexCount)
+			FTopologicalEdge* FirstEdge = Edges[0];
+			for (FTopologicalEdge* TwinEdge : Edges)
 			{
-				ShellToVertices.FindOrAdd(Pair.Key).Reserve(Pair.Value);
+				FirstEdge->Link(*TwinEdge);
 			}
+		}
+	}
+}
 
-			for (FTopologicalVertex* TwinVertex : Vertex->GetTwinEntities())
+void FTopomaker::UnlinkEdgesOfSeparateShells()
+{
+	// Mandatory: UnlinkNonManifoldVertex has to be call before
+	TArray<FTopologicalEdge*> Edges;
+	Edges.Reserve(Faces.Num() * 10);
+	for (TSharedPtr<FTopologicalFace> Face : Faces)
+	{
+		if (Face->IsDeletedOrDegenerated())
+		{
+			continue;
+		}
+
+		for (TSharedPtr<FTopologicalLoop> Loop : Face->GetLoops())
+		{
+			for (const FOrientedEdge& OrientedEdge : Loop->GetEdges())
 			{
-				ensureCADKernel(!TwinVertex->GetDirectConnectedEdges().IsEmpty());
-				FTopologicalEdge* Edge = TwinVertex->GetDirectConnectedEdges()[0];
-				FShell* Shell = (FShell*)Edge->GetLoop()->GetFace()->GetHost();
-				if (Shell != nullptr)
+				FTopologicalEdge& Edge = *OrientedEdge.Entity->GetLinkActiveEdge();
+				if (!Edge.IsProcessed())
 				{
-					ShellToVertices.FindOrAdd(Shell).Add(TwinVertex);
+					bool bNeedToSplit = false;
+					Edge.SetProcessedMarker();
+					Edges.Add(&Edge);
+
+					const FShell* Shell = (FShell*)Edge.GetLoop()->GetFace()->GetHost();
+					for (FTopologicalEdge* TwinEdge : Edge.GetTwinEntities())
+					{
+						const FShell* TwinEdgeShell = (FShell*)TwinEdge->GetLoop()->GetFace()->GetHost();
+						if (Shell != TwinEdgeShell)
+						{
+							SplitEdgeLinkByShell(Edge);
+							break;
+						}
+					}
 				}
 			}
+		}
+	}
+	for (FTopologicalEdge* Edge : Edges)
+	{
+		Edge->ResetMarker1();
+	}
+}
 
-			Vertex->UnlinkTwinEntities();
-			for (TPair<FShell*, TArray<FTopologicalVertex*>> ShellVertices : ShellToVertices)
+
+void FTopomaker::UnlinkNonManifoldVertex()
+{
+	TArray<FTopologicalVertex*> Vertices;
+	GetVertices(Vertices);
+
+	// check that each vertex belongs to only one shell
+	for (const FTopologicalVertex* Vertex : Vertices)
+	{
+		for (const FTopologicalVertex* TwinVertex : Vertex->GetTwinEntities())
+		{
+			const FShell* Shell = nullptr;
+			{
+				FTopologicalEdge* Edge = TwinVertex->GetDirectConnectedEdges()[0];
+				Shell = (FShell*)Edge->GetLoop()->GetFace()->GetHost();
+			}
+			if (Shell != nullptr)
 			{
-				if (ShellVertices.Value.Num() > 1)
+				for (FTopologicalEdge* Edge : TwinVertex->GetDirectConnectedEdges())
 				{
-					TArray<FTopologicalVertex*>& Twins = ShellVertices.Value;
-					FTopologicalVertex* FirstVertex = Twins[0];
-					for (FTopologicalVertex* TwinVertex : Twins)
-					{
-						FirstVertex->Link(*TwinVertex);
-					}
+					ensureCADKernel(Edge->GetLoop()->GetFace()->GetHost() == Shell);
 				}
 			}
 		}
 	}
+
+	for (FTopologicalVertex* Vertex : Vertices)
+	{
+		const TArray<FTopologicalVertex*>& Twins = Vertex->GetTwinEntities();
+		const FTopologicalEdge* Edge = Twins[0]->GetDirectConnectedEdges()[0];
+		const FShell* Shell = (const FShell*)Edge->GetLoop()->GetFace()->GetHost();
+
+		for (const FTopologicalVertex* TwinVertex : Twins)
+		{
+			const FShell* OtherShell = (const FShell*)TwinVertex->GetDirectConnectedEdges()[0]->GetLoop()->GetFace()->GetHost();
+			if (OtherShell != Shell)
+			{
+				SplitVertexLinkByShell(Vertex);
+				break;
+			}
+		}
+	}
 }
 
 void FindBorderVertex(const TArray<FTopologicalEdge*>& Edges, TArray<FTopologicalVertex*>& OutBorderVertices)
@@ -909,9 +999,9 @@ void FTopomaker::RemoveThinFaces(TArray<FTopologicalEdge*>& NewBorderEdges)
 #endif
 }
 
-void FTopomaker::SetSelfConnectedEdgeDegenerated(TArray<TSharedPtr<FTopologicalVertex>>& Vertices)
+void FTopomaker::SetSelfConnectedEdgeDegenerated(TArray<FTopologicalVertex*>& Vertices)
 {
-	for (const TSharedPtr<FTopologicalVertex>& Vertex : Vertices)
+	for (const FTopologicalVertex* Vertex : Vertices)
 	{
 		for (const FTopologicalVertex* TwinVertex : Vertex->GetTwinEntities())
 		{
@@ -929,16 +1019,16 @@ void FTopomaker::SetSelfConnectedEdgeDegenerated(TArray<TSharedPtr<FTopologicalV
 	}
 }
 
-void FTopomaker::CheckSelfConnectedEdge(double MaxLengthOfDegeneratedEdge, TArray<TSharedPtr<FTopologicalVertex>>& OutBorderVertices)
+void FTopomaker::CheckSelfConnectedEdge(double MaxLengthOfDegeneratedEdge, TArray<FTopologicalVertex*>& OutBorderVertices)
 {
 	FTimePoint StartTime = FChrono::Now();
 
 	TFunction<void(FTopologicalVertex&)> AddVertexIfBorder = [&](FTopologicalVertex& Vertex)
 	{
-		TSharedPtr<FTopologicalVertex> ActiveVertex = Vertex.GetLinkActiveEntity();
+		TSharedRef<FTopologicalVertex> ActiveVertex = Vertex.GetLinkActiveEntity();
 		if (ActiveVertex->IsBorderVertex())
 		{
-			OutBorderVertices.Add(ActiveVertex);
+			OutBorderVertices.Add(&*ActiveVertex);
 		}
 	};
 
@@ -996,11 +1086,6 @@ void FTopomaker::SplitIntoConnectedShells()
 			for (const FOrientedEdge& OrientedEdge : Loop->GetEdges())
 			{
 				const TSharedPtr<FTopologicalEdge>& Edge = OrientedEdge.Entity;
-				if (Edge->HasMarker1())
-				{
-					continue;
-				}
-				Edge->SetMarker1();
 
 				if (Edge->GetTwinEntityCount() == 1)
 				{
@@ -1018,19 +1103,12 @@ void FTopomaker::SplitIntoConnectedShells()
 
 				for (FTopologicalEdge* NextEdge : Edge->GetTwinEntities())
 				{
-					if (NextEdge->HasMarker1())
-					{
-						continue;
-					}
-					NextEdge->SetMarker1();
-
 					FTopologicalFace* NextFace = NextEdge->GetFace();
-					if ((NextFace == nullptr) || NextFace->HasMarker1())
+					if ((NextFace == nullptr) || NextFace->IsNotToOrAlreadyProcess())
 					{
 						continue;
 					}
-
-					NextFace->SetMarker1();
+					NextFace->SetProcessedMarker();
 					Front.Add(NextFace);
 				}
 			}
@@ -1052,18 +1130,11 @@ void FTopomaker::SplitIntoConnectedShells()
 		}
 	};
 
-	for (const TSharedPtr<FTopologicalFace>& Face : Faces)
-	{
-		if (Face->IsDeleted() || Face->IsDegenerated())
-		{
-			Face->SetMarker1();
-			continue;
-		}
-	}
+	SetToProcessMarkerOfFaces();	
 		
 	for (const TSharedPtr<FTopologicalFace>& Face : Faces)
 	{
-		if (Face->HasMarker1())
+		if (Face->IsProcessedDeletedOrDegenerated())
 		{
 			continue;
 		}
@@ -1072,7 +1143,7 @@ void FTopomaker::SplitIntoConnectedShells()
 		Shell.Faces.Reserve(TopologicalFaceCount - ProcessFaceCount);
 		Front.Empty(TopologicalFaceCount);
 
-		Face->SetMarker1();
+		Face->SetProcessedMarker();
 		Front.Add(Face.Get());
 		SpreadFront(Shell);
 		ProcessFaceCount += Shell.Faces.Num();
@@ -1083,18 +1154,7 @@ void FTopomaker::SplitIntoConnectedShells()
 		}
 	}
 
-	// reset Marker
-	for (const TSharedPtr<FTopologicalFace>& Face : Faces)
-	{
-		Face->ResetMarkers();
-		for (const TSharedPtr<FTopologicalLoop>& Loop : Face->GetLoops())
-		{
-			for (const FOrientedEdge& OrientedEdge : Loop->GetEdges())
-			{
-				OrientedEdge.Entity->ResetMarkers();
-			}
-		}
-	}
+	ResetMarkersOfFaces();
 
 	// for each FaceSubset, find the main shell
 	for (FFaceSubset& FaceSubset : SubShells)
@@ -1136,7 +1196,9 @@ void FTopomaker::SplitIntoConnectedShells()
 		if (FaceSubset.MainShell != nullptr)
 		{
 			FShell* Shell = (FShell*)FaceSubset.MainShell;
+#ifdef CADKERNEL_DEBUG
 			ensure(Shell->GetFaces().Num() == 0);
+#endif
 			Shell->Add(FaceSubset.Faces);
 		}
 		else
@@ -1168,6 +1230,7 @@ void FTopomaker::SplitIntoConnectedShells()
 	}
 
 	UnlinkNonManifoldVertex();
+	UnlinkEdgesOfSeparateShells();
 
 	RemoveEmptyShells();
 
@@ -1345,4 +1408,23 @@ void FTopomaker::RemoveDuplicatedFaces()
 
 }
 
+void FTopomaker::SetToProcessMarkerOfFaces()
+{
+	for (const TSharedPtr<FTopologicalFace>& Face : Faces)
+	{
+		if (!Face->IsDeletedOrDegenerated())
+		{
+			Face->SetToProcessMarker();
+		}
+	}
+}
+
+void FTopomaker::ResetMarkersOfFaces()
+{
+	for (const TSharedPtr<FTopologicalFace>& Face : Faces)
+	{
+		Face->ResetMarkers();
+	}
+}
+
 }
diff --git a/Engine/Source/Runtime/Datasmith/CADKernel/Private/CADKernel/UI/Display.cpp b/Engine/Source/Runtime/Datasmith/CADKernel/Private/CADKernel/UI/Display.cpp
--- a/Engine/Source/Runtime/Datasmith/CADKernel/Private/CADKernel/UI/Display.cpp
+++ b/Engine/Source/Runtime/Datasmith/CADKernel/Private/CADKernel/UI/Display.cpp
@@ -5,6 +5,7 @@
 #include "CADKernel/Core/Chrono.h"
 #include "CADKernel/Core/Group.h"
 #include "CADKernel/Core/System.h"
+#include "CADKernel/Core/OrientedEntity.h"
 #include "CADKernel/Geo/Curves/RestrictionCurve.h"
 #include "CADKernel/Geo/Curves/BezierCurve.h"
 #include "CADKernel/Geo/Curves/NURBSCurve.h"
@@ -120,7 +121,7 @@ void Draw(const FLinearBoundary& Boundary, const FRestrictionCurve& Curve, EVisu
 		double Height = Length / 20.0;
 		double Base = Height / 2;
 
-		DrawQuadripode(Height, Base, Point.Point, Point.Gradient);
+		DrawQuadripode(Height, Base, Point.Point, Point.Gradient, Property);
 	}
 #endif
 }
@@ -153,7 +154,7 @@ void Draw2D(const FCurve& Curve, const FLinearBoundary& Boundary, EVisuProperty
 		double Height = Length / 20.0;
 		double Base = Height / 2;
 
-		DrawQuadripode(Height, Base, Point.Point, Point.Gradient);
+		DrawQuadripode(Height, Base, Point.Point, Point.Gradient, Property);
 	}
 #endif
 }
@@ -186,7 +187,7 @@ void Draw3D(const FCurve& Curve, const FLinearBoundary& Boundary, EVisuProperty
 		double Height = Length / 20.0;
 		double Base = Height / 2;
 
-		DrawQuadripode(Height, Base, Point.Point, Point.Gradient);
+		DrawQuadripode(Height, Base, Point.Point, Point.Gradient, Property);
 	}
 #endif
 }
@@ -214,7 +215,7 @@ void Draw(const FCurve& Curve, EVisuProperty Property)
 #endif
 }
 
-void DrawQuadripode(double Height, double Base, FPoint& Center, FPoint& InDirection)
+void DrawQuadripode(double Height, double Base, FPoint& Center, FPoint& InDirection, EVisuProperty Property)
 {
 #ifdef CADKERNEL_DEBUG
 	FPoint Direction = InDirection;
@@ -249,7 +250,7 @@ void DrawQuadripode(double Height, double Base, FPoint& Center, FPoint& InDirect
 	Polygone.Add(Point0);
 	Polygone.Add(Point4);
 
-	Draw(Polygone);
+	Draw(Polygone, Property);
 #endif
 }
 
@@ -295,9 +296,6 @@ void DisplayEntity(const FEntity& Entity)
 	case EEntity::Model:
 		Display((const FModel&)Entity);
 		break;
-	case EEntity::Group:
-		Display((const FGroup&)Entity);
-		break;
 	case EEntity::MeshModel:
 		Display((const FModelMesh&)Entity);
 		break;
@@ -697,7 +695,9 @@ void DisplayControlPolygon(const FCurve& Curve)
 void DisplayControlPolygon(const FSurface& Surface)
 {
 #ifdef CADKERNEL_DEBUG
-	TFunction<void(const TArray<FPoint>&, int32, int32)> DisplayHull = [](const TArray<FPoint>& Poles, int32 PoleUNum, int32 PoleVNum)
+	bool bShowOrientation = FSystem::Get().GetVisu()->GetParameters()->bDisplayCADOrient;
+
+	TFunction<void(const TArray<FPoint>&, int32, int32)> DisplayHull = [&bShowOrientation](const TArray<FPoint>& Poles, int32 PoleUNum, int32 PoleVNum)
 	{
 		for (int32 Index = 0; Index < Poles.Num(); Index++)
 		{
@@ -710,7 +710,11 @@ void DisplayControlPolygon(const FSurface& Surface)
 			for (int32 IndexU = 1; IndexU < PoleUNum; IndexU++, Index++)
 			{
 				FPoint Segment = Poles[Index] - Poles[Index - 1];
-				DisplaySegment(Poles[Index - 1] + Segment * 0.1, Poles[Index] - Segment * 0.1, 0, EVisuProperty::GreenCurve);
+				DisplaySegment(Poles[Index - 1] + Segment * 0.1, Poles[Index] - Segment * 0.1, 0, EVisuProperty::YellowCurve);
+				if (bShowOrientation)
+				{
+					DrawSegmentOrientation(Poles[Index - 1], Poles[Index], EVisuProperty::YellowCurve);
+				}
 			}
 		}
 
@@ -720,6 +724,10 @@ void DisplayControlPolygon(const FSurface& Surface)
 			{
 				FPoint Segment = Poles[Index] - Poles[Index - PoleUNum];
 				DisplaySegment(Poles[Index - PoleUNum] + Segment * 0.1, Poles[Index] - Segment * 0.1, 0, EVisuProperty::GreenCurve);
+				if (bShowOrientation)
+				{
+					DrawSegmentOrientation(Poles[Index - PoleUNum], Poles[Index], EVisuProperty::GreenCurve);
+				}
 			}
 		}
 	};
@@ -946,11 +954,6 @@ void Display2DWithScale(const FTopologicalEdge& Edge, EVisuProperty Property)
 	TArray<FPoint2D> Polyline;
 	Edge.GetCurve()->GetDiscretizationPoints(Edge.GetBoundary(), EOrientation::Front, Polyline);
 	DisplayPolylineWithScale(Polyline, Property);
-
-	for (const FPoint2D& Point : Polyline)
-	{
-		DisplayPoint2DWithScale(Point, EVisuProperty ::RedPoint);
-	}
 #endif
 }
 
@@ -1072,6 +1075,44 @@ void Display(const FGroup& Group)
 #endif
 }
 
+void DisplayMesh(const FModel& Model)
+{
+	for (const TSharedPtr<FBody>& Body : Model.GetBodies())
+	{
+		if(Body)
+		{
+			DisplayMesh(*Body);
+		}
+	}
+}
+
+void DisplayMesh(const FBody& Body)
+{
+	for (const TSharedPtr<FShell>& Shell : Body.GetShells())
+	{
+		if (Shell)
+		{
+			DisplayMesh(*Shell);
+		}
+	}
+}
+
+void DisplayMesh(const FShell& Shell)
+{
+	for (const FOrientedFace& OrientedFace : Shell.GetFaces())
+	{
+		if (OrientedFace.Entity)
+		{
+			FTopologicalFace& Face = *OrientedFace.Entity;
+			if (!Face.IsDeletedOrDegenerated() && Face.IsMeshed())
+			{
+				FFaceMesh& Mesh = *Face.GetMesh();
+				DisplayMesh(Mesh);
+			}
+		}
+	}
+}
+
 void DisplayMesh(const FFaceMesh& Mesh)
 {
 #ifdef CADKERNEL_DEBUG
@@ -1203,9 +1244,9 @@ void DisplayMesh(const FVertexMesh& Mesh)
 void Display(const FModelMesh& MeshModel)
 {
 #ifdef CADKERNEL_DEBUG
-	for (const TSharedPtr<FFaceMesh>& Mesh : MeshModel.GetFaceMeshes())
+	for (const FFaceMesh* Mesh : MeshModel.GetFaceMeshes())
 	{
-		if (Mesh.IsValid())
+		if (Mesh)
 		{
 			DisplayMesh(*Mesh);
 		}
diff --git a/Engine/Source/Runtime/Datasmith/CADKernel/Public/CADKernel/Core/Database.h b/Engine/Source/Runtime/Datasmith/CADKernel/Public/CADKernel/Core/Database.h
--- a/Engine/Source/Runtime/Datasmith/CADKernel/Public/CADKernel/Core/Database.h
+++ b/Engine/Source/Runtime/Datasmith/CADKernel/Public/CADKernel/Core/Database.h
@@ -292,6 +292,7 @@ public:
 
 	TSharedPtr<FEntity> GetEntity(FIdent id) const;
 	void GetEntities(const TArray<FIdent>& ids, TArray<TSharedPtr<FEntity>>& Entities) const;
+	void GetEntities(const TArray<FIdent>& ids, TArray<FEntity*>& Entities) const;
 	void GetTopologicalEntities(const TArray<FIdent>& ids, TArray<FTopologicalEntity*>& Entities) const;
 	void GetTopologicalShapeEntities(const TArray<FIdent>& ids, TArray<FTopologicalShapeEntity*>& Entities) const;
 
diff --git a/Engine/Source/Runtime/Datasmith/CADKernel/Public/CADKernel/Core/Entity.h b/Engine/Source/Runtime/Datasmith/CADKernel/Public/CADKernel/Core/Entity.h
--- a/Engine/Source/Runtime/Datasmith/CADKernel/Public/CADKernel/Core/Entity.h
+++ b/Engine/Source/Runtime/Datasmith/CADKernel/Public/CADKernel/Core/Entity.h
@@ -16,18 +16,22 @@ namespace UE::CADKernel
 {
 enum class EEntity : uint8
 {
-	NullEntity = 0,
+	None = 0,
 
 	Curve,
 	Surface,
 
 	EdgeLink,
 	VertexLink,
+
 	TopologicalEdge,
 	TopologicalFace,
+
 	TopologicalLink,
 	TopologicalLoop,
+
 	TopologicalVertex,
+
 	Shell,
 	Body,
 	Model,
@@ -39,7 +43,7 @@ enum class EEntity : uint8
 	Criterion,
 	Property,
 
-	EntityTypeEnd
+	EntityTypeEnd,
 };
 
 class FDatabase;
@@ -145,6 +149,17 @@ public:
 		SerializeIdents(Ar, (TArray<TSharedPtr<FEntity>>&) Array, bSaveSelection);
 	}
 
+	/**
+	 * Spawn Ident on each FEntity* of the array
+	 */
+	static void SpawnIdentOnEntities(TArray<FEntity*>& Array, FDatabase& Database)
+	{
+		for (FEntity* Entity : Array)
+		{
+			Entity->SpawnIdent(Database);
+		}
+	}
+
 	/**
 	 * Spawn Ident on each TSharedPtr<FEntity> of the array
 	 */
@@ -173,34 +188,42 @@ public:
 	/**
 	 * Reset Processed flags Recursively on each TSharedPtr<FEntity> of the array
 	 */
-	static void ResetMarkersRecursivelyOnEntities(TArray<TWeakPtr<FEntity>>& Array)
+	static void ResetMarkersRecursivelyOnEntities(const TArray<FEntity*>& Array)
 	{
-		for (TWeakPtr<FEntity>& Entity : Array)
+		for (const FEntity* Entity : Array)
+		{
+			Entity->ResetMarkersRecursively();
+		}
+	}
+
+	static void ResetMarkersRecursivelyOnEntities(const TArray<TWeakPtr<FEntity>>& Array)
+	{
+		for (const TWeakPtr<FEntity>& Entity : Array)
 		{
 			Entity.Pin()->ResetMarkersRecursively();
 		}
 	}
 
-	static void ResetMarkersRecursivelyOnEntities(TArray<TSharedPtr<FEntity>>& Array)
+	static void ResetMarkersRecursivelyOnEntities(const TArray<TSharedPtr<FEntity>>& Array)
 	{
-		for (TSharedPtr<FEntity>& Entity : Array)
+		for (const TSharedPtr<FEntity>& Entity : Array)
 		{
 			Entity->ResetMarkersRecursively();
 		}
 	}
 
-	static void ResetMarkersRecursivelyOnEntities(TArray<TOrientedEntity<FEntity>>& Array)
+	static void ResetMarkersRecursivelyOnEntities(const TArray<TOrientedEntity<FEntity>>& Array)
 	{
-		for (TOrientedEntity<FEntity>& OrientedEntity : Array)
+		for (const TOrientedEntity<FEntity>& OrientedEntity : Array)
 		{
 			OrientedEntity.Entity->ResetMarkersRecursively();
 		}
 	}
 
 	template<typename EntityType>
-	static void ResetMarkersRecursivelyOnEntities(TArray<TSharedPtr<EntityType>>& Array)
+	static void ResetMarkersRecursivelyOnEntities(const TArray<TSharedPtr<EntityType>>& Array)
 	{
-		ResetMarkersRecursivelyOnEntities((TArray<TSharedPtr<FEntity>>&) Array);
+		ResetMarkersRecursivelyOnEntities((const TArray<TSharedPtr<FEntity>>&) Array);
 	}
 
 	/**
@@ -229,6 +252,22 @@ public:
 	static const TCHAR* GetTypeName(EEntity Type);
 
 	virtual EEntity GetEntityType() const = 0;
+
+	bool IsTopologicalEntity() const
+	{
+		return (GetEntityType() >= EEntity::EdgeLink && GetEntityType() <= EEntity::Model);
+	}
+
+	bool IsTopologicalShapeEntity() const
+	{
+		return (GetEntityType() == EEntity::TopologicalFace) || (GetEntityType() >= EEntity::Shell && GetEntityType() <= EEntity::Model);
+	}
+
+	bool IsGeometricalEntity()
+	{
+		return (GetEntityType() == EEntity::Curve || GetEntityType() == EEntity::Surface);
+	}
+
 	const TCHAR* GetTypeName() const
 	{
 		return GetTypeName(GetEntityType());
@@ -253,7 +292,7 @@ public:
 		}
 	}
 
-	virtual void ResetMarkersRecursively()
+	virtual void ResetMarkersRecursively() const
 	{
 		ResetMarkers();
 	}
diff --git a/Engine/Source/Runtime/Datasmith/CADKernel/Public/CADKernel/Core/HaveStates.h b/Engine/Source/Runtime/Datasmith/CADKernel/Public/CADKernel/Core/HaveStates.h
--- a/Engine/Source/Runtime/Datasmith/CADKernel/Public/CADKernel/Core/HaveStates.h
+++ b/Engine/Source/Runtime/Datasmith/CADKernel/Public/CADKernel/Core/HaveStates.h
@@ -6,34 +6,38 @@
 
 enum class EHaveStates : uint16
 {
-	None = 0x0000u,  // No flags.
+	None              = 0x0000u,  // No flags.
+					  
+	IsApplyCriteria   = 0x0001u,  // used for FTopologicalFace, FEdge
+	IsMeshed          = 0x0002u,  // used for FTopologicalFace, FEdge
+	IsPreMeshed       = 0x0004u,  // used for FTopologicalFace, FEdge
+					  
+	ThinZone          = 0x0008u,  // used for FTopologicalFace, FEdge, FEdgeSegment  
+					  
+	Degenerated       = 0x0010u,  // used for FEdge, FGrid
+				      
+	IsDeleted         = 0x0020u,  // used for all class
 
-	IsApplyCriteria = 0x0001u,  // used for FTopologicalFace, FEdge
-	IsRemoved = 0x0001u,  // used for FThinZone2D
-
-	IsMeshed = 0x0002u,  // used for FTopologicalFace, FEdge
-	IsInner = 0x0002u,  // used for FEdgeSegment, FShell, 
-	IsFirstSide = 0x0002u,  // used for FThinZoneSide
-	FirstSideClosed = 0x0002u,  // used for FThinZone2D
-
-	IsBackOriented = 0x0004u,  // used for FTopologicalFace
-	ThinPeak = 0x0004u,  // used for FEdge 
-	SecondSideClosed = 0x0004u,  // used for FThinZone2D
+	IsVirtuallyMeshed = 0x0040u,  // used for FEdge
 
-	ThinZone = 0x0008u,  // used for FTopologicalFace, FEdge, FEdgeSegment  
+	ThinPeak          = 0x0080u,  // used for FEdge 
 
-	Degenerated = 0x0010u,  // used for FEdge, FGrid
+	IsBackOriented    = 0x0100u,  // used for FTopologicalFace
+	IsInner           = 0x0100u,  // used for FEdgeSegment, FShell, 
 
-	IsDeleted = 0x0020u,  // used for all class
-	IsVirtuallyMeshed = 0x0040u,  // used for FEdge
+	FirstSideClosed   = 0x0200u,  // used for FThinZone2D
+	SecondSideClosed  = 0x0400u,  // used for FThinZone2D
 
-	HasMarker1 = 0x1000u,  // used for all class
-	HasMarker2 = 0x2000u,  // used for all class
-	HasMarker3 = 0x4000u,  // used for all class
-	IsProcess4 = 0x8000u,  // used for all class
-	AllMarkers = 0xF000u,  // used for all class
+	HasMarker1        = 0x0800u,  // used for all class
+	HasMarker2        = 0x1000u,  // used for all class
+				      
+	ToProcess         = 0x2000u,  // used for all class
+	IsWaiting         = 0x4000u,  // used for all class
+	IsProcessed       = 0x8000u,  // used for all class
+				      
+	AllMarkers        = 0xF800u,  // used for all class
 
-	All = 0xFFu
+	All = 0xFFFFu
 };
 
 ENUM_CLASS_FLAGS(EHaveStates);
@@ -71,9 +75,21 @@ public:
 		return ((States & EHaveStates::HasMarker2) == EHaveStates::HasMarker2);
 	}
 
-	bool HasMarker3() const
+	//bool HasMarker3() const
+	//{
+	//	return ((States & EHaveStates::HasMarker3) == EHaveStates::HasMarker3);
+	//}
+
+	bool HasMarker1And2() const
 	{
-		return ((States & EHaveStates::HasMarker3) == EHaveStates::HasMarker3);
+		constexpr EHaveStates HasMarker12 = EHaveStates::HasMarker1 | EHaveStates::HasMarker2;
+		return ((States & HasMarker12) == HasMarker12);
+	}
+
+	bool HasMarker1Or2() const
+	{
+		constexpr EHaveStates HasMarker12 = EHaveStates::HasMarker1 | EHaveStates::HasMarker2;
+		return ((States & HasMarker12) != EHaveStates::None);
 	}
 
 	void SetMarker1() const
@@ -86,11 +102,6 @@ public:
 		States |= EHaveStates::HasMarker2;
 	}
 
-	void SetMarker3() const
-	{
-		States |= EHaveStates::HasMarker3;
-	}
-
 	void ResetMarker1() const
 	{
 		States &= ~EHaveStates::HasMarker1;
@@ -101,11 +112,6 @@ public:
 		States &= ~EHaveStates::HasMarker2;
 	}
 
-	void ResetMarker3() const
-	{
-		States &= ~EHaveStates::HasMarker3;
-	}
-
 	void ResetMarkers() const
 	{
 		States &= ~EHaveStates::AllMarkers;
@@ -147,6 +153,69 @@ public:
 		States &= ~EHaveStates::Degenerated;
 	}
 
+	void SetWaitingMarker() const
+	{
+		States |= EHaveStates::IsWaiting;
+	}
+
+	void ResetWaitingMarker() const
+	{
+		States &= ~EHaveStates::IsWaiting;
+	}
+
+	virtual bool IsWaiting() const
+	{
+		return ((States & EHaveStates::IsWaiting) == EHaveStates::IsWaiting);
+	}
+
+	void SetProcessedMarker() const
+	{
+		States |= EHaveStates::IsProcessed;
+	}
+
+	void ResetProcessedMarker() const
+	{
+		States &= ~EHaveStates::IsProcessed;
+	}
+
+	virtual bool IsProcessed() const
+	{
+		return ((States & EHaveStates::IsProcessed) == EHaveStates::IsProcessed);
+	}
+
+	virtual bool IsProcessedDeletedOrDegenerated() const
+	{
+		constexpr EHaveStates ProcessedDeletedOrDegenerated = EHaveStates::Degenerated | EHaveStates::IsDeleted | EHaveStates::IsProcessed;
+		return ((States & ProcessedDeletedOrDegenerated) != EHaveStates::None);
+	}
+
+
+	void SetToProcessMarker() const
+	{
+		States |= EHaveStates::ToProcess;
+	}
+
+	void ResetToProcessMarker() const
+	{
+		States &= ~EHaveStates::ToProcess;
+	}
+
+	virtual bool IsToProcess() const
+	{
+		return ((States & EHaveStates::ToProcess) == EHaveStates::ToProcess);
+	}
+
+	virtual bool IsNotToProcess() const
+	{
+		return ((States & EHaveStates::ToProcess) != EHaveStates::ToProcess);
+	}
+
+	virtual bool IsNotToOrAlreadyProcess() const
+	{
+		return IsNotToProcess() || IsProcessed();
+	}
+
+
 };
 }
 
diff --git a/Engine/Source/Runtime/Datasmith/CADKernel/Public/CADKernel/Core/Types.h b/Engine/Source/Runtime/Datasmith/CADKernel/Public/CADKernel/Core/Types.h
--- a/Engine/Source/Runtime/Datasmith/CADKernel/Public/CADKernel/Core/Types.h
+++ b/Engine/Source/Runtime/Datasmith/CADKernel/Public/CADKernel/Core/Types.h
@@ -13,7 +13,7 @@
 
 #define AThird 0.33333333333333333333333333333333
 #define AQuarter 0.25
-#define ASixth 0.16666666666666666666666666666667
+constexpr double ASixth = 1. / 6.; //0.16666666666666666666666666666667
 #define AEighth 0.125
 
 typedef uint32 FIdent;
@@ -21,6 +21,57 @@ typedef uint32 FIdent;
 namespace UE::CADKernel
 {
 
+struct FPairOfIndex
+{
+	int32 Value0;
+	int32 Value1;
+
+	FPairOfIndex(int32 InValue1, int32 InValue2)
+		: Value0(InValue1)
+		, Value1(InValue2)
+	{
+	}
+
+	FPairOfIndex(int32 InValue1)
+		: Value0(InValue1)
+		, Value1(-1)
+	{
+	}
+
+	int32& operator[](const int32 Index)
+	{
+		return Index == 0 ? Value0 : Value1;
+	}
+
+	const int32& operator[](const int32 Index) const
+	{
+		return Index == 0 ? Value0 : Value1;
+	}
+
+	void Add(const int32 Value)
+	{
+		if(Value >= 0)
+		{
+			if (Value0 < 0)
+			{
+				Value0 = Value;
+			}
+			else if (Value0 != Value && Value1 < 0)
+			{
+				Value1 = Value;
+			}
+		}
+	}
+
+	void Add(const FPairOfIndex Values)
+	{
+		Add(Values.Value0);
+		Add(Values.Value1);
+	}
+
+	static const FPairOfIndex Undefined;
+};
+
 namespace Ident
 {
 static const FIdent Undefined = (FIdent)-1;
diff --git a/Engine/Source/Runtime/Datasmith/CADKernel/Public/CADKernel/Geo/Curves/RestrictionCurve.h b/Engine/Source/Runtime/Datasmith/CADKernel/Public/CADKernel/Geo/Curves/RestrictionCurve.h
--- a/Engine/Source/Runtime/Datasmith/CADKernel/Public/CADKernel/Geo/Curves/RestrictionCurve.h
+++ b/Engine/Source/Runtime/Datasmith/CADKernel/Public/CADKernel/Geo/Curves/RestrictionCurve.h
@@ -148,10 +148,11 @@ public:
 
 	/**
 	 * Project a set of points of a twin curve on the 3D polyline and return the coordinate of the projected point
+	 * @param ToleranceOfProjection: Max error between the both curve to stop the search of projection of a point 
 	 */
-	void ProjectTwinCurvePoints(const TArray<FPoint>& InPointsToProject, bool bSameOrientation, TArray<double>& OutProjectedPointCoords) const
+	void ProjectTwinCurvePoints(const TArray<FPoint>& InPointsToProject, bool bSameOrientation, TArray<double>& OutProjectedPointCoords, double ToleranceOfProjection) const
 	{
-		Polyline.ProjectCoincidentalPolyline(InPointsToProject, bSameOrientation, OutProjectedPointCoords);
+		Polyline.ProjectCoincidentalPolyline(InPointsToProject, bSameOrientation, OutProjectedPointCoords, ToleranceOfProjection);
 	}
 
 	void ComputeIntersectionsWithIsos(const FLinearBoundary& InBoundary, const TArray<double>& InIsoCoordinates, const EIso InTypeIso, const FSurfacicTolerance& ToleranceIso, TArray<double>& OutIntersection) const
diff --git a/Engine/Source/Runtime/Datasmith/CADKernel/Public/CADKernel/Geo/Sampling/PolylineTools.h b/Engine/Source/Runtime/Datasmith/CADKernel/Public/CADKernel/Geo/Sampling/PolylineTools.h
--- a/Engine/Source/Runtime/Datasmith/CADKernel/Public/CADKernel/Geo/Sampling/PolylineTools.h
+++ b/Engine/Source/Runtime/Datasmith/CADKernel/Public/CADKernel/Geo/Sampling/PolylineTools.h
@@ -86,7 +86,7 @@ inline double SectionCoordinate(const TArray<double>& Array, const int32 Index,
 }
 
 template<typename PointType>
-double ComputeLength(TArray<PointType>& Polyline)
+double ComputeLength(const TArray<PointType>& Polyline)
 {
 	double Length = 0;
 	for (int32 Index = 1; Index < Polyline.Num(); ++Index)
@@ -96,6 +96,45 @@ double ComputeLength(TArray<PointType>& Polyline)
 	return Length;
 }
 
+template<typename PointType>
+double ComputeSquareToleranceForProjection(const TArray<PointType>& Polyline)
+{
+	const double Tolerance = ComputeLength(Polyline) * 0.1;
+	return Tolerance * Tolerance;
+}
+
+template<typename PointType>
+TArray<double> ComputePolylineSegmentLengths(const PointType& StartPoint, const TArray<PointType>& InnerPolyline, const PointType& EndPoint)
+{
+	TArray<double> ElementLength;
+	const int32 InnerPointCount = InnerPolyline.Num();
+	if (InnerPointCount > 0)
+	{
+		ElementLength.Reserve(InnerPointCount + 1);
+		const FPoint* PrevPoint = &InnerPolyline[0];
+		{
+			ElementLength.Add(PrevPoint->Distance(StartPoint));
+		}
+
+		for (int32 Index = 1; Index < InnerPointCount; ++Index)
+		{
+			const FPoint& CurrentPoint = InnerPolyline[Index];
+			ElementLength.Add(PrevPoint->Distance(CurrentPoint));
+			PrevPoint = &CurrentPoint;
+		}
+		{
+			ElementLength.Add(PrevPoint->Distance(EndPoint));
+		}
+	}
+	else
+	{
+		ElementLength.Reserve(1);
+		ElementLength.Add(StartPoint.Distance(EndPoint));
+	}
+
+	return MoveTemp(ElementLength);
+}
+
 /**
  * Progressively deforms a polyline (or a control polygon) so that its end is in the desired position
  */
@@ -416,7 +455,7 @@ public:
 			return;
 		}
 
-		int32 SegmentNum = (int32)FMath::Max(CurveLength / DesiredSegmentLength + 0.5, 1.0);
+		int32 SegmentNum = FMath::IsNearlyZero(DesiredSegmentLength) ? 2 : (int32)FMath::Max(CurveLength / DesiredSegmentLength + 0.5, 1.0);
 
 		double SectionLength = CurveLength / (double)(SegmentNum);
 
@@ -433,7 +472,7 @@ public:
 		double LastCoordinate = InBoundary.Min;
 		for (int32 IndexCurvilinear = 1, IndexCoordinate = BoundaryIndices[0] + 1; IndexCurvilinear < CurvilinearCoordinates.Num(); ++IndexCurvilinear, ++IndexCoordinate)
 		{
-			while (CurvilinearLength < CurvilinearCoordinates[IndexCurvilinear] + DOUBLE_SMALL_NUMBER)
+			while (CurvilinearLength + DOUBLE_SMALL_NUMBER < CurvilinearCoordinates[IndexCurvilinear])
 			{
 				double Coordinate = ComputeSamplePointCoordinate(IndexCurvilinear, IndexCoordinate, CurvilinearLength, LastCoordinate);
 				OutCoordinates.Add(Coordinate);
@@ -483,14 +522,18 @@ public:
 
 	/**
 	 * Project each point of a coincidental polyline on the Polyline.
+	 * @param ToleranceOfProjection: Max error between the both curve to stop the search of projection
+	 *    if ToleranceOfProjection < 0, it's compute with ComputeSquareToleranceForProjection
 	 */
-	void ProjectCoincidentalPolyline(const TArray<PointType>& InPointsToProject, bool bSameOrientation, TArray<double>& OutProjectedPointCoords) const
+	void ProjectCoincidentalPolyline(const TArray<PointType>& InPointsToProject, bool bSameOrientation, TArray<double>& OutProjectedPointCoords, const double ToleranceOfProjection) const
 	{
 #ifdef DEBUG_PROJECT_COINCIDENTAL_POLYLINE
-		static bool bDisplay = false;
+		static int32 Counter = 0;
+		++Counter;
+		bool bDisplay = true; // (Counter == 22);
 		if (bDisplay)
 		{
-			F3DDebugSession Session(TEXT("ProjectCoincidentalPolyline"));
+			F3DDebugSession Session(FString::Printf(TEXT("ProjectCoincidentalPolyline %d"), Counter));
 			{
 				F3DDebugSession _(TEXT("ProjectCoincidentalPolyline"));
 				DisplayPolyline(InPointsToProject, EVisuProperty::BlueCurve);
@@ -501,11 +544,18 @@ public:
 			}
 		}
 #endif
+
+		const double SquareTol = ToleranceOfProjection > 0 ? FMath::Square(ToleranceOfProjection) : PolylineTools::ComputeSquareToleranceForProjection(PolylinePoints);
+
 		int32 StartIndex = 0;
 		const int32 EndIndex = PolylinePoints.Num() - 1;
 
 		TFunction<double(const PointType&)> ProjectPointToPolyline = [&](const PointType& InPointToProject)
 		{
+#ifdef DEBUG_PROJECT_COINCIDENTAL_POLYLINE
+			PointType ClosePoint;
+#endif
+
 			double MinDistance = HUGE_VAL;
 			double UForMinDistance = 0;
 
@@ -514,16 +564,35 @@ public:
 			for (int32 Index = StartIndex; Index < EndIndex; ++Index)
 			{
 				PointType ProjectPoint = ProjectPointOnSegment(InPointToProject, PolylinePoints[Index], PolylinePoints[Index + 1], ParamU, true);
-
 				double SquareDistance = ProjectPoint.SquareDistance(InPointToProject);
-				if (SquareDistance > MinDistance)
+				if (SquareDistance > MinDistance + SquareTol)
 				{
 					break;
 				}
-				MinDistance = SquareDistance;
-				UForMinDistance = ParamU;
-				SegmentIndex = Index;
+
+				if (SquareDistance < MinDistance)
+				{
+					MinDistance = SquareDistance;
+					UForMinDistance = ParamU;
+					SegmentIndex = Index;
+#ifdef DEBUG_PROJECT_COINCIDENTAL_POLYLINE
+					ClosePoint = ProjectPoint;
+#endif
+				}
+
 			}
+#ifdef DEBUG_PROJECT_COINCIDENTAL_POLYLINE
+			if (bDisplay)
+			{
+				F3DDebugSession Session(TEXT("ProjectPoint"));
+				{
+					DisplayPoint(InPointToProject, EVisuProperty::BluePoint);
+					DisplayPoint(ClosePoint, EVisuProperty::RedPoint);
+					//Wait();
+				}
+			}
+#endif
+
 			StartIndex = SegmentIndex;
 			return PolylineTools::LinearInterpolation(PolylineCoordinates, SegmentIndex, UForMinDistance);
 		};
diff --git a/Engine/Source/Runtime/Datasmith/CADKernel/Public/CADKernel/Geo/Sampling/SurfacicPolyline.h b/Engine/Source/Runtime/Datasmith/CADKernel/Public/CADKernel/Geo/Sampling/SurfacicPolyline.h
--- a/Engine/Source/Runtime/Datasmith/CADKernel/Public/CADKernel/Geo/Sampling/SurfacicPolyline.h
+++ b/Engine/Source/Runtime/Datasmith/CADKernel/Public/CADKernel/Geo/Sampling/SurfacicPolyline.h
@@ -257,11 +257,12 @@ public:
 
 	/**
 	 * Project each point of a coincidental polyline on the Polyline.
+	 * @param ToleranceOfProjection: Max error between the both curve to stop the search of projection
 	 */
-	void ProjectCoincidentalPolyline(const TArray<FPoint>& InPointsToProject, bool bSameOrientation, TArray<double>& OutProjectedPointCoordinates) const
+	void ProjectCoincidentalPolyline(const TArray<FPoint>& InPointsToProject, bool bSameOrientation, TArray<double>& OutProjectedPointCoordinates, double ToleranceOfProjection) const
 	{
 		TPolylineApproximator<FPoint> Approximator3D(Coordinates, Points3D);
-		Approximator3D.ProjectCoincidentalPolyline(InPointsToProject, bSameOrientation, OutProjectedPointCoordinates);
+		Approximator3D.ProjectCoincidentalPolyline(InPointsToProject, bSameOrientation, OutProjectedPointCoordinates, ToleranceOfProjection);
 	}
 
 	/**
diff --git a/Engine/Source/Runtime/Datasmith/CADKernel/Public/CADKernel/Math/Boundary.h b/Engine/Source/Runtime/Datasmith/CADKernel/Public/CADKernel/Math/Boundary.h
--- a/Engine/Source/Runtime/Datasmith/CADKernel/Public/CADKernel/Math/Boundary.h
+++ b/Engine/Source/Runtime/Datasmith/CADKernel/Public/CADKernel/Math/Boundary.h
@@ -12,7 +12,7 @@ namespace UE::CADKernel
  * MINIMAL_UNIT_LINEAR_TOLERANCE allows to define the minimal tolerance value of a parametric space
  * @see FLinearBoundary::ComputeMinimalTolerance
  */
-#define MINIMAL_UNIT_LINEAR_TOLERANCE 10e-5
+#define MINIMAL_UNIT_LINEAR_TOLERANCE 1e-5
 
 struct CADKERNEL_API FLinearBoundary
 {
@@ -407,6 +407,13 @@ public:
 		UVBoundaries[EIso::IsoV].MoveInsideIfNot(Point.V, Tolerance);
 	}
 
+	FSurfacicBoundary& operator=(const FSurfacicBoundary& InBounds)
+	{
+		UVBoundaries[EIso::IsoU] = InBounds.UVBoundaries[EIso::IsoU];
+		UVBoundaries[EIso::IsoV] = InBounds.UVBoundaries[EIso::IsoV];
+		return *this;
+	}
+
 	double Length(const EIso& Iso) const
 	{
 		return UVBoundaries[Iso].Length();
diff --git a/Engine/Source/Runtime/Datasmith/CADKernel/Public/CADKernel/Math/Geometry.h b/Engine/Source/Runtime/Datasmith/CADKernel/Public/CADKernel/Math/Geometry.h
--- a/Engine/Source/Runtime/Datasmith/CADKernel/Public/CADKernel/Math/Geometry.h
+++ b/Engine/Source/Runtime/Datasmith/CADKernel/Public/CADKernel/Math/Geometry.h
@@ -502,11 +502,11 @@ inline double CoordinateOfProjectedPointOnSegment(const PointType& Point, const
 CADKERNEL_API void FindLoopIntersectionsWithIso(const EIso Iso, const double IsoParameter, const TArray<TArray<FPoint2D>>& Loops, TArray<double>& OutIntersections);
 
 /**
- * Similar as IntersectSegments2D but do not check intersection if both segment are carried by the same line.
+ * Similar as DoIntersect but do not check intersection if both segment are carried by the same line.
  * Must be done before (with BBox comparison for example)
  * This method is 50% faster than IntersectSegments2D even if
  */
-inline bool FastIntersectSegments2D(const TSegment<FPoint2D>& SegmentAB, const TSegment<FPoint2D>& SegmentCD)
+inline bool FastDoIntersect(const FSegment2D& SegmentAB, const FSegment2D& SegmentCD)
 {
 	constexpr const double Min = -DOUBLE_SMALL_NUMBER;
 	constexpr const double Max = 1. + DOUBLE_SMALL_NUMBER;
@@ -535,7 +535,7 @@ inline bool FastIntersectSegments2D(const TSegment<FPoint2D>& SegmentAB, const T
 /**
  * The segments must intersect because no check is done
  */
-inline FPoint2D FindIntersectionOfSegments2D(const TSegment<FPoint2D>& SegmentAB, const TSegment<FPoint2D>& SegmentCD, double& OutABIntersectionCoordinate)
+inline FPoint2D FindIntersectionOfSegments2D(const FSegment2D& SegmentAB, const FSegment2D& SegmentCD, double& OutABIntersectionCoordinate)
 {
 	const FPoint2D AB = SegmentAB[1] - SegmentAB[0];
 	const FPoint2D DC = SegmentCD[0] - SegmentCD[1];
@@ -584,7 +584,7 @@ inline FPoint2D FindIntersectionOfSegments2D(const TSegment<FPoint2D>& SegmentAB
 /**
  * The segments must intersect because no check is done
  */
-inline FPoint2D FindIntersectionOfSegments2D(const TSegment<FPoint2D>& SegmentAB, const TSegment<FPoint2D>& SegmentCD)
+inline FPoint2D FindIntersectionOfSegments2D(const FSegment2D& SegmentAB, const FSegment2D& SegmentCD)
 {
 	double ABIntersectionCoordinate;
 	return FindIntersectionOfSegments2D(SegmentAB, SegmentCD, ABIntersectionCoordinate);
@@ -593,7 +593,7 @@ inline FPoint2D FindIntersectionOfSegments2D(const TSegment<FPoint2D>& SegmentAB
 /**
  * @return false if the lines are parallel
  */
-inline bool FindIntersectionOfLines2D(const TSegment<FPoint2D>& LineAB, const TSegment<FPoint2D>& LineCD, FPoint2D& OutIntersectionPoint)
+inline bool FindIntersectionOfLines2D(const FSegment2D& LineAB, const FSegment2D& LineCD, FPoint2D& OutIntersectionPoint)
 {
 	constexpr const double Min = -DOUBLE_SMALL_NUMBER;
 	constexpr const double Max = 1. + DOUBLE_SMALL_NUMBER;
@@ -614,10 +614,11 @@ inline bool FindIntersectionOfLines2D(const TSegment<FPoint2D>& LineAB, const TS
 }
 
 /**
- * Similar as FastIntersectSegments2D but check intersection if both segment are carried by the same line.
+ * Similar as FastDoIntersect but check intersection if both segment are carried by the same line.
  * This method is 50% slower than FastIntersectSegments2D even if the segments tested are never carried by the same line
  */
-CADKERNEL_API bool IntersectSegments2D(const TSegment<FPoint2D>& SegmentAB, const TSegment<FPoint2D>& SegmentCD);
+CADKERNEL_API bool DoIntersect(const FSegment2D& SegmentAB, const FSegment2D& SegmentCD);
+CADKERNEL_API bool DoIntersectInside(const FSegment2D& SegmentAB, const FSegment2D& SegmentCD);
 
 inline double ComputeCosinus(FVector Vector, FVector OtherVector)
 {
diff --git a/Engine/Source/Runtime/Datasmith/CADKernel/Public/CADKernel/Math/MathConst.h b/Engine/Source/Runtime/Datasmith/CADKernel/Public/CADKernel/Math/MathConst.h
--- a/Engine/Source/Runtime/Datasmith/CADKernel/Public/CADKernel/Math/MathConst.h
+++ b/Engine/Source/Runtime/Datasmith/CADKernel/Public/CADKernel/Math/MathConst.h
@@ -59,5 +59,26 @@ uint8 ToUInt8(T Value)
 	return FMath::Clamp((uint8)(Value / 255.), (uint8)0, (uint8)255);
 }
 
+/**
+ * Wraps a periodic value into a prime period defined by its StartOfPeriod value and its EndOfPeriod value
+ * Mandatory EndOfPeriod - StartOfPeriod = PeriodSize
+ *
+ * Fast FMath::Fmod or FMath::Floor for the special case of Slope functions since most of the time the input value doesn't need to be change
+ * @see SlopeUtils.h
+ */
+inline double WrapTo(double Slope, const double StartOfPeriod, const double EndOfPeriod, const double PeriodLength)
+{
+	while (Slope < StartOfPeriod)
+	{
+		Slope += PeriodLength;
+	}
+	while (Slope >= EndOfPeriod)
+	{
+		Slope -= PeriodLength;
+	}
+	return Slope;
+}
+
+
 
 } // namespace UE::CADKernel	
diff --git a/Engine/Source/Runtime/Datasmith/CADKernel/Public/CADKernel/Math/Point.h b/Engine/Source/Runtime/Datasmith/CADKernel/Public/CADKernel/Math/Point.h
--- a/Engine/Source/Runtime/Datasmith/CADKernel/Public/CADKernel/Math/Point.h
+++ b/Engine/Source/Runtime/Datasmith/CADKernel/Public/CADKernel/Math/Point.h
@@ -519,7 +519,7 @@ public:
 
 	double Distance(const FPoint2D& Point) const
 	{
-		return FMath::Sqrt(FMath::Square(Point.U - U) + FMath::Square(Point.V - V));
+		return FMath::Sqrt(SquareDistance(Point));
 	}
 
 	double SquareDistance(const FPoint2D& Point) const
@@ -529,7 +529,7 @@ public:
 
 	double Length() const
 	{
-		return FMath::Sqrt(FMath::Square(U) + FMath::Square(V));
+		return FMath::Sqrt(SquareLength());
 	}
 
 	double SquareLength() const
diff --git a/Engine/Source/Runtime/Datasmith/CADKernel/Public/CADKernel/Math/SlopeUtils.h b/Engine/Source/Runtime/Datasmith/CADKernel/Public/CADKernel/Math/SlopeUtils.h
--- a/Engine/Source/Runtime/Datasmith/CADKernel/Public/CADKernel/Math/SlopeUtils.h
+++ b/Engine/Source/Runtime/Datasmith/CADKernel/Public/CADKernel/Math/SlopeUtils.h
@@ -9,16 +9,56 @@
 namespace UE::CADKernel
 {
 
+/**
+ * "Slope" is a fast angle approximation.
+ * 
+ * This file propose all useful methods to use slope instead of angle.
+ * 
+ * The method "ComputeSlope(const FPoint2D& StartPoint, const FPoint2D& EndPoint) is the main method. 
+ * It compute the slope between the input segment defined by two points and [0, u) axis.
+ * The return value is a real in the interval [0, 8] for an angle in the interval [0, 2Pi]
+ * 
+ * Warning, it's only an approximation... The conversion is not linear but the error is small near the integer value of slope (0, 1, 2, 3, ...8)
+ *
+ * To compute an angle value between two segments, the call of acos (and asin for an oriented angle) is necessary while with this approximation, only a division is useful.
+ * 
+ * This approximation is very good when only comparison of angles is needed and more faster than acos and/or asin i.e. Slope approximation need only a division and few addition and test
+ * 
+ * [0 - 2Pi] is divide into 8 angular sector i.e. [0, Pi/4] = [0,1], [Pi/4, Pi/2] = [1,2], ...
+ *
+ * The value of the slope for an angle in [0, Pi/4] = tan(angle)
+ * @return a slope between [0, 8] i.e. an equivalent angle between [0, 2Pi]
+ *
+ * Angle (Degree) to Slop
+ *		  0   = 0
+ *        1   ~ 0.0175
+ *        2   ~ 0.035
+ *        5   ~ 0.0875
+ *		 10   ~ 0.176
+ *		 15.  ~ 0.268
+ *		 20   ~ 0.364
+ *		 25   ~ 0.466
+ *       30   ~ 0.577
+ *		 45   = 1
+ *		 60   ~ 1.423  == 2 - Slope(30)
+ *		 90   = 2
+ *		120   ~ 2.577  == 2 + Slope(30)
+ *		135   = 3
+ *		180   = 4
+ *		360   = 8
+ */
+
 namespace Slope
 {
 
 constexpr double NullSlope = 0.;
-constexpr double FullSlope = 8.;
 
 /**
  * RightSlope i.e. Right angle i.e Pi / 2
  */
 constexpr double RightSlope = 2.;
+constexpr double HalfPiSlope = 2.;
+constexpr double NinetySlope = 2.;
 
 /**
  * ThreeRightSlope i.e. 3Pi / 2
@@ -34,20 +74,65 @@ constexpr double MinusRightSlope = -2.;
  * PiSlope i.e. Pi angle
  */
 constexpr double PiSlope = 4.;
+
+/**
+ * PiSlope i.e. Pi angle
+ */
+constexpr double TwoPiSlope = 8.;
+
+/**
+ * ThirdPiSlope i.e. Pi/3 angle (60 deg)
+ */
+constexpr double ThirdPiSlope = 1.422649730810374235490851219498;
+constexpr double SixtySlope = 1.422649730810374235490851219498;
+
+/**
+ * ThirdPiSlope i.e. Pi/4 angle (45 deg)
+ */
+constexpr double QuaterPiSlope = 1;
+constexpr double FortyFiveSlope = 0.57735026918962576450914878050196;
+
+/**
+ * ThirdPiSlope i.e. Pi/6 angle (30 deg)
+ */
+constexpr double SixthPiSlope = 0.57735026918962576450914878050196;
+constexpr double ThirtySlope  = 0.57735026918962576450914878050196;
+
+/**
+ * ThreeQuaterPiSlope i.e. 3Pi/4 angle (135 deg)
+ */
+constexpr double ThreeQuaterPiSlope = 3;
+                              
+constexpr double OneDegree  = 0.01745506492821758576512889521973;
+constexpr double TwoDegree  = 0.03492076949174773050040262577373;
+constexpr double FiveDegree = 0.08748866352592400522201866943496;
+constexpr double TenDegree = 0.17632698070846497347109038686862;
+constexpr double FifteenDegree = 0.26794919243112270647255365849413;
+constexpr double TwentyDegree = 0.36397023426620236135104788277683;
+constexpr double TwentyFiveDegree = 0.46630765815499859283000619479956;
+
+constexpr double Epsilon = 0.001;
+
 }
 
 typedef TFunction<double(const FPoint2D&, const FPoint2D&, double)> SlopeMethod;
 
-
 /**
  * Transform a positive slope into an oriented slope [-4, 4] i.e. an equivalent angle between [-Pi, Pi]
  * @return a slope between [-4, 4]
  */
 inline double TransformIntoOrientedSlope(double Slope)
 {
-	if (Slope > 4.) Slope -= 8;
-	if (Slope < -4.) Slope += 8;
-	return Slope;
+	return WrapTo(Slope, -Slope::PiSlope, Slope::PiSlope, Slope::TwoPiSlope);
+}
+
+/**
+ * Transform a positive slope into an unoriented slope [0, 4] i.e. an equivalent angle between [0, Pi]
+ * @return a slope between [0, 4]
+ */
+inline double TransformIntoUnorientedSlope(double Slope)
+{
+	return FMath::Abs(WrapTo(Slope, -Slope::PiSlope, Slope::PiSlope, Slope::TwoPiSlope));
 }
 
 /**
@@ -56,41 +141,37 @@ inline double TransformIntoOrientedSlope(double Slope)
  */
 inline double TransformIntoPositiveSlope(double Slope)
 {
-	if (Slope < 0) Slope += 8;
+	return WrapTo(Slope, Slope::NullSlope, Slope::TwoPiSlope, Slope::TwoPiSlope);
+}
+
+/**
+ * return a slope between [0, 2] relative to reference Axis i.e.
+ * ComputeUnorientedSlope => 0.5 return 0.5
+ * ComputeUnorientedSlope => 2.3 return 1.7
+ */
+inline double TransformIntoSlopeRelativeToReferenceAxis(double Slope)
+{
+	Slope = TransformIntoUnorientedSlope(Slope);
+
+	if (Slope > Slope::RightSlope)
+	{
+		Slope = Slope::PiSlope - Slope;
+	}
+
 	return Slope;
 }
 
+
 /**
- * Fast angle approximation by a "slope" :
- * This method is used to compute an approximation of the angle between the input segment defined by two points and [0, u) axis.
- * The return value is a real in the interval [0, 8] for an angle in the interval [0, 2Pi]
- * Warning, it's only an approximation... The conversion is not linear but the error is small near the integer value of slope (0, 1, 2, 3, ...8)
- *
- * This approximation is very good when only comparison of angles is needed.
- * This method is not adapted to compute angle
- *
- * To compute an angle value between two segments, the call of acos (and asin for an oriented angle) is necessary while with this approximation, only a division is useful.
- *
- * [0 - 2Pi] is divide into 8 angular sector i.e. [0, Pi/4] = [0,1], [Pi/4, Pi/2] = [1,2], ...
- *
- * @return a slope between [0, 8] i.e. an equivalent angle between [0, 2Pi]
- *
- * Angle (Degree) to Slop
- *		  0   = 0
- *		  7.  = 0.125
- *		 14.  = 0.25
- *		 30   = 0.5
- *		 36.8 = 0.75
- *		 45   = 1
- *		 53.2 = 1.25
- *		 60   = 1.5
- *		 76.  = 1.75
- *		 90   = 2
- *		120   = 2.5
- *		135   = 3
- *		180   = 4
- *		360   = 8
+ * Swap a slope i.e Slope + PiSlope i.e. angle + Pi
+ * @return a slope between [0, 8]
  */
+inline double SwapSlopeOrientation(double Slope)
+{
+	const double SwapedSlope = Slope < Slope::PiSlope ? Slope + Slope::PiSlope : Slope - Slope::PiSlope;
+	return TransformIntoPositiveSlope(SwapedSlope);
+}
+
 inline double ComputeSlope(const FPoint2D& StartPoint, const FPoint2D& EndPoint)
 {
 	double DeltaU = EndPoint.U - StartPoint.U;
@@ -112,18 +193,18 @@ inline double ComputeSlope(const FPoint2D& StartPoint, const FPoint2D& EndPoint)
 			}
 			else
 			{
-				Delta = 2 - DeltaU / DeltaV;
+				Delta = 2. - DeltaU / DeltaV;
 			}
 		}
 		else
 		{
 			if (DeltaU > -DeltaV)
 			{
-				Delta = 8 + DeltaV / DeltaU;
+				Delta = 8. + DeltaV / DeltaU;
 			}
 			else
 			{
-				Delta = 6 - DeltaU / DeltaV; // deltaU/deltaV <0
+				Delta = 6. - DeltaU / DeltaV; // deltaU/deltaV <0
 			}
 		}
 	}
@@ -133,22 +214,22 @@ inline double ComputeSlope(const FPoint2D& StartPoint, const FPoint2D& EndPoint)
 		{
 			if (-DeltaU > DeltaV)
 			{
-				Delta = 4 + DeltaV / DeltaU;
+				Delta = 4. + DeltaV / DeltaU;
 			}
 			else
 			{
-				Delta = 2 - DeltaU / DeltaV;
+				Delta = 2. - DeltaU / DeltaV;
 			}
 		}
 		else
 		{
 			if (-DeltaU > -DeltaV)
 			{
-				Delta = 4 + DeltaV / DeltaU;
+				Delta = 4. + DeltaV / DeltaU;
 			}
 			else
 			{
-				Delta = 6 - DeltaU / DeltaV;
+				Delta = 6. - DeltaU / DeltaV;
 			}
 		}
 	}
@@ -156,6 +237,27 @@ inline double ComputeSlope(const FPoint2D& StartPoint, const FPoint2D& EndPoint)
 	return Delta;
 }
 
+/**
+ * Compute the slope of a segment according to a reference slope
+ * return the slope
+ */
+inline double ComputeSlope(const FPoint2D& StartPoint, const FPoint2D& EndPoint, double ReferenceSlope)
+{
+	const double Slope = ComputeSlope(StartPoint, EndPoint);
+	return Slope - ReferenceSlope;
+}
+
+/**
+ * Compute the slope between the segments [StartPoint, EndPoint1] and [StartPoint, EndPoint2]
+ * @return a slope i.e. an equivalent angle
+ */
+inline double ComputeSlope(const FPoint2D& StartPoint, const FPoint2D& EndPoint1, const FPoint2D& EndPoint2)
+{
+	const double ReferenceSlope = ComputeSlope(StartPoint, EndPoint1);
+	const double Slope = ComputeSlope(StartPoint, EndPoint2);
+	return Slope - ReferenceSlope;
+}
+
 /**
  * Compute the oriented slope of a segment according to a reference slope
  * This method is used to compute an approximation of the angle between two segments in 2D.
@@ -163,10 +265,8 @@ inline double ComputeSlope(const FPoint2D& StartPoint, const FPoint2D& EndPoint)
  */
 inline double ComputePositiveSlope(const FPoint2D& StartPoint, const FPoint2D& EndPoint, double ReferenceSlope)
 {
-	double Slope = ComputeSlope(StartPoint, EndPoint);
-	Slope -= ReferenceSlope;
-	if (Slope < 0) Slope += 8;
-	return Slope;
+	double Slope = ComputeSlope(StartPoint, EndPoint, ReferenceSlope);
+	return TransformIntoPositiveSlope(Slope);
 }
 
 /**
@@ -175,15 +275,13 @@ inline double ComputePositiveSlope(const FPoint2D& StartPoint, const FPoint2D& E
  */
 inline double ComputePositiveSlope(const FPoint2D& StartPoint, const FPoint2D& EndPoint1, const FPoint2D& EndPoint2)
 {
-	double ReferenceSlope = ComputeSlope(StartPoint, EndPoint1);
-	double Slope = ComputeSlope(StartPoint, EndPoint2);
-	Slope -= ReferenceSlope;
+	double Slope = ComputeSlope(StartPoint, EndPoint1, EndPoint2);
 	return TransformIntoPositiveSlope(Slope);
 }
 
 inline double ClockwiseSlope(const FPoint2D& StartPoint, const FPoint2D& EndPoint, double ReferenceSlope)
 {
-	return 8. - ComputePositiveSlope(StartPoint, EndPoint, ReferenceSlope);
+	return Slope::TwoPiSlope - ComputePositiveSlope(StartPoint, EndPoint, ReferenceSlope);
 }
 
 inline double CounterClockwiseSlope(const FPoint2D& StartPoint, const FPoint2D& EndPoint, double ReferenceSlope)
@@ -197,7 +295,7 @@ inline double CounterClockwiseSlope(const FPoint2D& StartPoint, const FPoint2D&
  */
 inline double ComputeOrientedSlope(const FPoint2D& StartPoint, const FPoint2D& EndPoint, double ReferenceSlope)
 {
-	return TransformIntoOrientedSlope(ComputePositiveSlope(StartPoint, EndPoint, ReferenceSlope));
+	return TransformIntoOrientedSlope(ComputeSlope(StartPoint, EndPoint, ReferenceSlope));
 }
 
 /**
@@ -206,7 +304,7 @@ inline double ComputeOrientedSlope(const FPoint2D& StartPoint, const FPoint2D& E
  */
 inline double ComputeOrientedSlope(const FPoint2D& StartPoint, const FPoint2D& EndPoint1, const FPoint2D& EndPoint2)
 {
-	return TransformIntoOrientedSlope(ComputePositiveSlope(StartPoint, EndPoint1, EndPoint2));
+	return TransformIntoOrientedSlope(ComputeSlope(StartPoint, EndPoint1, EndPoint2));
 }
 
 /**
@@ -214,7 +312,8 @@ inline double ComputeOrientedSlope(const FPoint2D& StartPoint, const FPoint2D& E
  */
 inline double ComputeUnorientedSlope(const FPoint2D& StartPoint, const FPoint2D& EndPoint, double ReferenceSlope)
 {
-	return FMath::Abs(ComputeOrientedSlope(StartPoint, EndPoint, ReferenceSlope));
+	const double Slope = ComputeSlope(StartPoint, EndPoint, ReferenceSlope);
+	return TransformIntoUnorientedSlope(Slope);
 }
 
 /**
@@ -225,18 +324,18 @@ inline double ComputeUnorientedSlope(const FPoint2D& StartPoint, const FPoint2D&
  */
 inline double ComputeSlopeRelativeToNearestAxis(const FPoint2D& StartPoint, const FPoint2D& EndPoint)
 {
-	double Slope = FMath::Abs(TransformIntoOrientedSlope(ComputeSlope(StartPoint, EndPoint)));
-	if (Slope > 2)
+	double Slope = TransformIntoUnorientedSlope(ComputeSlope(StartPoint, EndPoint));
+	if (Slope > Slope::RightSlope)
 	{
-		Slope = 4 - Slope;
+		Slope = Slope::PiSlope - Slope;
 	}
 
 	// if slope close to 2 means segment close to IsoU, otherwise segment close to IsoV
 	// Wants a slope between 0 and 1 to manage either IsoU and IsoV
 	// Close to 0 means close to IsoU or IsoV
-	if (Slope > 1)
+	if (Slope > Slope::QuaterPiSlope)
 	{
-		Slope = 2 - Slope;
+		Slope = Slope::RightSlope - Slope;
 	}
 
 	return Slope;
@@ -249,13 +348,8 @@ inline double ComputeSlopeRelativeToNearestAxis(const FPoint2D& StartPoint, cons
  */
 inline double ComputeSlopeRelativeToReferenceAxis(const FPoint2D& StartPoint, const FPoint2D& EndPoint, double ReferenceAxisSlope)
 {
-	double Slope = ComputeUnorientedSlope(StartPoint, EndPoint, ReferenceAxisSlope);
-	if (Slope > 2)
-	{
-		Slope = 4 - Slope;
-	}
-
-	return Slope;
+	const double Slope = ComputeSlope(StartPoint, EndPoint, ReferenceAxisSlope);
+	return TransformIntoSlopeRelativeToReferenceAxis(Slope);
 }
 
 /**
@@ -263,7 +357,7 @@ inline double ComputeSlopeRelativeToReferenceAxis(const FPoint2D& StartPoint, co
  */
 inline double ComputeUnorientedSlope(const FPoint2D& StartPoint, const FPoint2D& EndPoint1, const FPoint2D& EndPoint2)
 {
-	return FMath::Abs(ComputeOrientedSlope(StartPoint, EndPoint1, EndPoint2));
+	return TransformIntoUnorientedSlope(ComputeSlope(StartPoint, EndPoint1, EndPoint2));
 }
 
 /**
@@ -367,7 +461,7 @@ inline FPoint2D SlopeToVector(const double Slope)
 	case 6:
 		// 6 - DeltaU / DeltaV // deltaU/deltaV <0
 		Vector[0] = Slope - 6.;
-		Vector[1] = -1;
+		Vector[1] = -1.;
 		break;
 	case 7:
 		// 8 + DeltaV / DeltaU; // deltaU/deltaV <0
diff --git a/Engine/Source/Runtime/Datasmith/CADKernel/Public/CADKernel/Mesh/MeshEnum.h b/Engine/Source/Runtime/Datasmith/CADKernel/Public/CADKernel/Mesh/MeshEnum.h
--- a/Engine/Source/Runtime/Datasmith/CADKernel/Public/CADKernel/Mesh/MeshEnum.h
+++ b/Engine/Source/Runtime/Datasmith/CADKernel/Public/CADKernel/Mesh/MeshEnum.h
@@ -28,6 +28,15 @@ enum EMeshingState : uint8
 	FullyMeshed
 };
 
+enum ECoordinateType : uint8
+{
+	VertexCoordinate,
+	ImposedCoordinate, // imposted cutting point by thin zone meshing process
+	IsoUCoordinate,
+	IsoVCoordinate,
+	IsoUVCoordinate,
+	OtherCoordinate
+};
 namespace GridSpace
 {
 
diff --git a/Engine/Source/Runtime/Datasmith/CADKernel/Public/CADKernel/Mesh/Structure/EdgeMesh.h b/Engine/Source/Runtime/Datasmith/CADKernel/Public/CADKernel/Mesh/Structure/EdgeMesh.h
--- a/Engine/Source/Runtime/Datasmith/CADKernel/Public/CADKernel/Mesh/Structure/EdgeMesh.h
+++ b/Engine/Source/Runtime/Datasmith/CADKernel/Public/CADKernel/Mesh/Structure/EdgeMesh.h
@@ -39,6 +39,8 @@ public:
 	{
 		return EdgeVerticesIndex.Num();
 	}
+
+	TArray<double> GetElementLengths() const;
 };
 }
 
diff --git a/Engine/Source/Runtime/Datasmith/CADKernel/Public/CADKernel/Mesh/Structure/EdgeSegment.h b/Engine/Source/Runtime/Datasmith/CADKernel/Public/CADKernel/Mesh/Structure/EdgeSegment.h
--- a/Engine/Source/Runtime/Datasmith/CADKernel/Public/CADKernel/Mesh/Structure/EdgeSegment.h
+++ b/Engine/Source/Runtime/Datasmith/CADKernel/Public/CADKernel/Mesh/Structure/EdgeSegment.h
@@ -23,7 +23,7 @@ class FEdgeSegment : public FHaveStates
 private:
 	FTopologicalEdge* Edge;
 	double Coordinates[2];
-	FPoint2D Points[2];
+	FPoint2D USSPoints[2]; // in Uniform Scaled space
 
 	FEdgeSegment* NextSegment;
 	FEdgeSegment* PreviousSegment;
@@ -69,19 +69,19 @@ public:
 		Edge = InEdge;
 		Coordinates[ELimit::Start] = InStartU;
 		Coordinates[ELimit::End] = InEndU;
-		Points[ELimit::Start] = InStartPoint;
-		Points[ELimit::End] = InEndPoint;
+		USSPoints[ELimit::Start] = InStartPoint;
+		USSPoints[ELimit::End] = InEndPoint;
 		NextSegment = nullptr;
 		PreviousSegment = nullptr;
 		ClosedSegment = nullptr;
 
 		SquareDistanceToClosedSegment = HUGE_VAL;
-		Length = Points[ELimit::Start].Distance(Points[ELimit::End]);
+		Length = USSPoints[ELimit::Start].Distance(USSPoints[ELimit::End]);
 
 		Id = ++LastId;
 		ChainIndex = Ident::Undefined;
 
-		Boundary.Set(Points[ELimit::Start], Points[ELimit::End]);
+		Boundary.Set(USSPoints[ELimit::Start], USSPoints[ELimit::End]);
 
 		AxisMin = Boundary[EIso::IsoU].Min + Boundary[EIso::IsoV].Min;
 	};
@@ -139,12 +139,7 @@ public:
 		return Id;
 	}
 
-	const FTopologicalEdge* GetEdge() const
-	{
-		return Edge;
-	}
-
-	FTopologicalEdge* GetEdge()
+	FTopologicalEdge* GetEdge() const
 	{
 		return Edge;
 	}
@@ -154,20 +149,29 @@ public:
 		return Length;
 	}
 
+	/**
+	 * @return Center of the segment in uniform scaled space
+	 */
 	FPoint2D GetCenter() const
 	{
-		return (Points[ELimit::Start] + Points[ELimit::End]) * 0.5;
+		return (USSPoints[ELimit::Start] + USSPoints[ELimit::End]) * 0.5;
 	}
 
+	/**
+	 * @return Point a parameter U of the segment in uniform scaled space
+	 */
 	FPoint2D ComputeEdgePoint(double EdgeParamU) const
 	{
 		double SegmentParamS = (EdgeParamU - Coordinates[ELimit::Start]) / (Coordinates[ELimit::End] - Coordinates[ELimit::Start]);
-		return Points[ELimit::Start] + (Points[ELimit::End] - Points[ELimit::Start]) * SegmentParamS;
+		return USSPoints[ELimit::Start] + (USSPoints[ELimit::End] - USSPoints[ELimit::Start]) * SegmentParamS;
 	}
 
+	/**
+	 * @return the extremity Point of the segment in uniform scaled space
+	 */
 	const FPoint2D& GetExtemity(const ELimit Limit) const
 	{
-		return Points[Limit];
+		return USSPoints[Limit];
 	}
 
 	double GetCoordinate(const ELimit Limit) const
@@ -175,7 +179,7 @@ public:
 		return Coordinates[Limit];
 	}
 
-	bool IsForward()
+	bool IsForward() const
 	{
 		return (Coordinates[ELimit::End] - Coordinates[ELimit::Start]) >= 0;
 	}
@@ -183,21 +187,39 @@ public:
 	/**
 	 * Compute the slope of the input Segment according to this.
 	 */
-	double ComputeUnorientedSlopeOf(const FEdgeSegment* Segment)
+	double ComputeUnorientedSlopeOf(const FEdgeSegment* Segment) const 
 	{
-		double ReferenceSlope = ComputeSlope(Points[ELimit::Start], Points[ELimit::End]);
-		return ComputeUnorientedSlope(Segment->Points[ELimit::Start], Segment->Points[ELimit::End], ReferenceSlope);
+		double ReferenceSlope = ComputeSlope(USSPoints[ELimit::Start], USSPoints[ELimit::End]);
+		return ComputeUnorientedSlope(Segment->USSPoints[ELimit::Start], Segment->USSPoints[ELimit::End], ReferenceSlope);
+	}
+
+	/**
+	 * Compute the slope of the input Segment according to this.
+	 */
+	double ComputeOrientedSlopeOf(const FEdgeSegment* Segment) const
+	{
+		double ReferenceSlope = ComputeSlope(USSPoints[ELimit::Start], USSPoints[ELimit::End]);
+		return ComputeOrientedSlope(Segment->USSPoints[ELimit::Start], Segment->USSPoints[ELimit::End], ReferenceSlope);
 	}
 
 	/**
 	 * Compute the slope of the Segment defined by the two input points according to this.
 	 */
-	double ComputeUnorientedSlopeOf(const FPoint2D& Middle, const FPoint2D& Projection)
+	double ComputeUnorientedSlopeOf(const FPoint2D& Middle, const FPoint2D& Projection) const 
 	{
-		double ReferenceSlope = ComputeSlope(Points[ELimit::Start], Points[ELimit::End]);
+		double ReferenceSlope = ComputeSlope(USSPoints[ELimit::Start], USSPoints[ELimit::End]);
 		return ComputeUnorientedSlope(Projection, Middle, ReferenceSlope);
 	}
 
+	/**
+	 * Compute the slope of the Segment defined by the two input points according to this.
+	 */
+	double ComputeOrientedSlopeOf(const FPoint2D& Middle, const FPoint2D& Projection) const
+	{
+		double ReferenceSlope = ComputeSlope(USSPoints[ELimit::Start], USSPoints[ELimit::End]);
+		return ComputeOrientedSlope(Projection, Middle, ReferenceSlope);
+	}
+
 	FEdgeSegment* GetNext() const
 	{
 		return NextSegment;
@@ -251,9 +273,15 @@ public:
 		return Coordinates[ELimit::Start] + (Coordinates[ELimit::End] - Coordinates[ELimit::Start]) * SegmentU;
 	}
 
+	/** Compute the delta U corresponding to a delta length in the 2d space */
+	double ComputeDeltaU(const double DeltaLength) const
+	{
+		return FMath::Abs(Coordinates[ELimit::End] - Coordinates[ELimit::Start]) * DeltaLength / Length;
+	}
+
 	FPoint2D ProjectPoint(const FPoint2D& PointToProject, double& SegmentU) const
 	{
-		return ProjectPointOnSegment(PointToProject, Points[ELimit::Start], Points[ELimit::End], SegmentU, true);
+		return ProjectPointOnSegment(PointToProject, USSPoints[ELimit::Start], USSPoints[ELimit::End], SegmentU, true);
 	}
 
 private:
diff --git a/Engine/Source/Runtime/Datasmith/CADKernel/Public/CADKernel/Mesh/Structure/ModelMesh.h b/Engine/Source/Runtime/Datasmith/CADKernel/Public/CADKernel/Mesh/Structure/ModelMesh.h
--- a/Engine/Source/Runtime/Datasmith/CADKernel/Public/CADKernel/Mesh/Structure/ModelMesh.h
+++ b/Engine/Source/Runtime/Datasmith/CADKernel/Public/CADKernel/Mesh/Structure/ModelMesh.h
@@ -27,9 +27,9 @@ class CADKERNEL_API FModelMesh : public FEntityGeom
 	TArray<TArray<FPoint>*> GlobalPointCloud;
 	FIdent LastIdUsed;
 
-	TArray<TSharedPtr<FVertexMesh>> VertexMeshes;
-	TArray<TSharedPtr<FEdgeMesh>> EdgeMeshes;
-	TArray<TSharedPtr<FFaceMesh>> FaceMeshes;
+	TArray<FVertexMesh*> VertexMeshes;
+	TArray<FEdgeMesh*> EdgeMeshes;
+	TArray<FFaceMesh*> FaceMeshes;
 
 	bool QuadAnalyse = false;
 	double MinSize = DOUBLE_SMALL_NUMBER;
@@ -64,17 +64,17 @@ public:
 			return;
 		}
 
-		SpawnIdentOnEntities(VertexMeshes, Database);
-		SpawnIdentOnEntities(EdgeMeshes, Database);
-		SpawnIdentOnEntities(FaceMeshes, Database);
+		SpawnIdentOnEntities((TArray<FEntity*>&) VertexMeshes, Database);
+		SpawnIdentOnEntities((TArray<FEntity*>&) EdgeMeshes, Database);
+		SpawnIdentOnEntities((TArray<FEntity*>&) FaceMeshes, Database);
 	};
 
-	virtual void ResetMarkersRecursively()
+	virtual void ResetMarkersRecursively() const override
 	{
 		ResetMarkers();
-		ResetMarkersRecursivelyOnEntities(VertexMeshes);
-		ResetMarkersRecursivelyOnEntities(EdgeMeshes);
-		ResetMarkersRecursivelyOnEntities(FaceMeshes);
+		ResetMarkersRecursivelyOnEntities((TArray<FEntity*>&) VertexMeshes);
+		ResetMarkersRecursivelyOnEntities((TArray<FEntity*>&) EdgeMeshes);
+		ResetMarkersRecursivelyOnEntities((TArray<FEntity*>&) FaceMeshes);
 	};
 
 
@@ -115,19 +115,19 @@ public:
 		return Sag;
 	}
 
-	void AddMesh(TSharedRef<FVertexMesh> Mesh)
+	void AddMesh(FVertexMesh& Mesh)
 	{
-		VertexMeshes.Add(Mesh);
+		VertexMeshes.Add(&Mesh);
 	}
 
-	void AddMesh(TSharedRef<FEdgeMesh> Mesh)
+	void AddMesh(FEdgeMesh& Mesh)
 	{
-		EdgeMeshes.Add(Mesh);
+		EdgeMeshes.Add(&Mesh);
 	}
 
-	void AddMesh(TSharedRef<FFaceMesh> Mesh)
+	void AddMesh(FFaceMesh& Mesh)
 	{
-		FaceMeshes.Add(Mesh);
+		FaceMeshes.Add(&Mesh);
 	}
 
 	void RegisterCoordinates(TArray<FPoint>& Coordinates, int32& OutStartVertexId, int32& OutIndex)
@@ -139,18 +139,18 @@ public:
 		LastIdUsed += (int32)Coordinates.Num();
 	}
 
-	const int32 GetIndexOfVertexFromId(const int32 Id) const;
-	const int32 GetIndexOfEdgeFromId(const int32 Id) const;
-	const int32 GetIndexOfSurfaceFromId(const int32 Id) const;
+	//const int32 GetIndexOfVertexFromId(const int32 Id) const;
+	//const int32 GetIndexOfEdgeFromId(const int32 Id) const;
+	//const int32 GetIndexOfSurfaceFromId(const int32 Id) const;
 
-	const TSharedPtr<FVertexMesh> GetMeshOfVertexNodeId(const int32 Id) const;
+	const FVertexMesh* GetMeshOfVertexNodeId(const int32 Id) const;
 
 	void GetNodeCoordinates(TArray<FPoint>& NodeCoordinates) const;
 	void GetNodeCoordinates(TArray<FVector3f>& NodeCoordinates) const;
 
-	const TArray<TSharedPtr<FMesh>>& GetMeshes() const;
+	const TArray<FMesh*>& GetMeshes() const;
 
-	const TArray<TSharedPtr<FFaceMesh>>& GetFaceMeshes() const
+	const TArray<FFaceMesh*>& GetFaceMeshes() const
 	{
 		return FaceMeshes;
 	}
diff --git a/Engine/Source/Runtime/Datasmith/CADKernel/Public/CADKernel/Mesh/Structure/ThinZone2D.h b/Engine/Source/Runtime/Datasmith/CADKernel/Public/CADKernel/Mesh/Structure/ThinZone2D.h
--- a/Engine/Source/Runtime/Datasmith/CADKernel/Public/CADKernel/Mesh/Structure/ThinZone2D.h
+++ b/Engine/Source/Runtime/Datasmith/CADKernel/Public/CADKernel/Mesh/Structure/ThinZone2D.h
@@ -3,7 +3,6 @@
 
 #include "CADKernel/Core/Chrono.h"
 #include "CADKernel/Core/HaveStates.h"
-//#include "CADKernel/Core/Factory.h"
 #include "CADKernel/Core/Types.h"
 #include "CADKernel/Math/Point.h"
 
@@ -23,7 +22,19 @@ namespace UE::CADKernel
 		TooSmall        // -> to delete
 	};
 
+	enum class ESide : uint8
+	{
+		First = 0,
+		Second,
+		None
+	};
+
+	class FModelMesh;
 	class FThinZone2D;
+	class FTopologicalFace;
+
+	using FReserveContainerFunc = TFunction<void(int32)>;
+	using FAddMeshNodeFunc = TFunction<void(const int32, const FPoint2D&, const double, const FEdgeSegment&, const FPairOfIndex&)>;
 
 	class FThinZoneSide : public FHaveStates
 	{
@@ -33,6 +44,8 @@ namespace UE::CADKernel
 		TArray<FEdgeSegment> Segments;
 		FThinZoneSide& FrontSide;
 
+		TArray<FTopologicalEdge*> Edges;
+
 		double SideLength;
 		double MediumThickness;
 		double MaxThickness;
@@ -40,6 +53,7 @@ namespace UE::CADKernel
 	public:
 
 		FThinZoneSide(FThinZoneSide* InFrontSide, const TArray<FEdgeSegment*>& InSegments);
+
 		virtual ~FThinZoneSide() = default;
 
 		void Empty()
@@ -57,7 +71,25 @@ namespace UE::CADKernel
 			return Segments.Last();
 		}
 
-		void SetEdgesAsThinZone();
+		/** 
+		 * Use Marker1 flag to selected edges once 
+		 * Edge Marker1 has to be reset after.
+		 */
+		void GetEdges(TArray<FTopologicalEdge*>& OutEdges) const;
+
+		const TArray<FTopologicalEdge*>& GetEdges() const
+		{
+			return Edges;
+		}
+
+		TArray<FTopologicalEdge*>& GetEdges()
+		{
+			return Edges;
+		}
+
+		void AddToEdge();
+
+		void CleanMesh();
 
 		const TArray<FEdgeSegment>& GetSegments() const
 		{
@@ -95,6 +127,18 @@ namespace UE::CADKernel
 			return Segments[0].IsInner();
 		}
 
+		FThinZoneSide& GetFrontThinZoneSide()
+		{
+			return FrontSide;
+		}
+
+		void CheckEdgesZoneSide();
+		void SetEdgesZoneSide(ESide Side);
+
+		int32 GetImposedPointCount();
+
+		void GetExistingMeshNodes(const FTopologicalFace& Face, FModelMesh& MeshModel, FReserveContainerFunc& Reserve, FAddMeshNodeFunc& AddMeshNode, const bool bWithTolerance) const;
+
 	private:
 		void ComputeThicknessAndLength();
 
@@ -104,14 +148,16 @@ namespace UE::CADKernel
 	{
 	private:
 
-		FThinZoneSide FirstSide;
-		FThinZoneSide SecondSide;
+		FThinZoneSide SideA;
+		FThinZoneSide SideB;
 
 		EThinZone2DType Category;
 
 		double Thickness;
 		double MaxThickness;
 
+		bool bIsSwap = false;
+
 	public:
 
 		/**
@@ -119,8 +165,8 @@ namespace UE::CADKernel
 		 * So FThinZone2D can be transfered into the FTopologicalFace 
 		 */
 		FThinZone2D(const TArray<FEdgeSegment*>& InFirstSideSegments, const TArray<FEdgeSegment*>& InSecondSideSegments)
-			: FirstSide(&SecondSide, InFirstSideSegments)
-			, SecondSide(&FirstSide, InSecondSideSegments)
+			: SideA(&SideB, InFirstSideSegments)
+			, SideB(&SideA, InSecondSideSegments)
 			, Category(EThinZone2DType::Undefined)
 		{
 			Finalize();
@@ -130,10 +176,9 @@ namespace UE::CADKernel
 
 		void Empty()
 		{
-			FirstSide.Empty();
-			SecondSide.Empty();
+			SideA.Empty();
+			SideB.Empty();
 			Thickness = -1;
-			SetRemoved();
 		}
 
 		double GetThickness() const
@@ -146,24 +191,34 @@ namespace UE::CADKernel
 			return MaxThickness;
 		};
 
+		const FThinZoneSide& GetSide(ESide Side) const
+		{
+			return ((Side == ESide::Second) == bIsSwap) ? SideA : SideB;
+		}
+
+		FThinZoneSide& GetSide(ESide Side)
+		{
+			return const_cast<FThinZoneSide&> (static_cast<const FThinZone2D*>(this)->GetSide(Side));
+		}
+
 		FThinZoneSide& GetFirstSide()
 		{
-			return FirstSide;
+			return bIsSwap ? SideB : SideA;
 		}
 
 		FThinZoneSide& GetSecondSide()
 		{
-			return SecondSide;
+			return bIsSwap ? SideA : SideB;
 		}
 
 		const FThinZoneSide& GetFirstSide() const
 		{
-			return FirstSide;
+			return bIsSwap ? SideB : SideA;
 		}
 
 		const FThinZoneSide& GetSecondSide() const 
 		{
-			return SecondSide;
+			return bIsSwap ? SideA : SideB;
 		}
 
 		EThinZone2DType GetCategory() const
@@ -171,37 +226,72 @@ namespace UE::CADKernel
 			return Category;
 		}
 
-		void SetEdgesAsThinZone();
 		static void SetPeakEdgesMarker(const TArray<const FTopologicalEdge*>&);
 
 		double Length() const
 		{
-			return FirstSide.Length() + SecondSide.Length();
+			return SideA.Length() + SideB.Length();
 		}
 
 		double GetMaxSideLength() const
 		{
-			return FMath::Max(FirstSide.Length(), SecondSide.Length());
+			return FMath::Max(SideA.Length(), SideB.Length());
 		}
 
-		bool IsRemoved() const
+		void SetCategory(EThinZone2DType InType)
 		{
-			return ((States & EHaveStates::IsRemoved) == EHaveStates::IsRemoved);
+			Category = InType;
 		}
 
-		void SetRemoved() const
+		/**
+		 * Use Marker 3 flag to count edge once
+		 * Edge Marker 3 has to be reset after.
+		 */
+		void GetEdges(TArray<FTopologicalEdge*>& OutSideAEdges, TArray<FTopologicalEdge*>& OutSideBEdges) const
 		{
-			States |= EHaveStates::IsRemoved;
+			SideA.GetEdges(OutSideAEdges);
+			SideB.GetEdges(OutSideBEdges);
 		}
 
-		void ResetRemoved() const
+		/**
+		 * Use Marker 1 and 2 flags
+		 * they have to be reset after.
+		 */
+		void CheckEdgesZoneSide()
 		{
-			States &= ~EHaveStates::IsRemoved;
+			SideA.CheckEdgesZoneSide();
+			SideB.CheckEdgesZoneSide();
 		}
 
-		void SetCategory(EThinZone2DType InType)
+		/**
+		 * For each edge of the zone, set Marker1 flag if edge is in first zone, and Marker2 if edge is in second zone
+		 */
+		void SetEdgesZoneSide()
 		{
-			Category = InType;
+			SideA.SetEdgesZoneSide(ESide::First);
+			SideB.SetEdgesZoneSide(ESide::Second);
+		}
+
+		void AddToEdge();
+
+		void Swap()
+		{
+			bIsSwap = !bIsSwap;
+			switch (Category)
+			{
+			case EThinZone2DType::PeakStart:
+			{
+				Category = EThinZone2DType::PeakStart;
+				break;
+			}
+			case EThinZone2DType::PeakEnd:
+			{
+				Category = EThinZone2DType::PeakStart;
+				break;
+			}
+			default:
+				break;
+			};
 		}
 
 #ifdef CADKERNEL_DEV
diff --git a/Engine/Source/Runtime/Datasmith/CADKernel/Public/CADKernel/Topo/Linkable.h b/Engine/Source/Runtime/Datasmith/CADKernel/Public/CADKernel/Topo/Linkable.h
--- a/Engine/Source/Runtime/Datasmith/CADKernel/Public/CADKernel/Topo/Linkable.h
+++ b/Engine/Source/Runtime/Datasmith/CADKernel/Public/CADKernel/Topo/Linkable.h
@@ -160,6 +160,12 @@ public:
 		States &= ~EHaveStates::ThinZone;
 	}
 
+	virtual void ResetMarkersRecursively() const override
+	{
+		TopologicalLink->ResetMarkersRecursively();
+	}
+
+
 protected:
 
 	void MakeLink(EntityType& Twin)
diff --git a/Engine/Source/Runtime/Datasmith/CADKernel/Public/CADKernel/Topo/Shell.h b/Engine/Source/Runtime/Datasmith/CADKernel/Public/CADKernel/Topo/Shell.h
--- a/Engine/Source/Runtime/Datasmith/CADKernel/Public/CADKernel/Topo/Shell.h
+++ b/Engine/Source/Runtime/Datasmith/CADKernel/Public/CADKernel/Topo/Shell.h
@@ -29,6 +29,11 @@ public:
 	{
 	}
 
+	FOrientedFace(const FOrientedFace& OrientiredEntity)
+		: TOrientedEntity(OrientiredEntity)
+	{
+	}
+
 	FOrientedFace()
 		: TOrientedEntity()
 	{
@@ -83,7 +88,7 @@ public:
 		SpawnIdentOnEntities((TArray<TOrientedEntity<FEntity>>&) TopologicalFaces, Database);
 	}
 
-	virtual void ResetMarkersRecursively() override
+	virtual void ResetMarkersRecursively() const override
 	{
 		ResetMarkers();
 		ResetMarkersRecursivelyOnEntities((TArray<TOrientedEntity<FEntity>>&) TopologicalFaces);
@@ -93,7 +98,7 @@ public:
 	virtual void Empty() override;
 
 	void Add(TSharedRef<FTopologicalFace> InTopologicalFace, EOrientation InOrientation);
-	void Add(TArray<FTopologicalFace*> Faces);
+	void Add(TArray<FTopologicalFace*>& Faces);
 
 #ifdef CADKERNEL_DEV
 	virtual FInfoEntity& GetInfo(FInfoEntity&) const override;
@@ -152,7 +157,7 @@ public:
 		return ((States & EHaveStates::IsInner) == EHaveStates::IsInner);
 	}
 
-	bool IsOutter() const
+	bool IsOuter() const
 	{
 		return ((States & EHaveStates::IsInner) != EHaveStates::IsInner);
 	}
@@ -162,7 +167,7 @@ public:
 		States |= EHaveStates::IsInner;
 	}
 
-	void SetOutter()
+	void SetOuter()
 	{
 		States &= ~EHaveStates::IsInner;
 	}
diff --git a/Engine/Source/Runtime/Datasmith/CADKernel/Public/CADKernel/Topo/TopologicalEdge.h b/Engine/Source/Runtime/Datasmith/CADKernel/Public/CADKernel/Topo/TopologicalEdge.h
--- a/Engine/Source/Runtime/Datasmith/CADKernel/Public/CADKernel/Topo/TopologicalEdge.h
+++ b/Engine/Source/Runtime/Datasmith/CADKernel/Public/CADKernel/Topo/TopologicalEdge.h
@@ -6,6 +6,7 @@
 #include "CADKernel/Geo/Sampling/SurfacicPolyline.h"
 #include "CADKernel/Geo/GeoEnum.h"
 #include "CADKernel/Math/Boundary.h"
+#include "CADKernel/Mesh/MeshEnum.h"
 #include "CADKernel/Topo/Linkable.h"
 #include "CADKernel/Topo/TopologicalVertex.h"
 #include "CADKernel/Utils/Cache.h"
@@ -16,16 +17,8 @@ namespace UE::CADKernel
 typedef TTopologicalLink<FTopologicalEdge> FEdgeLink;
 
 class FOrientedEdge;
-
-enum ECoordinateType : uint8
-{
-	VertexCoordinate,
-	ImposedCoordinate, // imposted cutting point by thin zone meshing process
-	IsoUCoordinate,
-	IsoVCoordinate,
-	IsoUVCoordinate,
-	OtherCoordinate
-};
+class FThinZone2D;
+class FThinZoneSide;
 
 /**
  * Cutting point used for thin zone purpose
@@ -35,16 +28,18 @@ struct CADKERNEL_API FImposedCuttingPoint
 	/**
 	 * coordinate of the edge's mesh nodes
 	 */
-	double Coordinate = 0;
-	int32 OppositNodeIndex = -1;
+	const double Coordinate = 0;
+	const int32 OppositNodeIndex = -1;
+	double DeltaU = 0;
 
 	FImposedCuttingPoint()
 	{
 	}
 
-	FImposedCuttingPoint(double InCoordinate, int32 NodeIndex1)
+	FImposedCuttingPoint(const double InCoordinate, const int32 NodeIndex1, const double InDeltaU = 0.)
 		: Coordinate(InCoordinate)
 		, OppositNodeIndex(NodeIndex1)
+		, DeltaU(InDeltaU)
 	{
 	};
 };
@@ -52,12 +47,15 @@ struct CADKERNEL_API FImposedCuttingPoint
 template<typename FCuttingPointType>
 void GetCuttingPointCoordinates(const TArray<FCuttingPointType>& CuttingPoints, TArray<double>& CuttingPointCoordinates);
 
+using FAddCuttingPointFunc = TFunction<void(const double, const ECoordinateType, const FPairOfIndex, const double)>;
+
 struct FEdge2DProperties;
 struct FCuttingPoint;
 
 class FModelMesh;
 class FEdgeMesh;
 class FSurface;
+class FThinZone;
 class FTopologicalLoop;
 class FTopologicalVertex;
 
@@ -67,6 +65,9 @@ class CADKERNEL_API FTopologicalEdge : public TLinkable<FTopologicalEdge, FEdgeL
 	friend class FTopologicalLoop;
 	friend class FTopologicalFace;
 
+private:
+	const double FactorToComputeMaxTol = 0.1;
+
 protected:
 
 	TSharedPtr<FTopologicalVertex> StartVertex;
@@ -80,6 +81,9 @@ protected:
 	TSharedPtr<FRestrictionCurve> Curve;
 	mutable double Length3D = -1.;
 
+	// To avoid huge tolerance in case of degenerated edge, the max tol is defined as Length3D / 10.
+	mutable double Max2DTolerance = -1;
+
 	FTopologicalLoop* Loop = nullptr;
 
 	TSharedPtr<FEdgeMesh> Mesh;
@@ -93,6 +97,8 @@ protected:
 	 * U coordinates of the edge's mesh nodes for thin zone purpose
 	 */
 	TArray<FImposedCuttingPoint> ImposedCuttingPointUs;
+
+	TArray<FThinZoneSide*> ThinZoneSides;
 	TArray<FLinearBoundary> ThinZoneBounds;
 
 	/**
@@ -174,14 +180,14 @@ public:
 		Ar << Boundary;
 		SerializeIdent(Ar, &Loop);
 		Ar << Length3D;
+		Max2DTolerance = Length3D * FactorToComputeMaxTol;
 	}
 
 	virtual void SpawnIdent(FDatabase& Database) override;
 
-	virtual void ResetMarkersRecursively() override
+	virtual void ResetMarkersRecursively() const override
 	{
-		ResetMarkers();
-		TLinkable<FTopologicalEdge, FEdgeLink>::ResetMarkers();
+		TLinkable<FTopologicalEdge, FEdgeLink>::ResetMarkersRecursively();
 		StartVertex->ResetMarkersRecursively();
 		EndVertex->ResetMarkersRecursively();
 
@@ -199,7 +205,7 @@ public:
 
 	double GetTolerance2DAt(double Coordinate) const
 	{
-		return GetCurve()->GetToleranceAt(Coordinate);
+		return FMath::Min(Max2DTolerance, GetCurve()->GetToleranceAt(Coordinate));
 	}
 
 	virtual EEntity GetEntityType() const override
@@ -268,6 +274,16 @@ public:
 		return FirstTwinEdge;
 	}
 
+	FTopologicalEdge* GetTwinEdge() const
+	{
+		const TArray<FTopologicalEdge*>& TwinEdges = GetTwinEntities();
+		if (TwinEdges.Num() > 1)
+		{
+			return TwinEdges[0] == this ? TwinEdges[1] : TwinEdges[0];
+		}
+		return nullptr;
+	}
+
 	/**
 	 * @return true if the twin edge is in the same direction as this
 	 */
@@ -452,17 +468,29 @@ public:
 
 	// ======   Meshing Function   ======
 
-	const TSharedRef<FEdgeMesh> GetMesh() const
+	const FEdgeMesh* GetMesh() const
 	{
 		if (GetLinkActiveEntity() != AsShared())
 		{
 			return GetLinkActiveEdge()->GetMesh();
 		}
 		ensureCADKernel(Mesh.IsValid());
-		return Mesh.ToSharedRef();
+		return Mesh.Get();
 	}
 
-	TSharedRef<FEdgeMesh> GetOrCreateMesh(FModelMesh& MeshModel);
+	FEdgeMesh& GetOrCreateMesh(FModelMesh& MeshModel);
+
+	/**
+	 * If the mesh of the edge is not built, Empty the CuttingPoints 
+	 * This allows to recompute a new discretization of the mesh based among other things on a new imposed cutting points (mesh of thin zone process)
+	 */
+	void RemovePreMesh();
+
+	const FTopologicalEdge* GetPreMeshedTwin() const;
+	FTopologicalEdge* GetPreMeshedTwin()
+	{
+		return const_cast<FTopologicalEdge*> (static_cast<const FTopologicalEdge*>(this)->GetPreMeshedTwin());
+	}
 
 	/**
 	 * Generate a sampling of the curve.
@@ -537,20 +565,40 @@ public:
 		return CuttingPointUs;
 	}
 
+	TArray<double> GetCuttingPointCoordinates() const;
+
+	void TransferCuttingPointFromMeshedEdge(bool bOnlyWithOppositeNode, FAddCuttingPointFunc AddCuttingPoint);
+
+	/**
+	 * Compute the lengths of each pre-elements of the edge i.e the elements based of the cutting points of the edges. 
+	 */
+	TArray<double> GetPreElementLengths() const;
+
 	// For thin zone purpose
-	void SortImposedCuttingPoints()
-	{
-		Algo::Sort(ImposedCuttingPointUs, [](const FImposedCuttingPoint& C1, const FImposedCuttingPoint& C2) { return C1.Coordinate < C2.Coordinate; });
-	}
+	void SortImposedCuttingPoints();
 
 	const TArray<FImposedCuttingPoint>& GetImposedCuttingPoints() const
 	{
 		return ImposedCuttingPointUs;
 	}
 
-	void AddThinZone(const FLinearBoundary& InThinZoneBounds)
+	void AddThinZone(FThinZoneSide* InThinZoneSide, const FLinearBoundary& InThinZoneBounds)
 	{
-		ThinZoneBounds.Add(InThinZoneBounds);
+		if(InThinZoneSide)
+		{
+			ThinZoneSides.AddUnique(InThinZoneSide);
+			ThinZoneBounds.Add(InThinZoneBounds);
+		}
+	}
+
+	int32 GetThinZoneCount() const
+	{
+		return ThinZoneSides.Num();
+	}
+
+	const TArray<FThinZoneSide*>& GetThinZoneSides() const
+	{
+		return ThinZoneSides;
 	}
 
 	const TArray<FLinearBoundary>& GetThinZoneBounds() const
@@ -558,8 +606,10 @@ public:
 		return ThinZoneBounds;
 	}
 
-	void AddImposedCuttingPointU(const double ImposedCuttingPointU, int32 OppositeNodeIndex);
+	void AddImposedCuttingPointU(const double ImposedCuttingPointU, const int32 OppositeNodeIndex, const double DeltaU);
+	void AddTwinsCuttingPoint(const double Coord, const double DeltaU);
 
+	void GenerateMeshElements(FModelMesh& MeshModel);
 
 	// ======   Curve Functions   ======
 
@@ -573,6 +623,7 @@ public:
 		return Curve.ToSharedRef();
 	}
 
+	void ComputeLength();
 	double Length() const;
 
 	/**
@@ -668,7 +719,8 @@ public:
 	 */
 	void ProjectTwinEdgePoints(const TArray<FPoint>& InPointsToProject, bool bSameOrientation, TArray<double>& OutProjectedPointCoords) const
 	{
-		Curve->ProjectTwinCurvePoints(InPointsToProject, bSameOrientation, OutProjectedPointCoords);
+		const double ToleranceOfProjection = Length3D * 0.1;
+		Curve->ProjectTwinCurvePoints(InPointsToProject, bSameOrientation, OutProjectedPointCoords, ToleranceOfProjection);
 	}
 
 	/**
@@ -692,7 +744,9 @@ public:
 
 	double TransformLocalCoordinateToActiveEdgeCoordinate(const double LocalCoordinate) const;
 	double TransformActiveEdgeCoordinateToLocalCoordinate(const double ActiveEdgeCoordinate) const;
+	double TransformTwinEdgeCoordinateToLocalCoordinate(const FTopologicalEdge& TwinEdge, const double InTwinCoordinate) const;
 
+	void TransformTwinEdgeCoordinatesToLocalCoordinates(const FTopologicalEdge& TwinEdge, const TArray<double>& InActiveEdgeCoordinate, TArray<double>& OutLocalCoordinate) const;
 	void TransformActiveEdgeCoordinatesToLocalCoordinates(const TArray<double>& InActiveEdgeCoordinate, TArray<double>& OutLocalCoordinate) const;
 	void TransformLocalCoordinatesToActiveEdgeCoordinates(const TArray<double>& InLocalCoordinate, TArray<double>& OutActiveEdgeCoordinate) const;
 
@@ -893,39 +947,39 @@ struct CADKERNEL_API FCuttingPoint
 	/**
 	 * coordinate of the edge's mesh nodes
 	 */
-	double Coordinate = 0;
-	ECoordinateType Type = ECoordinateType::OtherCoordinate;
-	int32 OppositNodeIndex = -1;
-	int32 OppositNodeIndex2 = -1;
+	double Coordinate;
+	ECoordinateType Type;
+	FPairOfIndex OppositNodeIndices;
 	double IsoDeltaU = 0;
 
 	FCuttingPoint()
+		: Coordinate(0)
+		, Type(ECoordinateType::OtherCoordinate)
+		, OppositNodeIndices(FPairOfIndex::Undefined)
+		, IsoDeltaU(HUGE_VAL)
 	{
 	}
 
 	FCuttingPoint(double InCoordinate, ECoordinateType InType)
 		: Coordinate(InCoordinate)
 		, Type(InType)
-		, OppositNodeIndex(-1)
-		, OppositNodeIndex2(-1)
+		, OppositNodeIndices(FPairOfIndex::Undefined)
 		, IsoDeltaU(HUGE_VAL)
 	{
 	}
 
-	FCuttingPoint(double InCoordinate, ECoordinateType InType, int32 NodeIndex1, int32 NodeIndex2)
+	FCuttingPoint(double InCoordinate, ECoordinateType InType, FPairOfIndex InOppositNodeIndices, double DeltaU)
 		: Coordinate(InCoordinate)
 		, Type(InType)
-		, OppositNodeIndex(NodeIndex1)
-		, OppositNodeIndex2(NodeIndex2)
-		, IsoDeltaU(HUGE_VAL)
+		, OppositNodeIndices(InOppositNodeIndices)
+		, IsoDeltaU(DeltaU)
 	{
 	}
 
 	FCuttingPoint(double InCoordinate, ECoordinateType InType, int32 InOppositeNodeId, double DeltaU)
 		: Coordinate(InCoordinate)
 		, Type(InType)
-		, OppositNodeIndex(InOppositeNodeId)
-		, OppositNodeIndex2(-1)
+		, OppositNodeIndices(InOppositeNodeId)
 		, IsoDeltaU(DeltaU)
 	{
 	}
diff --git a/Engine/Source/Runtime/Datasmith/CADKernel/Public/CADKernel/Topo/TopologicalEntity.h b/Engine/Source/Runtime/Datasmith/CADKernel/Public/CADKernel/Topo/TopologicalEntity.h
--- a/Engine/Source/Runtime/Datasmith/CADKernel/Public/CADKernel/Topo/TopologicalEntity.h
+++ b/Engine/Source/Runtime/Datasmith/CADKernel/Public/CADKernel/Topo/TopologicalEntity.h
@@ -63,19 +63,35 @@ public:
 		return ((States & MeshedOrDeletedOrDegenerated) == EHaveStates::None);
 	}
 
+	/** 
+	 * An edge is premesh if the cutting points of the mesh have been defined but the mesh is not build 
+	 * This is used during thin zone meshing i.e. due to the process, an edge can be premeshed but not finalized
+	 */
+	bool IsPreMeshed() const
+	{
+		constexpr EHaveStates MeshedOrPreMeshed = EHaveStates::IsMeshed | EHaveStates::IsPreMeshed;
+		return ((States & MeshedOrPreMeshed) != EHaveStates::None);
+	}
+
 	bool IsMeshed() const
 	{
 		return ((States & EHaveStates::IsMeshed) == EHaveStates::IsMeshed);
 	}
 
-	virtual void SetMeshed()
+	virtual void SetPreMeshedMarker()
+	{
+		States |= EHaveStates::IsPreMeshed;
+	}
+
+	virtual void SetMeshedMarker()
 	{
-		States |= EHaveStates::IsMeshed;
+		constexpr EHaveStates MeshedOrPreMeshed = EHaveStates::IsMeshed | EHaveStates::IsPreMeshed;
+		States |= MeshedOrPreMeshed;
 	}
 
-	virtual void ResetMeshed()
+	virtual void ResetPreMeshed()
 	{
-		States &= ~EHaveStates::IsMeshed;
+		States &= ~EHaveStates::IsPreMeshed;
 	}
 };
 
diff --git a/Engine/Source/Runtime/Datasmith/CADKernel/Public/CADKernel/Topo/TopologicalFace.h b/Engine/Source/Runtime/Datasmith/CADKernel/Public/CADKernel/Topo/TopologicalFace.h
--- a/Engine/Source/Runtime/Datasmith/CADKernel/Public/CADKernel/Topo/TopologicalFace.h
+++ b/Engine/Source/Runtime/Datasmith/CADKernel/Public/CADKernel/Topo/TopologicalFace.h
@@ -122,7 +122,7 @@ public:
 
 	virtual void SpawnIdent(FDatabase& Database) override;
 
-	virtual void ResetMarkersRecursively() override
+	virtual void ResetMarkersRecursively() const override
 	{
 		ResetMarkers();
 		ResetMarkersRecursivelyOnEntities(Loops);
@@ -356,7 +356,7 @@ public:
 
 	// ======   Meshing Function   ======
 
-	TSharedRef<FFaceMesh> GetOrCreateMesh(FModelMesh& ModelMesh);
+	FFaceMesh& GetOrCreateMesh(FModelMesh& ModelMesh);
 
 	const bool HasTesselation() const
 	{
@@ -440,7 +440,7 @@ public:
 		return ((States & EHaveStates::ThinZone) == EHaveStates::ThinZone);
 	}
 
-	void SetHasThinZone()
+	void SetHasThinZoneMarker()
 	{
 		States |= EHaveStates::ThinZone;
 	}
@@ -624,7 +624,10 @@ public:
 
 	void SetEstimatedMinimalElementLength(double Value)
 	{
-		EstimatedMinimalElementLength = Value;
+		if(Value < EstimatedMinimalElementLength)
+		{
+			EstimatedMinimalElementLength = Value;
+		}
 	}
 
 	/**
diff --git a/Engine/Source/Runtime/Datasmith/CADKernel/Public/CADKernel/Topo/TopologicalLink.h b/Engine/Source/Runtime/Datasmith/CADKernel/Public/CADKernel/Topo/TopologicalLink.h
--- a/Engine/Source/Runtime/Datasmith/CADKernel/Public/CADKernel/Topo/TopologicalLink.h
+++ b/Engine/Source/Runtime/Datasmith/CADKernel/Public/CADKernel/Topo/TopologicalLink.h
@@ -188,6 +188,17 @@ public:
 		}
 		return false;
 	}
+
+	void ResetMarkersRecursively() const
+	{
+		for (EntityType* Entity : TwinEntities)
+		{
+			if (Entity)
+			{
+				Entity->ResetMarkers();
+			}
+		}
+	}
 };
 
 /**
diff --git a/Engine/Source/Runtime/Datasmith/CADKernel/Public/CADKernel/Topo/TopologicalLoop.h b/Engine/Source/Runtime/Datasmith/CADKernel/Public/CADKernel/Topo/TopologicalLoop.h
--- a/Engine/Source/Runtime/Datasmith/CADKernel/Public/CADKernel/Topo/TopologicalLoop.h
+++ b/Engine/Source/Runtime/Datasmith/CADKernel/Public/CADKernel/Topo/TopologicalLoop.h
@@ -26,6 +26,11 @@ public:
 	{
 	}
 
+	FOrientedEdge(const FOrientedEdge& OrientiredEntity)
+		: TOrientedEntity(OrientiredEntity)
+	{
+	}
+
 	FOrientedEdge()
 		: TOrientedEntity()
 	{
@@ -109,7 +114,7 @@ public:
 		SpawnIdentOnEntities((TArray<TOrientedEntity<FEntity>>&) Edges, Database);
 	}
 
-	virtual void ResetMarkersRecursively() override
+	virtual void ResetMarkersRecursively() const override
 	{
 		ResetMarkers();
 		ResetMarkersRecursivelyOnEntities((TArray<TOrientedEntity<FEntity>>&) Edges);
diff --git a/Engine/Source/Runtime/Datasmith/CADKernel/Public/CADKernel/Topo/TopologicalVertex.h b/Engine/Source/Runtime/Datasmith/CADKernel/Public/CADKernel/Topo/TopologicalVertex.h
--- a/Engine/Source/Runtime/Datasmith/CADKernel/Public/CADKernel/Topo/TopologicalVertex.h
+++ b/Engine/Source/Runtime/Datasmith/CADKernel/Public/CADKernel/Topo/TopologicalVertex.h
@@ -62,12 +62,6 @@ public:
 
 	virtual void SpawnIdent(FDatabase& Database) override;
 
-	virtual void ResetMarkersRecursively() override
-	{
-		ResetMarkers();
-	}
-
-
 #ifdef CADKERNEL_DEV
 	virtual FInfoEntity& GetInfo(FInfoEntity&) const override;
 #endif
@@ -135,15 +129,15 @@ public:
 		return Coordinates.SquareDistance(Point);
 	}
 
-	TSharedRef<FVertexMesh> GetOrCreateMesh(FModelMesh& MeshModel);
+	FVertexMesh& GetOrCreateMesh(FModelMesh& MeshModel);
 
-	const TSharedRef<FVertexMesh> GetMesh() const
+	const FVertexMesh* GetMesh() const
 	{
 		if (!IsActiveEntity())
 		{
 			return GetLinkActiveEntity()->GetMesh();
 		}
-		return Mesh.ToSharedRef();
+		return Mesh.Get();
 	}
 
 	void Link(FTopologicalVertex& InEntity);
diff --git a/Engine/Source/Runtime/Datasmith/CADKernel/Public/CADKernel/Topo/Topomaker.h b/Engine/Source/Runtime/Datasmith/CADKernel/Public/CADKernel/Topo/Topomaker.h
--- a/Engine/Source/Runtime/Datasmith/CADKernel/Public/CADKernel/Topo/Topomaker.h
+++ b/Engine/Source/Runtime/Datasmith/CADKernel/Public/CADKernel/Topo/Topomaker.h
@@ -127,6 +127,11 @@ public:
 	 */
 	void UnlinkNonManifoldVertex();
 
+	/**
+	 * Mandatory: UnlinkNonManifoldVertex has to be call before
+	 */
+	void UnlinkEdgesOfSeparateShells();
+
 	void RemoveThinFaces(TArray<FTopologicalEdge*>& NewBorderEdges);
 
 #ifdef CADKERNEL_DEV
@@ -163,17 +168,17 @@ private:
 	/**
 	 * Return an array of active vertices.
 	 */
-	void GetVertices(TArray<TSharedPtr<FTopologicalVertex>>& Vertices);
+	void GetVertices(TArray<FTopologicalVertex*>& Vertices);
 
 	/**
 	 * Return an array of active border vertices.
 	 */
-	void GetBorderVertices(TArray<TSharedPtr<FTopologicalVertex>>& BorderVertices);
+	void GetBorderVertices(TArray<FTopologicalVertex*>& BorderVertices);
 
-	void CheckSelfConnectedEdge(double MaxLengthOfDegeneratedEdge, TArray<TSharedPtr<FTopologicalVertex>>& OutBorderVertices);
+	void CheckSelfConnectedEdge(double MaxLengthOfDegeneratedEdge, TArray<FTopologicalVertex*>& OutBorderVertices);
 	void RemoveIsolatedEdges(); // Useful ???
 
-	void SetSelfConnectedEdgeDegenerated(TArray<TSharedPtr<FTopologicalVertex>>& Vertices);
+	void SetSelfConnectedEdgeDegenerated(TArray<FTopologicalVertex*>& Vertices);
 
 	/**
 	 * For each loop of each surface, check if successive edges are unconnected and if their common vertices are connected only to them (UV vs CV: vertex connected to many faces ).
@@ -194,6 +199,11 @@ private:
 	 */
 	void MergeUnconnectedSuccessiveEdges();
 
+	void SetToProcessMarkerOfFaces();
+	void ResetMarkersOfFaces();
+
+private:
+
 	ESewOption SewOptions;
 
 	double Tolerance;
