CADKernel update: Meshing of thin faces

diff --git a/Engine/Source/Runtime/Datasmith/CADKernel/Public/CADKernel/UI/Display.h b/Engine/Source/Runtime/Datasmith/CADKernel/Public/CADKernel/UI/Display.h
--- a/Engine/Source/Runtime/Datasmith/CADKernel/Public/CADKernel/UI/Display.h
+++ b/Engine/Source/Runtime/Datasmith/CADKernel/Public/CADKernel/UI/Display.h
@@ -242,6 +242,10 @@ CADKERNEL_API void DisplayMesh(const FFaceMesh& Mesh);
 CADKERNEL_API void DisplayMesh(const FEdgeMesh& Mesh);
 CADKERNEL_API void DisplayMesh(const FVertexMesh& Mesh);
 
+CADKERNEL_API void DisplayMesh(const FModel& Model);
+CADKERNEL_API void DisplayMesh(const FBody& Body);
+CADKERNEL_API void DisplayMesh(const FShell &Shel);
+
 CADKERNEL_API void Display(const FModelMesh& MeshModel);
 
 CADKERNEL_API void DisplayControlPolygon(const FCurve& Entity);
@@ -311,7 +315,7 @@ void DisplayOrientedPolyline(const TArray<TPoint>& Points, EVisuProperty Propert
 #endif
 }
 
-CADKERNEL_API void DrawQuadripode(double Height, double Base, FPoint& Centre, FPoint& Direction);
+CADKERNEL_API void DrawQuadripode(double Height, double Base, FPoint& Centre, FPoint& Direction, EVisuProperty Property);
 
 CADKERNEL_API void Draw(const FTopologicalEdge& Edge, EVisuProperty Property = EVisuProperty::BlueCurve);
 CADKERNEL_API void Draw(const FTopologicalFace& Face);
@@ -344,7 +348,7 @@ void DrawSegmentOrientation(const TPoint& Point1, const TPoint& Point2, EVisuPro
 
 	FPoint Middle = (Point1 + Point2) / 2.;
 	FPoint Tangent = Point2 - Point1;
-	DrawQuadripode(Height, Base, Middle, Tangent);
+	DrawQuadripode(Height, Base, Middle, Tangent, Property);
 #endif
 }
 
diff --git a/Engine/Source/Runtime/Datasmith/CADKernel/Public/CADKernel/UI/Message.h b/Engine/Source/Runtime/Datasmith/CADKernel/Public/CADKernel/UI/Message.h
--- a/Engine/Source/Runtime/Datasmith/CADKernel/Public/CADKernel/UI/Message.h
+++ b/Engine/Source/Runtime/Datasmith/CADKernel/Public/CADKernel/UI/Message.h
@@ -67,6 +67,11 @@ public:
 		VReportPrintF(Header, TEXT("%lld"), Time);
 	}
 
+	static void FillReportFile(const TCHAR* Header, double Count)
+	{
+		VReportPrintF(Header, TEXT("%f"), Count);
+	}
+
 	static void FillReportFile(const TCHAR* Header, int32 Count)
 	{
 		VReportPrintF(Header, TEXT("%d"), Count);
