Fix NonUnity Compile UnrealGame Android

diff --git a/Engine/Source/Runtime/Datasmith/CADKernel/Private/CADKernel/Mesh/Meshers/CycleTriangulator.cpp b/Engine/Source/Runtime/Datasmith/CADKernel/Private/CADKernel/Mesh/Meshers/CycleTriangulator.cpp
--- a/Engine/Source/Runtime/Datasmith/CADKernel/Private/CADKernel/Mesh/Meshers/CycleTriangulator.cpp
+++ b/Engine/Source/Runtime/Datasmith/CADKernel/Private/CADKernel/Mesh/Meshers/CycleTriangulator.cpp
@@ -1,5 +1,4 @@
 // Copyright Epic Games, Inc. All Rights Reserved.
-#pragma once
 #include "CADKernel/Mesh/Meshers/CycleTriangulator.h"
 
 #include "CADKernel/Core/Factory.h"
@@ -380,9 +379,6 @@ bool FCycleTriangulator::FindCandidateNodes(int32 StartIndex)
 	FirstSideStartNode = SubCycleNodes[FirstSideStartIndex];
 	FirstSideEndNode = SubCycleNodes[FirstSideEndIndex];
 
-	//const SlopeMethod GetSlopeAtStartNode = CounterClockwiseSlope;
-	//const SlopeMethod GetSlopeAtEndNode = ClockwiseSlope;
-
 	const FPoint2D& StartPoint2D = FirstSideStartNode->Get2DPoint(EGridSpace::UniformScaled, Grid);
 	const FPoint2D& EndPoint2D = FirstSideEndNode->Get2DPoint(EGridSpace::UniformScaled, Grid);
 
