Fix Compilation Tests / Static Analysis Win64 (MSVC)

diff --git a/Engine/Source/Runtime/Datasmith/CADKernel/Private/CADKernel/Mesh/Structure/ThinZone2DFinder.cpp b/Engine/Source/Runtime/Datasmith/CADKernel/Private/CADKernel/Mesh/Structure/ThinZone2DFinder.cpp
--- a/Engine/Source/Runtime/Datasmith/CADKernel/Private/CADKernel/Mesh/Structure/ThinZone2DFinder.cpp
+++ b/Engine/Source/Runtime/Datasmith/CADKernel/Private/CADKernel/Mesh/Structure/ThinZone2DFinder.cpp
@@ -436,7 +436,7 @@ bool FThinZone2DFinder::CheckIfCloseSideOfThinSideIsNotDegenerated(TArray<FEdgeS
 	F3DDebugSession _(bDisplay, ("Check Thin Zone"));
 #endif
 
-	const FEdgeSegment* CloseSegments[3] = { nullptr, nullptr,nullptr };
+	const FEdgeSegment* CloseSegments[3] = { nullptr, nullptr, nullptr };
 	TFunction<FPoint2D(const FPoint2D&)> GetOppositeSegment = [&CloseSegments](const FPoint2D& Candidate) -> FPoint2D
 	{
 		FPoint2D OutOppositePoint;
@@ -469,6 +469,11 @@ bool FThinZone2DFinder::CheckIfCloseSideOfThinSideIsNotDegenerated(TArray<FEdgeS
 	// find the close segments
 	for (const FEdgeSegment* Segment : Side)
 	{
+		if (Segment == nullptr)
+		{
+			continue;
+		}
+
 		const FEdgeSegment* CloseSegment = Segment->GetCloseSegment();
 		if (!CloseSegments[0])
 		{
@@ -495,7 +500,7 @@ bool FThinZone2DFinder::CheckIfCloseSideOfThinSideIsNotDegenerated(TArray<FEdgeS
 		return true;
 	}
 
-	if (CloseSegments[1] != nullptr)
+	if (CloseSegments[0] != nullptr && CloseSegments[1] != nullptr)
 	{
 		// Case two chain -> not degenerated 
 		// 
@@ -523,7 +528,12 @@ bool FThinZone2DFinder::CheckIfCloseSideOfThinSideIsNotDegenerated(TArray<FEdgeS
 
 	TFunction<double()> ComputeTolerance = [&CloseSegments]() -> double
 	{
+		if (CloseSegments[0] == nullptr)
+		{
+			return 0;
+		}
 		double CloseLength = CloseSegments[0]->GetLength();
+
 		const double Ratio = CloseSegments[1] ? 12. : 6.;
 		if (CloseSegments[1])
 		{
@@ -533,6 +543,11 @@ bool FThinZone2DFinder::CheckIfCloseSideOfThinSideIsNotDegenerated(TArray<FEdgeS
 	};
 	const double ToleranceCloseSquare = ComputeTolerance();
 
+	if (Side[0] == nullptr)
+	{
+		return false;
+	}
+
 	const FPoint2D& StartPoint = Side[0]->GetExtemity(Start);
 	FPoint2D FirstOppositePoint = GetOppositeSegment(StartPoint);
 
@@ -547,6 +562,10 @@ bool FThinZone2DFinder::CheckIfCloseSideOfThinSideIsNotDegenerated(TArray<FEdgeS
 
 	for (const FEdgeSegment* Segment : Side)
 	{
+		if (Segment == nullptr)
+		{
+			continue;
+		}
 		const FPoint2D& NextPoint = Segment->GetExtemity(End);
 		FPoint2D OppositePoint = GetOppositeSegment(NextPoint);
 		const double SquarewDistance = OppositePoint.SquareDistance(FirstOppositePoint);
@@ -569,10 +588,10 @@ bool FThinZone2DFinder::CheckIfCloseSideOfThinSideIsNotDegenerated(TArray<FEdgeS
 
 	// Degenerated... 
 	// Check if at least one segment of the side is the close segment of the other side
-	const FEdgeSegment* CloseOfCloseSegments[2] = { CloseSegments[0]->GetCloseSegment(), CloseSegments[1] ? CloseSegments[1]->GetCloseSegment() : nullptr };
+	const FEdgeSegment* CloseOfCloseSegments[2] = { CloseSegments[0] ? CloseSegments[0]->GetCloseSegment() : nullptr, CloseSegments[1] ? CloseSegments[1]->GetCloseSegment() : nullptr };
 	for (const FEdgeSegment* Segment : Side)
 	{
-		if (Segment == CloseOfCloseSegments[0] || Segment == CloseOfCloseSegments[1])
+		if (Segment && (Segment == CloseOfCloseSegments[0] || Segment == CloseOfCloseSegments[1]))
 		{
 			return true;
 		}
