#rn none

diff --git a/Engine/Source/Runtime/Datasmith/CADKernel/Public/CADKernel/Math/Boundary.h b/Engine/Source/Runtime/Datasmith/CADKernel/Public/CADKernel/Math/Boundary.h
--- a/Engine/Source/Runtime/Datasmith/CADKernel/Public/CADKernel/Math/Boundary.h
+++ b/Engine/Source/Runtime/Datasmith/CADKernel/Public/CADKernel/Math/Boundary.h
@@ -248,6 +248,11 @@ public:
 		Set(Point1, Point2);
 	}
 
+	FSurfacicBoundary(const FSurfacicBoundary& SB)
+		: UVBoundaries { FLinearBoundary(SB.UVBoundaries[EIso::IsoU]), FLinearBoundary(SB.UVBoundaries[EIso::IsoV]) }
+	{
+	}
+
 	void Set(const FPoint2D& Point1, const FPoint2D& Point2)
 	{
 		UVBoundaries[EIso::IsoU].Set(Point1.U, Point2.U);
@@ -407,6 +412,13 @@ public:
 		UVBoundaries[EIso::IsoV].MoveInsideIfNot(Point.V, Tolerance);
 	}
 
+	FSurfacicBoundary& operator=(const FSurfacicBoundary& InBounds)
+	{
+		UVBoundaries[EIso::IsoU] = InBounds.UVBoundaries[EIso::IsoU];
+		UVBoundaries[EIso::IsoV] = InBounds.UVBoundaries[EIso::IsoV];
+		return *this;
+	}
+
 	double Length(const EIso& Iso) const
 	{
 		return UVBoundaries[Iso].Length();
