CADKernel update for Meshing of thin faces, Stitching = Sew

diff --git a/Engine/Plugins/Enterprise/DatasmithCADImporter/Source/CADKernelSurface/Public/CADModelToCADKernelConverterBase.h b/Engine/Plugins/Enterprise/DatasmithCADImporter/Source/CADKernelSurface/Public/CADModelToCADKernelConverterBase.h
--- a/Engine/Plugins/Enterprise/DatasmithCADImporter/Source/CADKernelSurface/Public/CADModelToCADKernelConverterBase.h
+++ b/Engine/Plugins/Enterprise/DatasmithCADImporter/Source/CADKernelSurface/Public/CADModelToCADKernelConverterBase.h
@@ -47,6 +47,7 @@ public:
 
 			UE::CADKernel::FTopomaker Topomaker(CADKernelSession, TopomakerOptions);
 			Topomaker.Sew();
+			Topomaker.SplitIntoConnectedShells();
 			Topomaker.OrientShells();
 		}
 
diff --git a/Engine/Plugins/Enterprise/DatasmithCADImporter/Source/DatasmithDispatcher/Private/DatasmithDispatcherNetworking.cpp b/Engine/Plugins/Enterprise/DatasmithCADImporter/Source/DatasmithDispatcher/Private/DatasmithDispatcherNetworking.cpp
--- a/Engine/Plugins/Enterprise/DatasmithCADImporter/Source/DatasmithDispatcher/Private/DatasmithDispatcherNetworking.cpp
+++ b/Engine/Plugins/Enterprise/DatasmithCADImporter/Source/DatasmithDispatcher/Private/DatasmithDispatcherNetworking.cpp
@@ -5,6 +5,8 @@
 #include "DatasmithDispatcherLog.h"
 
 #include "CoreMinimal.h"
+#include "HAL/PlatformTime.h"
+#include "HAL/PlatformProcess.h"
 #include "SocketSubsystem.h"
 #include "Sockets.h"
 #include "Serialization/MemoryReader.h"
diff --git a/Engine/Source/Runtime/Datasmith/CADKernel/Private/CADKernel/Geo/Sampler/SamplerAbstract.h b/Engine/Source/Runtime/Datasmith/CADKernel/Private/CADKernel/Geo/Sampler/SamplerAbstract.h
--- a/Engine/Source/Runtime/Datasmith/CADKernel/Private/CADKernel/Geo/Sampler/SamplerAbstract.h
+++ b/Engine/Source/Runtime/Datasmith/CADKernel/Private/CADKernel/Geo/Sampler/SamplerAbstract.h
@@ -274,7 +274,6 @@ protected:
 	{
 		if (NeededPointNumber > 50)
 		{
-			FMessage::Printf(EVerboseLevel::Debug, TEXT("SamplerAbstract needed point number %d\n"), NeededPointNumber);
 			NeededPointNumber = 50;
 		}
 
diff --git a/Engine/Source/Runtime/Datasmith/CADKernel/Private/CADKernel/Math/Geometry.cpp b/Engine/Source/Runtime/Datasmith/CADKernel/Private/CADKernel/Math/Geometry.cpp
--- a/Engine/Source/Runtime/Datasmith/CADKernel/Private/CADKernel/Math/Geometry.cpp
+++ b/Engine/Source/Runtime/Datasmith/CADKernel/Private/CADKernel/Math/Geometry.cpp
@@ -67,12 +67,12 @@ void FindLoopIntersectionsWithIso(const EIso Iso, const double IsoParameter, con
 
 bool DoIntersectInside(const FSegment2D& SegmentAB, const FSegment2D& SegmentCD)
 {
-	constexpr const double Min = DOUBLE_SMALL_NUMBER;
-	constexpr const double Max = 1. - DOUBLE_SMALL_NUMBER;
+	constexpr const double Min = DOUBLE_KINDA_SMALL_NUMBER;
+	constexpr const double Max = 1. - DOUBLE_KINDA_SMALL_NUMBER;
 
 	TFunction<bool(double, double, double, double)> Intersect = [](double A, double B, double C, double D)
 	{
-		return !((D < A + DOUBLE_SMALL_NUMBER) || (B < C + DOUBLE_SMALL_NUMBER));
+		return !((D < A + DOUBLE_KINDA_SMALL_NUMBER) || (B < C + DOUBLE_KINDA_SMALL_NUMBER));
 	};
 
 	TFunction<bool(double, double, double, double)> TestWhenParallel = [&](double A, double B, double C, double D)
@@ -106,10 +106,10 @@ bool DoIntersectInside(const FSegment2D& SegmentAB, const FSegment2D& SegmentCD)
 	const FPoint2D CA = SegmentAB[0] - SegmentCD[0];
 
 	const double ParallelCoef = CD ^ AB;
-	if (FMath::IsNearlyZero(ParallelCoef))
+	if (FMath::IsNearlyZero(ParallelCoef, DOUBLE_KINDA_SMALL_NUMBER))
 	{
 		const double ParallelCoef2 = CA ^ AB;
-		if (!FMath::IsNearlyZero(ParallelCoef2))
+		if (!FMath::IsNearlyZero(ParallelCoef2, DOUBLE_KINDA_SMALL_NUMBER))
 		{
 			return false;
 		}
@@ -170,10 +170,10 @@ bool DoIntersect(const TSegment<FPoint2D>& SegmentAB, const TSegment<FPoint2D>&
 	const FPoint2D CA = SegmentAB[0] - SegmentCD[0];
 
 	const double ParallelCoef = CD ^ AB;
-	if (FMath::IsNearlyZero(ParallelCoef))
+	if (FMath::IsNearlyZero(ParallelCoef, DOUBLE_KINDA_SMALL_NUMBER))
 	{
 		const double ParallelCoef2 = CA ^ AB;
-		if (!FMath::IsNearlyZero(ParallelCoef2))
+		if (!FMath::IsNearlyZero(ParallelCoef2, DOUBLE_KINDA_SMALL_NUMBER))
 		{
 			return false;
 		}
diff --git a/Engine/Source/Runtime/Datasmith/CADKernel/Private/CADKernel/Mesh/Meshers/BowyerWatsonTriangulator.h b/Engine/Source/Runtime/Datasmith/CADKernel/Private/CADKernel/Mesh/Meshers/BowyerWatsonTriangulator.h
--- a/Engine/Source/Runtime/Datasmith/CADKernel/Private/CADKernel/Mesh/Meshers/BowyerWatsonTriangulator.h
+++ b/Engine/Source/Runtime/Datasmith/CADKernel/Private/CADKernel/Mesh/Meshers/BowyerWatsonTriangulator.h
@@ -556,6 +556,13 @@ public:
 		}
 	}
 
+	TArray<int32> GetOuterVertices() const
+	{
+		TSet<int32> OuterVertices;
+		GetOuterVertices(OuterVertices);
+		return OuterVertices.Array();
+	}
+
 	void GetMesh(TArray<int32>& Triangles)
 	{
 		Triangles.Reset(TriangleSet.Num() * 3);
diff --git a/Engine/Source/Runtime/Datasmith/CADKernel/Private/CADKernel/Mesh/Meshers/CycleTriangulator.cpp b/Engine/Source/Runtime/Datasmith/CADKernel/Private/CADKernel/Mesh/Meshers/CycleTriangulator.cpp
--- a/Engine/Source/Runtime/Datasmith/CADKernel/Private/CADKernel/Mesh/Meshers/CycleTriangulator.cpp
+++ b/Engine/Source/Runtime/Datasmith/CADKernel/Private/CADKernel/Mesh/Meshers/CycleTriangulator.cpp
@@ -22,7 +22,7 @@ FCycleTriangulator::FCycleTriangulator(FIsoTriangulator& IsoTriangulator, const
 	: Grid(IsoTriangulator.Grid)
 	, Cycle(InCycle)
 	, CycleOrientation(InCycleOrientation)
-	, InnerToOuterSegmentsIntersectionTool(IsoTriangulator.InnerToOuterSegmentsIntersectionTool)
+	, InnerToOuterIsoSegmentsIntersectionTool(IsoTriangulator.InnerToOuterIsoSegmentsIntersectionTool)
 	, Mesh(IsoTriangulator.Mesh)
 	, IsoSegmentFactory(IsoTriangulator.IsoSegmentFactory)
 	, CycleIntersectionTool(Grid)
@@ -32,11 +32,12 @@ FCycleTriangulator::FCycleTriangulator(FIsoTriangulator& IsoTriangulator, const
 bool FCycleTriangulator::MeshCycle()
 {
 #ifdef ADD_TRIANGLE_2D
+	static int32 CycleId = 0;
+	F3DDebugSession _(Grid.bDisplay, FString::Printf(TEXT("Mesh cycle %d"), CycleId++));
 	if (Grid.bDisplay)
 	{
-		static int32 CycleId = 0;
-		F3DDebugSession _(Grid.bDisplay, FString::Printf(TEXT("Mesh cycle %d"), CycleId++));
 		Grid.DisplayIsoSegments(TEXT("Cycle"), EGridSpace::UniformScaled, Cycle);
+		Wait(CycleId == 2);
 	}
 #endif
 
@@ -144,6 +145,7 @@ bool FCycleTriangulator::FindTheCycleToMesh(FIsoSegment* Segment, bool bOrientat
 
 	double MinLength = HUGE_VAL;
 	int32 Index = 0;
+	double SquareLengthSum = 0;
 	for (FIsoSegment* NextSegment = FindNextSegment(Segment, Node, GetSlopeAtEndNode);; NextSegment = FindNextSegment(NextSegment, Node, GetSlopeAtEndNode))
 	{
 		FIsoNode* NextNode = &NextSegment->GetOtherNode(Node);
@@ -163,6 +165,7 @@ bool FCycleTriangulator::FindTheCycleToMesh(FIsoSegment* Segment, bool bOrientat
 		SubCycleNodes.Add(Node);
 		VertexIndexToSlopes.Emplace(Index, NodeSlop);
 		double Length = NextSegment->Get2DLengthSquare(EGridSpace::UniformScaled, Grid);
+		SquareLengthSum += Length;
 		if (Length < MinLength)
 		{
 			MinLength = Length;
@@ -180,6 +183,8 @@ bool FCycleTriangulator::FindTheCycleToMesh(FIsoSegment* Segment, bool bOrientat
 		Index++;
 	}
 
+	MeanSquareLength = SquareLengthSum / Index;
+
 	SubCycleNodeCount = SubCycleNodes.Num();
 	if (SubCycleNodeCount < 3)
 	{
@@ -189,15 +194,28 @@ bool FCycleTriangulator::FindTheCycleToMesh(FIsoSegment* Segment, bool bOrientat
 	return true;
 }
 
+bool IsCycleInsideCandidate(const FPoint2D& Candidate, const FPoint2D& Start, const FPoint2D& End, const FPoint2D& Next, TFunction<double(double)> ConvertSlope)
+{
+	const double StartEndSlop = ComputeSlope(Start, End);
+	const double StartCandidateSlope = ComputeSlope(Start, Candidate);
+	const double StartNextSlope = ComputeSlope(Start, Next);
+
+	const double NextSlop = ConvertSlope(TransformIntoPositiveSlope(StartCandidateSlope - StartNextSlope));
+	const double NewSlop = ConvertSlope(TransformIntoPositiveSlope(StartCandidateSlope - StartEndSlop));
+
+	return (NewSlop > NextSlop);
+};
+
 bool FCycleTriangulator::FindTheBestAcuteTriangle()
 {
+	const double MaxEdgeSquareLength = MeanSquareLength * 4.;
 	Algo::Sort(VertexIndexToSlopes, [](const TPair<int, double>& A, const TPair<int, double>& B) {return A.Value < B.Value; });
 
 	int32 NIndex = 0;
 
 	int32 Index = -1;
 	FIsoNode* CandidateNode = nullptr;
-	const FPoint2D* CandidadtePoint = nullptr;
+	const FPoint2D* CandidatePoint = nullptr;
 
 	int32 StartId = -1;
 	FIsoNode* NewSegmentStart = nullptr;
@@ -207,6 +225,9 @@ bool FCycleTriangulator::FindTheBestAcuteTriangle()
 	FIsoNode* NewSegmentEnd = nullptr;
 	const FPoint2D* EndPoint = nullptr;
 
+	const FPoint2D* NextPoint;
+	const FPoint2D* PrevPoint;
+
 	TFunction<void()> SetFinalCandidate = [&]()
 	{
 		FirstSideStartIndex = Index;
@@ -222,7 +243,7 @@ bool FCycleTriangulator::FindTheBestAcuteTriangle()
 		Index = InIndex;
 
 		CandidateNode = SubCycleNodes[Index];
-		CandidadtePoint = &CandidateNode->Get2DPoint(EGridSpace::UniformScaled, Grid);
+		CandidatePoint = &CandidateNode->Get2DPoint(EGridSpace::UniformScaled, Grid);
 
 		StartId = NextIndex(Index);
 		NewSegmentStart = SubCycleNodes[StartId];
@@ -232,59 +253,44 @@ bool FCycleTriangulator::FindTheBestAcuteTriangle()
 		NewSegmentEnd = SubCycleNodes[EndId];
 		EndPoint = &NewSegmentEnd->Get2DPoint(EGridSpace::UniformScaled, Grid);
 
+		const int32 NextId = NextIndex(StartId);
+		const FIsoNode* NextStart = SubCycleNodes[NextId];
+		NextPoint = &NextStart->Get2DPoint(EGridSpace::UniformScaled, Grid);
+
+		const int32 PrevId = PreviousIndex(EndId);
+		const FIsoNode* PrevEnd = SubCycleNodes[PrevId];
+		PrevPoint = &PrevEnd->Get2DPoint(EGridSpace::UniformScaled, Grid);
+
 #ifdef DEBUG_FIND_THE_CYCLE_TO_MESH
 		if (Grid.bDisplay)
 		{
+			static int32 CandidateTriangleIndex = 0;
 			FIsoNode* CandidateNode = SubCycleNodes[Index];
-			F3DDebugSession W(FString::Printf(TEXT("Candidate Triangle")));
+			F3DDebugSession W(FString::Printf(TEXT("Candidate Triangle %d"), ++CandidateTriangleIndex));
 			Grid.DisplayIsoSegment(EGridSpace::UniformScaled, *CandidateNode, *NewSegmentEnd, 0, EVisuProperty::YellowCurve);
 			Grid.DisplayIsoSegment(EGridSpace::UniformScaled, *NewSegmentStart, *CandidateNode, 0, EVisuProperty::YellowCurve);
 			Grid.DisplayIsoSegment(EGridSpace::UniformScaled, *NewSegmentStart, *NewSegmentEnd, 0, EVisuProperty::PurpleCurve);
 			Grid.DisplayIsoNode(EGridSpace::UniformScaled, *NewSegmentStart, 0, EVisuProperty::BluePoint);
 			Grid.DisplayIsoNode(EGridSpace::UniformScaled, *NewSegmentEnd, 0, EVisuProperty::RedPoint);
-			//Wait();
+
+			Grid.DisplayIsoSegment(EGridSpace::UniformScaled, *NewSegmentStart, *NextStart, 0, EVisuProperty::RedCurve);
+			Grid.DisplayIsoSegment(EGridSpace::UniformScaled, *NewSegmentEnd, *PrevEnd, 0, EVisuProperty::BlueCurve);
+			if(CandidateTriangleIndex == 42)
+			{
+				Wait();
+			}
 		}
 #endif
 	};
 
-	TFunction<bool(const int32&)> IsCycleInsideCandidateAtStart = [&](const int32& StartId) -> bool
+	TFunction<double(double)> TransformIntoClockwise = [](double Value) -> double
 	{
-		const int32 NextId = NextIndex(StartId);
-		FIsoNode* NextStart = SubCycleNodes[NextId];
-		const FPoint2D& NextPoint = NextStart->Get2DPoint(EGridSpace::UniformScaled, Grid);
-
-		double NewSegmentSlop = ComputeSlope(*StartPoint, *EndPoint);
-		double StartCandidateSlope = ComputeSlope(*StartPoint, *CandidadtePoint);
-		double StartNextSlope = ComputeSlope(*StartPoint, NextPoint);
-
-		double NextSlop = TransformIntoPositiveSlope(StartCandidateSlope - StartNextSlope);
-		double NewSlop = TransformIntoPositiveSlope(StartCandidateSlope - NewSegmentSlop);
-
-		if (NewSlop > NextSlop && CandidateNode != NextStart)
-		{
-			return true;
-		}
-		return false;
+		return TransformIntoClockwiseSlope(Value);
 	};
 
-	TFunction<bool(const int32&)> IsCycleInsideCandidateAtEnd = [&](const int32& EndId) -> bool
+	TFunction<double(double)> TransformIntoCounterClockwise = [](double Value) -> double
 	{
-		const int32 PrevId = PreviousIndex(EndId);
-		FIsoNode* PrevEnd = SubCycleNodes[PrevId];
-		const FPoint2D& PrevPoint = PrevEnd->Get2DPoint(EGridSpace::UniformScaled, Grid);
-
-		double New2SegmentSlop = ComputeSlope(*EndPoint, *StartPoint);
-		double EndCandidateSlope = ComputeSlope(*EndPoint, *CandidadtePoint);
-		double EndPrevSlope = ComputeSlope(*EndPoint, PrevPoint);
-
-		double PrevSlop = TransformIntoPositiveSlope(EndPrevSlope - EndCandidateSlope);
-		double New2Slop = TransformIntoPositiveSlope(New2SegmentSlop - EndCandidateSlope);
-
-		if (New2Slop > PrevSlop && CandidateNode != PrevEnd)
-		{
-			return true;
-		}
-		return false;
+		return Value;
 	};
 
 	// Check if the most Acute Triangle is candidate
@@ -300,28 +306,26 @@ bool FCycleTriangulator::FindTheBestAcuteTriangle()
 
 		GetCandidateTriangleNodes(Candidate.Key);
 
-		// Check that the min other angle is not smaller than the current angle
-		// At start of the new segment
-		double IsoscelesCriteria = IsoTriangulatorImpl::IsoscelesCriteriaMin(*StartPoint, *EndPoint, *CandidadtePoint);
-		if (IsoscelesCriteria < Candidate.Value)
+		if(IsCycleInsideCandidate(*CandidatePoint, *StartPoint, *EndPoint, *NextPoint, TransformIntoCounterClockwise)
+		 || IsCycleInsideCandidate(*CandidatePoint, *EndPoint, *StartPoint, *PrevPoint, TransformIntoClockwise))
 		{
 			continue;
 		}
 
-		// Check that the cycle is not inside the selected triangle
-		if(IsCycleInsideCandidateAtStart(StartId)
-		 || IsCycleInsideCandidateAtEnd(EndId))
+		FIsoSegment* Segment;
+		if (CycleIntersectionTool.DoesIntersect(NewSegmentStart, NewSegmentEnd, &Segment))
 		{
 			continue;
 		}
 
-		FIsoSegment* Segment;
-		if (!CycleIntersectionTool.DoesIntersect(NewSegmentStart, NewSegmentEnd, &Segment))
+		if (InnerToOuterIsoSegmentsIntersectionTool.DoesIntersect(*NewSegmentStart, *NewSegmentEnd))
 		{
-			bAcuteTriangle = true;
-			SetFinalCandidate();
-			return true;
+			continue;
 		}
+	
+		bAcuteTriangle = true;
+		SetFinalCandidate();
+		return true;
 	}
 
 	double BestIsoscelesCriteria = Slope::NullSlope;
@@ -329,7 +333,8 @@ bool FCycleTriangulator::FindTheBestAcuteTriangle()
 	// Check in the acute Triangles the most isosceles
 	int32 AcuteNodeIndex = -1;
 	constexpr double MaxSlopeToBeSelected = Slope::ThreeQuaterPiSlope;
-	for (; NIndex < SubCycleNodeCount; ++NIndex)
+
+	for (;  false && NIndex < SubCycleNodeCount; ++NIndex)
 	{
 		const TPair<int, double>& Candidate = VertexIndexToSlopes[NIndex];
 
@@ -341,8 +346,8 @@ bool FCycleTriangulator::FindTheBestAcuteTriangle()
 		GetCandidateTriangleNodes(Candidate.Key);
 		
 		// Check that the cycle is not inside the selected triangle
-		if (IsCycleInsideCandidateAtStart(StartId)
-			|| IsCycleInsideCandidateAtEnd(EndId))
+		if (IsCycleInsideCandidate(*CandidatePoint, *StartPoint, *EndPoint, *NextPoint, TransformIntoCounterClockwise)
+			|| IsCycleInsideCandidate(*CandidatePoint, *EndPoint, *StartPoint, *PrevPoint, TransformIntoClockwise))
 		{
 			continue;
 		}
@@ -350,10 +355,40 @@ bool FCycleTriangulator::FindTheBestAcuteTriangle()
 		FIsoSegment* Segment;
 		if (!CycleIntersectionTool.DoesIntersect(NewSegmentStart, NewSegmentEnd, &Segment))
 		{
-			double IsoscelesCriteria = IsoTriangulatorImpl::IsoscelesCriteriaMin(*StartPoint, *EndPoint, *CandidadtePoint);
+			double IsoscelesCriteria = IsoTriangulatorImpl::IsoscelesCriteriaMin(*StartPoint, *EndPoint, *CandidatePoint);
 
 			if (BestIsoscelesCriteria < IsoscelesCriteria)
 			{
+				double Length1 = StartPoint->SquareDistance(*CandidatePoint);
+
+				if (Length1 > MaxEdgeSquareLength)
+				{
+#ifdef DEBUG_FIND_THE_CYCLE_TO_MESH
+				if (Grid.bDisplay)
+				{
+					F3DDebugSession W(FString::Printf(TEXT("cancel due to Seg Length %f vs %f"), sqrt(Length1), sqrt(MaxEdgeSquareLength)));
+					Grid.DisplayIsoSegment(EGridSpace::UniformScaled, *NewSegmentStart, *CandidateNode, 0, EVisuProperty::YellowCurve);
+					Wait();
+				}
+#endif
+					continue;
+				}
+
+				double Length2 = EndPoint->SquareDistance(*CandidatePoint);
+
+				if (Length2 > MaxEdgeSquareLength)
+				{
+#ifdef DEBUG_FIND_THE_CYCLE_TO_MESH
+				if (Grid.bDisplay)
+				{
+					F3DDebugSession W(FString::Printf(TEXT("cancel due to Seg Length %f vs %f"), sqrt(Length2), sqrt(MaxEdgeSquareLength)));
+					Grid.DisplayIsoSegment(EGridSpace::UniformScaled, *NewSegmentEnd, *CandidateNode, 0, EVisuProperty::YellowCurve);
+					Wait();
+				}
+#endif
+					continue;
+				}
+
 				BestIsoscelesCriteria = IsoscelesCriteria;
 				AcuteNodeIndex = Index;
 				bAcuteTriangle = true;
@@ -406,20 +441,34 @@ bool FCycleTriangulator::FindCandidateNodes(int32 StartIndex)
 		return false;
 	}
 
-	const double StartMaxSlope = FMath::Min(CandidateNodes.Last().SlopeAtStartNode + Slope::Epsilon, Slope::PiSlope);
-	const double EndMaxSlope = FMath::Min(CandidateNodes[0].SlopeAtEndNode + Slope::Epsilon, Slope::PiSlope);
-
 #ifdef DEBUG_FIND_CANDIDATE_NODES
 	F3DDebugSession A(Grid.bDisplay, TEXT("Find Candidate Nodes"));
 	if (Grid.bDisplay)
 	{
 		Grid.DisplayIsoSegment(EGridSpace::UniformScaled, *FirstSideStartNode, *FirstSideEndNode, 0, EVisuProperty::YellowCurve);
-		Wait();
 	}
 #endif
 
+	constexpr double MaxSlope = Slope::PiSlope - Slope::OneDegree;
+	const double StartMaxSlope = FMath::Min(CandidateNodes.Last().SlopeAtStartNode + Slope::OneDegree, MaxSlope);
+	const double EndMaxSlope = FMath::Min(CandidateNodes[0].SlopeAtEndNode + Slope::OneDegree, MaxSlope);
+
+	if (CandidateNodes.Last().SlopeAtStartNode < MaxSlope)
+	{
+		CandidateNodes.Last().bIsValid = true;
+	}
+	if (CandidateNodes[0].SlopeAtEndNode < MaxSlope)
+	{
+		CandidateNodes[0].bIsValid = true;
+	}
+
 	for (FCandidateNode& CNode : CandidateNodes)
 	{
+		if (CNode.bIsValid)
+		{
+			continue;
+		}
+
 		FIsoNode* Node = CNode.Node;
 		const double SlopeAtStartNode = CNode.SlopeAtStartNode;
 		const double SlopeAtEndNode = CNode.SlopeAtEndNode;
@@ -445,6 +494,7 @@ bool FCycleTriangulator::FindCandidateNodes(int32 StartIndex)
 			//F3DDebugSession W(TEXT("Next Segment"));
 			Grid.DisplayIsoNode(EGridSpace::UniformScaled, *CNode.Node, 0, CNode.bIsValid ? EVisuProperty::BluePoint : EVisuProperty::PinkPoint);
 		}
+		Wait();
 	}
 #endif
 
@@ -455,20 +505,27 @@ bool FCycleTriangulator::FindTheBestCandidateNode()
 {
 	constexpr double MinSlopeToNotBeAligned = Slope::Epsilon;
 
-	int32 Intersections = 0;
+	IntersectionCountAllowed = 0;
 	int32 CandidateNodeIndex = -1;
 
 	double MinCriteria = HUGE_VALUE;
 	double CandidateSlopeAtStartNode = Slope::PiSlope - MinSlopeToNotBeAligned;
 	double CandidateSlopeAtEndNode = Slope::PiSlope - MinSlopeToNotBeAligned;
+	FCandidateNode* Candidate = nullptr;
+	double CandidateEndSquareDistance = HUGE_VALUE;
+	double CandidateStartSquareDistance = HUGE_VALUE;
+
 
 	const FPoint2D& StartPoint2D = FirstSideStartNode->Get2DPoint(EGridSpace::UniformScaled, Grid);
 	const FPoint2D& EndPoint2D = FirstSideEndNode->Get2DPoint(EGridSpace::UniformScaled, Grid);
 
 #ifdef DEBUG_FIND_THE_BEST_CANDIDATE_NODE
-	if (Grid.bDisplay)
+	static int32 BestCandidateNode = 0;
+	static int32 BestCandidateNodeTarget = 129;
+	++BestCandidateNode;
+	if (Grid.bDisplay /*&& BestCandidateNode >= BestCandidateNodeTarget*/)
 	{
-		F3DDebugSession _(TEXT("Candidate triangle"));
+		F3DDebugSession _(FString::Printf(TEXT("Start Segment %d"), BestCandidateNode));
 		Grid.DisplayIsoSegment(EGridSpace::UniformScaled, *FirstSideStartNode, *FirstSideEndNode, 1, EVisuProperty::RedCurve);
 		Wait();
 	}
@@ -476,13 +533,13 @@ bool FCycleTriangulator::FindTheBestCandidateNode()
 
 	for (FNodeIntersectionCount& NodeIntersection : NodeToIntersection)
 	{
-		if (NodeIntersection.Value > Intersections)
+		if (NodeIntersection.Value > IntersectionCountAllowed)
 		{
 			if (CandidateNodeIndex >= 0)
 			{
 				break;
 			}
-			Intersections++;
+			IntersectionCountAllowed++;
 		}
 
 		FCandidateNode* CNode = NodeIntersection.Key;
@@ -492,24 +549,32 @@ bool FCycleTriangulator::FindTheBestCandidateNode()
 
 		double PointCriteria = IsoTriangulatorImpl::IsoscelesCriteriaMax(StartPoint2D, EndPoint2D, NodePoint2D);
 
-		if (
-			// the candidate triangle is inside the current candidate triangle
-			((CNode->SlopeAtStartNode < (CandidateSlopeAtStartNode + MinSlopeToNotBeAligned)) && (CNode->SlopeAtEndNode < (CandidateSlopeAtEndNode + MinSlopeToNotBeAligned)))
-			||
-			// the candidate triangle is better the current candidate triangle and doesn't contain the current candidate triangle
-			((PointCriteria < MinCriteria) && ((CNode->SlopeAtStartNode > CandidateSlopeAtStartNode) ^ (CNode->SlopeAtEndNode > CandidateSlopeAtEndNode))))
-		{
-
 #ifdef DEBUG_FIND_THE_BEST_CANDIDATE_NODE
-			if (Grid.bDisplay)
-			{
-				F3DDebugSession _(TEXT("Candidate triangle"));
-				Grid.DisplayIsoSegment(EGridSpace::UniformScaled, *FirstSideStartNode, *Node, 1, EVisuProperty::RedCurve);
-				Grid.DisplayIsoSegment(EGridSpace::UniformScaled, *FirstSideEndNode, *Node, 2, EVisuProperty::RedCurve);
-				Wait();
-			}
+		if (Grid.bDisplay /*&& BestCandidateNode == BestCandidateNodeTarget*/)
+		{
+			F3DDebugSession _(TEXT("Candidate triangle"));
+			Grid.DisplayIsoSegment(EGridSpace::UniformScaled, *FirstSideStartNode, *Node, 1, EVisuProperty::RedCurve);
+			Grid.DisplayIsoSegment(EGridSpace::UniformScaled, *FirstSideEndNode, *Node, 2, EVisuProperty::RedCurve);
+			Wait();
+		}
 #endif
 
+		double CEndSquareDistance = StartPoint2D.SquareDistance(NodePoint2D);
+		double CStartSquareDistance = EndPoint2D.SquareDistance(NodePoint2D);
+
+		// the candidate triangle is inside the current candidate triangle
+		// Angle is not accurate enough to define if a triangle is inside or not an other triangle (especially with acute triangle).
+		// The length of the side is a better criteria but must be combine with the angles
+		const bool bIsInsideCandidate = (CNode->SlopeAtStartNode - CandidateSlopeAtStartNode) < Slope::OneDegree
+			&& (CNode->SlopeAtEndNode - CandidateSlopeAtEndNode) < Slope::OneDegree
+			&& CandidateEndSquareDistance > CEndSquareDistance
+			&& CandidateStartSquareDistance > CStartSquareDistance;
+
+		// the candidate triangle is better the current candidate triangle and doesn't contain the current candidate triangle
+		const bool bIsBetterCandidate = (PointCriteria < MinCriteria) && ((CNode->SlopeAtStartNode > CandidateSlopeAtStartNode) ^ (CNode->SlopeAtEndNode > CandidateSlopeAtEndNode));
+
+		if (bIsInsideCandidate || bIsBetterCandidate)
+		{
 			// Check if the candidate segment is not in intersection with existing segments
 			// if the segment exist, it has already been tested
 			FIsoSegment* StartSegment = nullptr;
@@ -544,22 +609,23 @@ bool FCycleTriangulator::FindTheBestCandidateNode()
 			CandidateNodeIndex = CNode->Index;
 			CandidateSlopeAtStartNode = CNode->SlopeAtStartNode;
 			CandidateSlopeAtEndNode = CNode->SlopeAtEndNode;
+			CandidateEndSquareDistance = CEndSquareDistance;
+			CandidateStartSquareDistance = CStartSquareDistance;
 		}
 	}
 
 #ifdef DEBUG_FIND_THE_BEST_CANDIDATE_NODE
 	static int32 TriangleIndex = 0;
-	if (Grid.bDisplay)
+	if (Grid.bDisplay && CandidateNodeIndex >= 0)
 	{
-		if (TriangleIndex == 1)
-		{
-			Wait();
-		}
 		FIsoNode* Node = SubCycleNodes[CandidateNodeIndex];
-		F3DDebugSession _(FString::Printf(TEXT("Selected triangle %d"), TriangleIndex++));
+		F3DDebugSession _(FString::Printf(TEXT("Selected triangle %d"), ++TriangleIndex));
 		Grid.DisplayIsoSegment(EGridSpace::UniformScaled, *FirstSideStartNode, *Node, 1, EVisuProperty::RedCurve);
 		Grid.DisplayIsoSegment(EGridSpace::UniformScaled, *FirstSideEndNode, *Node, 2, EVisuProperty::RedCurve);
-		Wait(false);
+		if (TriangleIndex == 1000)
+		{
+			Wait();
+		}
 	}
 #endif
 
@@ -646,25 +712,54 @@ void FCycleTriangulator::SelectFinalNodes(FAdditionalIso& Side1, FAdditionalIso&
 			Criteria[3] = Side2.EquilateralCriteria[1];
 
 			int32 Indices[4] = {0, 1, 2, 3};
-			Algo::Sort(Indices, [Criteria](int32 A, int32 B) { return Criteria[A] > Criteria[B]; });
+			Algo::Sort(Indices, [Criteria](int32 A, int32 B) { return Criteria[A] < Criteria[B]; });
 
-			for(int32 Index = 0; CandidateCount != 2; ++Index)
+			for(int32 Index = 0; Index < 4 && CandidateCount != 2; ++Index)
 			{
+				if (Criteria[Indices[Index]] > Slope::PiSlope)
+				{
+					break;
+				}
+
 				if (Indices[Index] > 1)
 				{
-					if(ValidComplementaryNode(Side2, (Indices[Index] - 2) == 0 ? 1 : 0))
+					if(!ValidComplementaryNodeOrDeleteIt(Side2, Indices[Index]-2))
 					{
-						Side2.RemoveCandidate(Indices[Index] - 2);
+						CandidateCount--;
 					}
 				}
 				else
 				{
-					if (ValidComplementaryNode(Side1, Indices[Index] == 0 ? 1 : 0))
+					if (!ValidComplementaryNodeOrDeleteIt(Side1, Indices[Index]))
+					{
+						CandidateCount--;
+					}
+				}
+			}
+			
+			if (CandidateCount > 2)
+			{
+				CandidateCount = 0;
+				int32 Index = 0;
+				for (; CandidateCount != 2; ++Index)
+				{
+					if (Criteria[Indices[Index]] > Slope::PiSlope)
+					{
+						continue;
+					}
+					CandidateCount++;
+				}
+				for (; Index < 4; ++Index)
+				{
+					if (Indices[Index] > 1)
+					{
+						Side2.RemoveCandidate(Indices[Index] - 2);
+					}
+					else
 					{
 						Side1.RemoveCandidate(Indices[Index]);
 					}
 				}
-				CandidateCount--;
 			}
 		}
 	}
@@ -715,22 +810,10 @@ bool FCycleTriangulator::BuildTheBestPolygonFromTheSelectedTriangle()
 	FAdditionalIso Side1(FirstSideEndIndex, TriangleThirdIndex, FirstSideEndNode, SelectedNode); // From end segment to candidate
 	FAdditionalIso Side2(TriangleThirdIndex, PreviousIndex(FirstSideEndIndex), SelectedNode, FirstSideStartNode); // From candidate to start segment
 
-#ifdef DEBUG_BuildTheBestPolygonFromTheSelectedTriangle
-	if(Grid.bDisplay)
-	{
-		static int32 IndexOfTraingle = 1;
-		F3DDebugSession G(FString::Printf(TEXT("Candidate Triangle %d"), IndexOfTraingle++));
-		Grid.DisplayTriangle(EGridSpace::UniformScaled, *FirstSideStartNode, *FirstSideEndNode, *SelectedNode);
-		if(IndexOfTraingle == 3)
-		{
-			Wait();
-		}
-	}
-#endif
-
 	// Check if there are nodes that could be meshed with
 	// i.e. if the next or previous node of the candidate or of the segment is nearly aligned, it would better to add the neighbor node 
 	// and mesh them as a quadrilateral or pentagon polygon
+	if (MaxIntersectionCounted == 0)
 	{
 		FindComplementaryNodes(Side1);
 		FindComplementaryNodes(Side2);
@@ -752,7 +835,7 @@ int32 PreviousIndex2(int32)
 	return 0;
 }
 
-bool FCycleTriangulator::ValidComplementaryNode(FAdditionalIso& Side, int32 CandidateIndex)
+bool FCycleTriangulator::ValidComplementaryNodeOrDeleteIt(FAdditionalIso& Side, int32 CandidateIndex)
 {
 	const FIsoNode* Candidate = Side.Nodes[CandidateIndex];
 	if (!Candidate)
@@ -768,18 +851,16 @@ bool FCycleTriangulator::ValidComplementaryNode(FAdditionalIso& Side, int32 Cand
 		return CandidateIndex ?  NextIndex(Index) : PreviousIndex(Index);
  	};
 
-	FIsoSegment* Segment;
-	if (!CycleIntersectionTool.DoesIntersect(Candidate, OppositeNode, &Segment))
-	{
-		const FPoint2D& CandidatePoint = Candidate->Get2DPoint(EGridSpace::UniformScaled, Grid);
-		const FPoint2D& OppositePoint = OppositeNode->Get2DPoint(EGridSpace::UniformScaled, Grid);
+	const FPoint2D& CandidatePoint = Candidate->Get2DPoint(EGridSpace::UniformScaled, Grid);
+	const FPoint2D& OppositePoint = OppositeNode->Get2DPoint(EGridSpace::UniformScaled, Grid);
 
-		const int32 TestPointIndex = NextFunc(OppositeIndex);
-		const FIsoNode* TestNode = SubCycleNodes[TestPointIndex];
-		const FPoint2D& TestPoint = TestNode->Get2DPoint(EGridSpace::UniformScaled, Grid);
+	const int32 TestPointIndex = NextFunc(OppositeIndex);
+	const FIsoNode* TestNode = SubCycleNodes[TestPointIndex];
+	const FPoint2D& TestPoint = TestNode->Get2DPoint(EGridSpace::UniformScaled, Grid);
 
 #ifdef DEBUG_VALID_COMPLEMENTARY_NODE
-		if (Grid.bDisplay)
+	if (Grid.bDisplay)
+	{
 		{
 			F3DDebugSession _(TEXT("Selected Node"));
 			Grid.DisplayIsoSegment(EGridSpace::UniformScaled, *OppositeNode, *Candidate, 1, EVisuProperty::BlueCurve);
@@ -789,12 +870,16 @@ bool FCycleTriangulator::ValidComplementaryNode(FAdditionalIso& Side, int32 Cand
 			F3DDebugSession _(TEXT("Slope"));
 			Grid.DisplayIsoSegment(EGridSpace::UniformScaled, *OppositeNode, *Candidate, 1, EVisuProperty::YellowCurve);
 			Grid.DisplayIsoSegment(EGridSpace::UniformScaled, *OppositeNode, *TestNode, 1, EVisuProperty::YellowCurve);
-			Wait();
 		}
+	}
 #endif
 
-		double Slope = ComputeOrientedSlope(OppositePoint, CandidatePoint, TestPoint);
-		if ((Slope > 0) == (CandidateIndex == 0))
+	// Is the new segment fully inside the cycle 
+	double Slope = ComputeOrientedSlope(OppositePoint, CandidatePoint, TestPoint);
+	if ((Slope > 0) == (CandidateIndex == 0))
+	{
+		FIsoSegment* Segment;
+		if (!CycleIntersectionTool.DoesIntersect(Candidate, OppositeNode, &Segment))
 		{
 			// The new segment is fully inside the cycle 
 			return true;
@@ -889,6 +974,18 @@ bool FCycleTriangulator::IsInnerSideSegmentInsideCycle(FAdditionalIso& Side)
 
 void FCycleTriangulator::ValidateComplementaryNodesWithInsideAndIntersectionsCriteria(FAdditionalIso& Side)
 {
+	for(int32 Index = 0; Index < 2; ++Index)
+	{
+		if (Side.Nodes[Index])
+		{
+			if(InnerToOuterIsoSegmentsIntersectionTool.DoesIntersect(Index == 0 ? *Side.EndNode : *Side.StartNode, *Side.Nodes[Index]))
+			{
+				Side.RemoveCandidate(Index);
+				Side.bForceNodes = false;
+			}
+		}
+	}
+
 	if (Side.bForceNodes)
 	{
 		return;
@@ -920,8 +1017,8 @@ void FCycleTriangulator::ValidateComplementaryNodesWithInsideAndIntersectionsCri
 				Wait();
 			}
 #endif
-			ValidComplementaryNode(Side, 0);
-			ValidComplementaryNode(Side, 1);
+			ValidComplementaryNodeOrDeleteIt(Side, 0);
+			ValidComplementaryNodeOrDeleteIt(Side, 1);
 
 			if (Side.CandidateNodeCount() == 2)
 			{
@@ -941,12 +1038,12 @@ void FCycleTriangulator::ValidateComplementaryNodesWithInsideAndIntersectionsCri
 
 	if (Side.Nodes[0])
 	{
-		ValidComplementaryNode(Side, 0);
+		ValidComplementaryNodeOrDeleteIt(Side, 0);
 	}
 
 	if (Side.Nodes[1])
 	{
-		ValidComplementaryNode(Side, 1);
+		ValidComplementaryNodeOrDeleteIt(Side, 1);
 	}
 }
 
@@ -1020,7 +1117,7 @@ bool FCycleTriangulator::BuildTheBestPolygon(FIsoSegment* StartSegment, bool bSt
 
 #ifdef DEBUG_FIND_BEST_POLYGON
 	static int32 PolyIndex = 0;
-	F3DDebugSession W(Grid.bDisplay, FString::Printf(TEXT("New Polygon %d"), PolyIndex++));
+	F3DDebugSession W(Grid.bDisplay, FString::Printf(TEXT("New Polygon %d"), ++PolyIndex));
 #endif
 
 	CleanContext();
@@ -1049,7 +1146,7 @@ bool FCycleTriangulator::BuildTheBestPolygon(FIsoSegment* StartSegment, bool bSt
 		// Step 4: 
 		//    Sort CandidateNodes according to the number of intersection with isolines
 		{
-			const int32 MaxIntersectionCounted = CountIntersectionWithIso();
+			MaxIntersectionCounted = CountIntersectionWithIso();
 			if (MaxIntersectionCounted)
 			{
 				NodeToIntersection.Sort([](const FNodeIntersectionCount& N1, const FNodeIntersectionCount& N2) { return N1.Value < N2.Value; });
@@ -1114,7 +1211,7 @@ bool  FCycleTriangulator::BuildTriangle(TArray<FIsoNode*>& CandidatNodes)
 		return false;
 	}
 
-	Mesh.AddTriangle(CandidatNodes[0]->GetFaceIndex(), CandidatNodes[1]->GetFaceIndex(), CandidatNodes[2]->GetFaceIndex());
+	Mesh.AddTriangle(CandidatNodes[0]->GetGlobalIndex(), CandidatNodes[1]->GetGlobalIndex(), CandidatNodes[2]->GetGlobalIndex());
 
 #ifdef ADD_TRIANGLE_2D
 	if (Grid.bDisplay)
@@ -1168,7 +1265,7 @@ void MeshTriangle(const FGrid& Grid, FIsoNode** Nodes, FFaceMesh& Mesh)
 		Grid.DisplayTriangle(EGridSpace::UniformScaled, *Nodes[0], *Nodes[1], *Nodes[2]);
 	}
 #endif 
-	Mesh.AddTriangle(Nodes[0]->GetFaceIndex(), Nodes[1]->GetFaceIndex(), Nodes[2]->GetFaceIndex());
+	Mesh.AddTriangle(Nodes[0]->GetGlobalIndex(), Nodes[1]->GetGlobalIndex(), Nodes[2]->GetGlobalIndex());
 }
 
 void MeshQuadrilateral(const FGrid& Grid, FIsoNode** Nodes, FFaceMesh& Mesh)
@@ -1217,8 +1314,8 @@ void MeshQuadrilateral(const FGrid& Grid, FIsoNode** Nodes, FFaceMesh& Mesh)
 	}
 #endif 
 
-	Mesh.AddTriangle(Nodes[NodeIndices[1]]->GetFaceIndex(), Nodes[NodeIndices[3]]->GetFaceIndex(), Nodes[NodeIndices[0]]->GetFaceIndex());
-	Mesh.AddTriangle(Nodes[NodeIndices[1]]->GetFaceIndex(), Nodes[NodeIndices[2]]->GetFaceIndex(), Nodes[NodeIndices[3]]->GetFaceIndex());
+	Mesh.AddTriangle(Nodes[NodeIndices[1]]->GetGlobalIndex(), Nodes[NodeIndices[3]]->GetGlobalIndex(), Nodes[NodeIndices[0]]->GetGlobalIndex());
+	Mesh.AddTriangle(Nodes[NodeIndices[1]]->GetGlobalIndex(), Nodes[NodeIndices[2]]->GetGlobalIndex(), Nodes[NodeIndices[3]]->GetGlobalIndex());
 }
 
 void MeshPentagon(const FGrid& Grid, FIsoNode** Nodes, FFaceMesh& Mesh)
@@ -1407,7 +1504,7 @@ void MeshPentagon(const FGrid& Grid, FIsoNode** Nodes, FFaceMesh& Mesh)
 	// Add the selected triangles to the mesh
 	for (int32 Index = 0; Index < 3; ++Index)
 	{
-		Mesh.AddTriangle(Nodes[TrianglesOfPentagon[SelectedTriangles[Index]][0]]->GetFaceIndex(), Nodes[TrianglesOfPentagon[SelectedTriangles[Index]][1]]->GetFaceIndex(), Nodes[TrianglesOfPentagon[SelectedTriangles[Index]][2]]->GetFaceIndex());
+		Mesh.AddTriangle(Nodes[TrianglesOfPentagon[SelectedTriangles[Index]][0]]->GetGlobalIndex(), Nodes[TrianglesOfPentagon[SelectedTriangles[Index]][1]]->GetGlobalIndex(), Nodes[TrianglesOfPentagon[SelectedTriangles[Index]][2]]->GetGlobalIndex());
 	}
 }
 
@@ -1506,6 +1603,7 @@ void FCycleTriangulator::InitializeArrays()
 	SubCycleNodes.Reserve(NodeCount);
 	PolygonNodes.Reserve(5);
 	VertexIndexToSlopes.Reserve(NodeCount);
+	MeanSquareLength = 0;
 
 	CycleIntersectionTool.AddSegments(Cycle);
 	SortCycleIntersectionToolIfNeeded();
@@ -1715,12 +1813,12 @@ void FCycleTriangulator::CleanContext()
 	bAcuteTriangle = false;
 
 	SubCycleNodeCount = 0;
-	SubCycleNodes.SetNum(0);
-	CandidateNodes.SetNum(0);
-	PolygonNodes.SetNum(0);
-	NodeToIntersection.Empty(NodeCount);
+	SubCycleNodes.Reset();
+	CandidateNodes.Reset();
+	PolygonNodes.Reset();
+	NodeToIntersection.Reset();
 
-	VertexIndexToSlopes.SetNum(0);
+	VertexIndexToSlopes.Reset();
 	
 	FirstSideStartIndex = -1;
 	FirstSideEndIndex = -1;
diff --git a/Engine/Source/Runtime/Datasmith/CADKernel/Private/CADKernel/Mesh/Meshers/CycleTriangulator.h b/Engine/Source/Runtime/Datasmith/CADKernel/Private/CADKernel/Mesh/Meshers/CycleTriangulator.h
--- a/Engine/Source/Runtime/Datasmith/CADKernel/Private/CADKernel/Mesh/Meshers/CycleTriangulator.h
+++ b/Engine/Source/Runtime/Datasmith/CADKernel/Private/CADKernel/Mesh/Meshers/CycleTriangulator.h
@@ -82,7 +82,7 @@ struct FAdditionalIso
 	{
 		NodeIndices[Index] =  -1;
 		Nodes[Index] = nullptr;
-		EquilateralCriteria[Index] = Slope::NullSlope;
+		EquilateralCriteria[Index] = Slope::TwoPiSlope;
 	}
 
 	void Reset()
@@ -109,7 +109,7 @@ private:
 	const TArray<bool>& CycleOrientation;
 
 	const EGridSpace Space = EGridSpace::UniformScaled;
-	const FIntersectionNodePairTool& InnerToOuterSegmentsIntersectionTool;
+	const FIntersectionNodePairTool& InnerToOuterIsoSegmentsIntersectionTool;
 	FFaceMesh& Mesh;
 	TFactory<FIsoSegment>& IsoSegmentFactory;
 
@@ -127,9 +127,12 @@ private:
 	TArray<FIsoNode*> SubCycleNodes;
 
 	TArray<TPair<int32, double>> VertexIndexToSlopes;
-
+	double MeanSquareLength = 0;
 	bool bAcuteTriangle = false;
 
+	int32 IntersectionCountAllowed = 0;
+	int32 MaxIntersectionCounted = 0;
+
 	int32 FirstSideStartIndex = -1;
 	int32 FirstSideEndIndex = -1;
 	int32 TriangleThirdIndex = -1;
@@ -200,7 +203,7 @@ private:
 	 * Valid additional nodes i.e. check if they can be added in The Side Of Selected Triangle
 	 */
 	void ValidateComplementaryNodesWithInsideAndIntersectionsCriteria(FAdditionalIso& Side);
-	bool ValidComplementaryNode(FAdditionalIso& Side, int32 Index);
+	bool ValidComplementaryNodeOrDeleteIt(FAdditionalIso& Side, int32 Index);
 	bool IsInnerSideSegmentInsideCycle(FAdditionalIso& Side);
 
 	void SelectFinalNodes(FAdditionalIso& Side1, FAdditionalIso& Side2);
@@ -222,14 +225,15 @@ private:
 
 	int32 IsIntersectingIso(const TArray<FIsoNode*>& Nodes)
 	{
-		const FIsoNode* StartNode = Nodes.Last();
-		for (const FIsoNode* Node : Nodes)
+		for (int32 IndexA = 0; IndexA < Nodes.Num() - 1; ++IndexA)
 		{
-			if (InnerToOuterSegmentsIntersectionTool.CountIntersections(*StartNode, *Node))
+			for (int32 IndexB = IndexA + 1; IndexB < Nodes.Num() - 1; ++IndexB)
 			{
-				return true;
+				if (InnerToOuterIsoSegmentsIntersectionTool.DoesIntersect(*Nodes[IndexA], *Nodes[IndexB]))
+				{
+					return true;
+				}
 			}
-			StartNode = Node;
 		}
 		return false;
 	}
@@ -240,8 +244,7 @@ private:
 		for (const FCandidateNode& CNode : Nodes)
 		{
 			const FIsoNode* Node = CNode.Node;
-
-			if (InnerToOuterSegmentsIntersectionTool.CountIntersections(*StartNode, *Node))
+			if (InnerToOuterIsoSegmentsIntersectionTool.DoesIntersect(*StartNode, *Node))
 			{
 				return true;
 			}
@@ -251,9 +254,10 @@ private:
 
 	int32 CountIntersectionWithIso()
 	{
-		const FIsoNode* StartNode = CandidateNodes[0].Node;
-		const FIsoNode* EndNode = CandidateNodes.Last().Node;
+		FirstSideStartNode = SubCycleNodes[FirstSideStartIndex];
+		FirstSideEndNode = SubCycleNodes[FirstSideEndIndex];
 
+		NodeToIntersection.Reserve(CandidateNodes.Num());
 		int32 Max = 0;
 		for (int32 Index = 0; Index < CandidateNodes.Num(); ++Index)
 		{
@@ -265,8 +269,9 @@ private:
 
 			const FIsoNode* Node = CNode.Node;
 
-			int32 IntersectionCount = InnerToOuterSegmentsIntersectionTool.CountIntersections(*StartNode, *Node);
-			IntersectionCount = FMath::Max(IntersectionCount, InnerToOuterSegmentsIntersectionTool.CountIntersections(*EndNode, *Node));
+			int32 IntersectionCount = InnerToOuterIsoSegmentsIntersectionTool.CountIntersections(*FirstSideStartNode, *Node);
+			IntersectionCount = FMath::Max(IntersectionCount, InnerToOuterIsoSegmentsIntersectionTool.CountIntersections(*FirstSideEndNode, *Node));
+
 			NodeToIntersection.Emplace(&CNode, IntersectionCount);
 			if (Max < IntersectionCount)
 			{
diff --git a/Engine/Source/Runtime/Datasmith/CADKernel/Private/CADKernel/Mesh/Meshers/IsoTriangulator.cpp b/Engine/Source/Runtime/Datasmith/CADKernel/Private/CADKernel/Mesh/Meshers/IsoTriangulator.cpp
--- a/Engine/Source/Runtime/Datasmith/CADKernel/Private/CADKernel/Mesh/Meshers/IsoTriangulator.cpp
+++ b/Engine/Source/Runtime/Datasmith/CADKernel/Private/CADKernel/Mesh/Meshers/IsoTriangulator.cpp
@@ -63,7 +63,7 @@ FIsoTriangulator::FIsoTriangulator(FGrid& InGrid, FFaceMesh& OutMesh, const FMes
 	, LoopSegmentsIntersectionTool(InGrid)
 	, InnerSegmentsIntersectionTool(InGrid)
 	, InnerToLoopSegmentsIntersectionTool(InGrid)
-	, InnerToOuterSegmentsIntersectionTool(InGrid)
+	, InnerToOuterIsoSegmentsIntersectionTool(InGrid, /*AllowOverlaping*/ true)
 	, ThinZoneIntersectionTool(InGrid)
 	, Tolerances(InTolerances)
 {
@@ -152,7 +152,7 @@ bool FIsoTriangulator::Triangulate()
 		Grid.DisplayIsoSegments(DisplaySpace, FinalInnerSegments, true, false, false, EVisuProperty::BlueCurve);
 		//Wait();
 	}
-	InnerToOuterSegmentsIntersectionTool.Display(bDisplay, TEXT("FIsoTrianguler::IntersectionTool InnerToOuter"), EVisuProperty::RedCurve);
+	InnerToOuterIsoSegmentsIntersectionTool.Display(bDisplay, TEXT("FIsoTrianguler::IntersectionTool InnerToOuter"), EVisuProperty::RedCurve);
 #endif
 
 	// =============================================================================================================
@@ -213,26 +213,26 @@ void FIsoTriangulator::BuildNodes()
 	FTimePoint StartTime = FChrono::Now();
 
 	LoopNodeCount = 0;
-	for (const TArray<FPoint2D>& Loop : Grid.GetLoops2D(EGridSpace::Default2D))
+	for (const TArray<FPoint2D>& LoopPoints : Grid.GetLoops2D(EGridSpace::Default2D))
 	{
-		LoopNodeCount += (int32)Loop.Num();
+		LoopNodeCount += (int32)LoopPoints.Num();
 	}
 	LoopStartIndex.Reserve(Grid.GetLoops2D(EGridSpace::Default2D).Num());
 	LoopNodes.Reserve((int32)(LoopNodeCount * 1.2 + 5)); // reserve more in case it need to create complementary nodes
 
 	// Loop nodes
-	int32 FaceIndex = 0;
+	int32 NodeIndex = 0;
 	int32 LoopIndex = 0;
-	for (const TArray<FPoint2D>& Loop : Grid.GetLoops2D(EGridSpace::Default2D))
+	for (const TArray<FPoint2D>& LoopPoints : Grid.GetLoops2D(EGridSpace::Default2D))
 	{
 		LoopStartIndex.Add(LoopNodeCount);
 		const TArray<int32>& LoopIds = Grid.GetNodeIdsOfFaceLoops()[LoopIndex];
 		FLoopNode* NextNode = nullptr;
-		FLoopNode* FirstNode = &LoopNodes.Emplace_GetRef(LoopIndex, 0, FaceIndex++, LoopIds[0]);
+		FLoopNode* FirstNode = &LoopNodes.Emplace_GetRef(LoopIndex, 0, NodeIndex++, LoopIds[0]);
 		FLoopNode* PreviousNode = FirstNode;
-		for (int32 Index = 1; Index < Loop.Num(); ++Index)
+		for (int32 Index = 1; Index < LoopPoints.Num(); ++Index)
 		{
-			NextNode = &LoopNodes.Emplace_GetRef(LoopIndex, Index, FaceIndex++, LoopIds[Index]);
+			NextNode = &LoopNodes.Emplace_GetRef(LoopIndex, Index, NodeIndex++, LoopIds[Index]);
 			PreviousNode->SetNextConnectedNode(NextNode);
 			NextNode->SetPreviousConnectedNode(PreviousNode);
 			PreviousNode = NextNode;
@@ -251,7 +251,7 @@ void FIsoTriangulator::BuildNodes()
 	{
 		if (Grid.IsNodeInsideAndMeshable(Index))
 		{
-			FIsoInnerNode& Node = InnerNodes.Emplace_GetRef(Index, FaceIndex++, InnerNodeCount++);
+			FIsoInnerNode& Node = InnerNodes.Emplace_GetRef(Index, NodeIndex++, InnerNodeCount++);
 			GlobalIndexToIsoInnerNodes[Index] = &Node;
 		}
 	}
@@ -624,7 +624,7 @@ void FIsoTriangulator::BuildInnerSegments()
 		const FIsoInnerNode* Node1 = GlobalIndexToIsoInnerNodes[IndexNode1];
 		const FIsoInnerNode* Node2 = GlobalIndexToIsoInnerNodes[IndexNode2];
 
-		InnerToOuterSegmentsIntersectionTool.AddSegment(Node1, Node2, Point1, Point2);
+		InnerToOuterIsoSegmentsIntersectionTool.AddSegment(Node1, Node2, Point1, Point2);
 	};
 
 	TFunction<void(const int32, const int32, const ESegmentType)> AddToInnerSegments = [&](const int32 IndexNode1, const int32 IndexNode2, const ESegmentType InType)
@@ -770,7 +770,7 @@ void FIsoTriangulator::BuildInnerSegments()
 	Wait(bDisplay);
 #endif
 
-	InnerToOuterSegmentsIntersectionTool.Sort();
+	InnerToOuterIsoSegmentsIntersectionTool.Sort();
 
 #ifdef CADKERNEL_DEV
 	Chronos.BuildInnerSegmentsDuration += FChrono::Elapse(StartTime);
@@ -925,443 +925,6 @@ void FIsoTriangulator::BuildInnerSegmentsIntersectionTool()
 #endif
 }
 
-void FIsoTriangulator::FindIsoSegmentToLinkLoopToLoop()
-{
-	FTimePoint StartTime = FChrono::Now();
-
-	// This coefficient is to defined the tolerance on coordinates according the iso strip...
-	// With some surface, the parameterization speed can vary enormously depending on the point of the surface.
-	// A good information is the width of iso strip around a point. Indeed, strips have the optimal width to respect meshing criteria.
-	// So a fraction of the strip's width defined a good tolerance around a given point. 
-	constexpr double ToleranceCoefficent = 1. / 12.; // Why 12 ? ;o)
-
-	const TArray<double>& IsoUCoordinates = Grid.GetUniformCuttingCoordinatesAlongIso(EIso::IsoU);
-	const TArray<double>& IsoVCoordinates = Grid.GetUniformCuttingCoordinatesAlongIso(EIso::IsoV);
-
-
-	// Warning, Min delta is compute in EGridSpace::Uniform
-	TFunction<double(const TArray<double>&)> GetMinDelta = [](const TArray<double>& IsoCoordinates)
-	{
-		double MinDelta = HUGE_VAL;
-		for (int32 Index = 0; Index < IsoCoordinates.Num() - 1; ++Index)
-		{
-			double Delta = IsoCoordinates[Index + 1] - IsoCoordinates[Index];
-			if (Delta < MinDelta)
-			{
-				MinDelta = Delta;
-			}
-		}
-		return MinDelta;
-	};
-
-	// Before creating a segment a set of check is done to verify that the segment is valid.
-	TFunction<void(FLoopNode*, const FPoint2D&, FLoopNode*, const FPoint2D&)> CreateSegment = [&](FLoopNode* Node1, const FPoint2D& Coordinate1, FLoopNode* Node2, const FPoint2D& Coordinate2)
-	{
-		if (&Node1->GetPreviousNode() == Node2 || &Node1->GetNextNode() == Node2)
-		{
-			return;
-		}
-
-		if (Node1->GetSegmentConnectedTo(Node2))
-		{
-			return;
-		}
-
-		// Is Outside and not too flat at Node1
-		ensureCADKernel(Node1->GetLoopIndex() > 0);
-		const double FlatAngle = 0.1;
-		if (Node1->IsSegmentBeInsideFace(Coordinate2, Grid, FlatAngle))
-		{
-			return;
-		}
-
-		// Is Outside and not too flat at Node2
-		ensureCADKernel(Node2->GetLoopIndex() > 0);
-		if (Node2->IsSegmentBeInsideFace(Coordinate1, Grid, FlatAngle))
-		{
-			return;
-		}
-
-		if (InnerSegmentsIntersectionTool.DoesIntersect(Coordinate1, Coordinate2))
-		{
-			return;
-		}
-
-		if (LoopSegmentsIntersectionTool.DoesIntersect(*Node1, *Node2))
-		{
-			return;
-		}
-
-		FIsoSegment& Segment = IsoSegmentFactory.New();
-		Segment.Init(*Node1, *Node2, ESegmentType::LoopToLoop);
-		Segment.ConnectToNode();
-		FinalToLoops.Add(&Segment);
-		InnerToLoopSegmentsIntersectionTool.AddSegment(Segment);
-#ifdef DEBUG_FIND_ISOSEGMENT_TO_LINK_LOOP_TO_LOOP
-		Display(EGridSpace::UniformScaled, *Node1, *Node2);
-#endif
-	};
-
-	// Find the index of the closed Iso strip. An iso strip is a the strip [Iso[Index], Iso[Index+1]]
-	// As the process is iterative with sorted points, Index can only be equal or bigger than with the previous node
-	TFunction<void(const TArray<double>&, int32&, const double&)> FindStripIndex = [&](const TArray<double>& Iso, int32& Index, const double& PointCoord)
-	{
-		if (Index > 0)
-		{
-			--Index;
-		}
-
-		// The last strip is not tested as it must be good if the previous are not even if PointCoord >= Iso.Last
-		for (; Index < Iso.Num() - 2; ++Index)
-		{
-			if (PointCoord < Iso[Index + 1])
-			{
-				break;
-			}
-		}
-	};
-
-	TArray<FLoopNode*> SortedLoopNodesAlong = SortedLoopNodes;
-
-	// Find pair of points iso aligned along Axe2
-	// For all loop nodes sorted along Axe1, check if the pair (Node[i], Node[i+1]) is aligned along Axe2. 
-	// The segment (Node[i], Node[i+1]) is valid if its length is smaller or nearly equal to a crossing strip. 
-	// Axe1 == 0 => IsoU, coordinate U is ~constant
-	// Axe1 == 1 => IsoV
-
-	TFunction<void(int32, const TArray<double>&, const TArray<double>&)> FindIsoSegmentAlong = [&](int32 InAxe, const TArray<double>& IsoU, const TArray<double>& IsoV)
-	{
-		EIso ComplementaryAxe = (InAxe + 1) % 2 == 0 ? EIso::IsoU : EIso::IsoV;
-
-		int32 IndexU = 0;
-		for (int32 Index = 0; Index < SortedLoopNodesAlong.Num() - 1; ++Index)
-		{
-			FLoopNode* LoopNode = SortedLoopNodesAlong[Index];
-			if (!LoopNode->IsIso(ComplementaryAxe))
-			{
-				continue;
-			}
-			FLoopNode* NextNode = SortedLoopNodesAlong[Index + 1];
-			if (!NextNode->IsIso(ComplementaryAxe))
-			{
-				continue;
-			}
-
-			const FPoint2D& LoopPoint = LoopNode->Get2DPoint(EGridSpace::UniformScaled, Grid);
-
-#ifdef DEBUG_FIND_ISOSEGMENT_TO_LINK_LOOP_TO_LOOP
-			F3DDebugSession G(FString::Printf(TEXT("Iso %d Index %d"), InAxe, Index));
-			Display(EGridSpace::UniformScaled, *LoopNode);
-#endif
-
-			FindStripIndex(IsoU, IndexU, LoopPoint[InAxe]);
-
-			double ToleranceU = (IsoU[IndexU + 1] - IsoU[IndexU]) * ToleranceCoefficent;
-
-#ifdef DEBUG_FIND_ISOSEGMENT_TO_LINK_LOOP_TO_LOOP
-			Display(EGridSpace::UniformScaled, *NextNode, 0, EVisuProperty::RedPoint);
-#endif
-			const FPoint2D& NextPoint = NextNode->Get2DPoint(EGridSpace::UniformScaled, Grid);
-			if (FMath::IsNearlyEqual(NextPoint[InAxe], LoopPoint[InAxe], ToleranceU))
-			{
-				// the nodes are nearly iso aligned, are they nearly in the same V Stip
-				double MinV = LoopPoint[ComplementaryAxe];
-				double MaxV = NextPoint[ComplementaryAxe];
-				GetMinMax(MinV, MaxV);
-
-				int32 IndexV = 0;
-				FindStripIndex(IsoV, IndexV, MinV);
-
-				if (IndexV >= IsoV.Num() - 1)
-				{
-					continue;
-				}
-
-				// We check that segment length is not greater the the crossing strip width
-				bool bIsSmallerThanStrip = false;
-				if (MaxV <= IsoV[IndexV + 1])
-				{
-					// both point are in the same strip
-					bIsSmallerThanStrip = true;
-				}
-				else
-				{
-					// either MinV is nearly equal to IsoV[IndexV + 1]- 
-					double FirstStripCrossingLength = (IsoV[IndexV + 1] - MinV) / (IsoV[IndexV + 1] - IsoV[IndexV]);
-					if (IndexV < IsoV.Num() - 2 && MaxV < IsoV[IndexV + 1])
-					{
-						double SecondStripCrossingLength = (MaxV - IsoV[IndexV + 1]) / (IsoV[IndexV + 2] - IsoV[IndexV + 1]);
-						if ((FirstStripCrossingLength + SecondStripCrossingLength) < 1 + ToleranceCoefficent)
-						{
-							bIsSmallerThanStrip = true;
-						}
-					}
-					// either MaxV is nearly equal to IsoV[IndexV + 1]+
-					else if (IndexV < IsoV.Num() - 3 && MaxV < IsoV[IndexV + 2])
-					{
-						double ThirdStripCrossingLength = (MaxV - IsoV[IndexV + 2]) / (IsoV[IndexV + 3] - IsoV[IndexV + 2]);
-						if ((FirstStripCrossingLength + ThirdStripCrossingLength) < ToleranceCoefficent)
-						{
-							bIsSmallerThanStrip = true;
-						}
-					}
-				}
-				if (bIsSmallerThanStrip)
-				{
-					CreateSegment(LoopNode, LoopPoint, NextNode, NextPoint);
-				}
-			}
-		}
-	};
-
-	int32 InitNum = InnerSegmentsIntersectionTool.Count() + LoopSegmentsIntersectionTool.Count();
-	FinalToLoops.Reserve(InitNum);
-	InnerToLoopSegmentsIntersectionTool.Reserve(InitNum);
-
-	// Nodes are sorted according to a value function of their coordinates.
-	// To sort along U, the value is U + DeltaFactor*(V - VMin)
-	// DeltaFactor is a value that for all values Ui of U, Ui + DeltaFactor.(VMax - VMin) < U(i+1)
-	// With this, Node[i+1] is either the next node of the same side of the loop, either the closed U aligned node of the opposite loop.  
-	{
-		constexpr int32 IsoU = 0; // coordinate U is ~constant
-		constexpr int32 IsoV = 1; // coordinate V is ~constant
-
-		double DeltaFactor = FMath::Min((GetMinDelta(IsoUCoordinates) / 1000.), (GetMinDelta(IsoVCoordinates) / 1000.));
-
-		// Bounds and GetMinDelta are defined in EGridSpace::Default2D,
-		//const FSurfacicBoundary& Bounds = Grid.GetFace()->GetBoundary();
-		double UMin = Grid.GetUniformCuttingCoordinates()[EIso::IsoU][0]; // Bounds.UVBoundaries[EIso::IsoU].Min;
-		double VMin = Grid.GetUniformCuttingCoordinates()[EIso::IsoV][0]; // Bounds.UVBoundaries[EIso::IsoV].Min;
-		Algo::Sort(SortedLoopNodesAlong, [&](const FLoopNode* LoopNode1, const FLoopNode* LoopNode2)
-			{
-				const FPoint2D& Node1Coordinates = LoopNode1->Get2DPoint(EGridSpace::UniformScaled, Grid);
-		const FPoint2D& Node2Coordinates = LoopNode2->Get2DPoint(EGridSpace::UniformScaled, Grid);
-		return (Node1Coordinates.U + (Node1Coordinates.V - VMin) * DeltaFactor) < (Node2Coordinates.U + (Node2Coordinates.V - VMin) * DeltaFactor);
-			});
-		FindIsoSegmentAlong(IsoU, IsoUCoordinates, IsoVCoordinates);
-
-		Algo::Sort(SortedLoopNodesAlong, [&](const FLoopNode* LoopNode1, const FLoopNode* LoopNode2)
-			{
-				const FPoint2D& Node1Coordinates = LoopNode1->Get2DPoint(EGridSpace::UniformScaled, Grid);
-		const FPoint2D& Node2Coordinates = LoopNode2->Get2DPoint(EGridSpace::UniformScaled, Grid);
-		return (Node1Coordinates.V + (Node1Coordinates.U - UMin) * DeltaFactor) < (Node2Coordinates.V + (Node2Coordinates.U - UMin) * DeltaFactor);
-			});
-		FindIsoSegmentAlong(IsoV, IsoVCoordinates, IsoUCoordinates);
-	}
-#ifdef CADKERNEL_DEV
-	Chronos.FindInnerSegmentToLinkLoopToLoopDuration += FChrono::Elapse(StartTime);
-#endif
-}
-
-void FIsoTriangulator::FindSegmentToLinkOuterLoopNodes(FCell& Cell)
-{
-	int SubdivisionCount = Cell.OuterLoopSubdivision.Num();
-	for (int32 Andex = 0; Andex < SubdivisionCount - 1; ++Andex)
-	{
-		TArray<FLoopNode*>& SubLoopA = Cell.OuterLoopSubdivision[Andex];
-		for (int32 Bndex = Andex + 1; Bndex < SubdivisionCount; ++Bndex)
-		{
-			TArray<FLoopNode*>& SubLoopB = Cell.OuterLoopSubdivision[Bndex];
-			TryToConnectTwoSubLoopsWithShortestSegment(Cell, SubLoopA, SubLoopB);
-		}
-	}
-	Cell.SelectSegmentInCandidateSegments(IsoSegmentFactory);
-}
-
-void FIsoTriangulator::FindSegmentToLinkOuterToInnerLoopNodes(FCell& Cell)
-{
-	TArray<FLoopNode*>& OuterLoop = Cell.SubLoops[0];
-
-	TFunction<void(FLoopNode&, FLoopNode&)> CreateCandidateSegment = [&](FLoopNode& Node, FLoopNode& CandidatNode)
-	{
-		FIsoSegment& Segment = IsoSegmentFactory.New();
-		Segment.Init(Node, CandidatNode, ESegmentType::LoopToLoop);
-		Segment.SetCandidate();
-		Cell.CandidateSegments.Add(&Segment);
-
-#ifdef DEBUG_FIND_SEGMENT_TO_LINK_OUTER_TO_INNER_LOOP
-		if (bDisplay)
-		{
-			{
-				F3DDebugSession _(FString::Printf(TEXT("Candidate Segment")));
-				Grid.DisplayIsoSegment(EGridSpace::UniformScaled, Node, CandidatNode, 0, EVisuProperty::RedCurve);
-			}
-		}
-#endif 
-	};
-
-	const SlopeMethod GetSlopeAtStartNode = CounterClockwiseSlope;
-	const SlopeMethod GetSlopeAtEndNode = ClockwiseSlope;
-	TFunction<void(FIsoSegment*)> FindBestTriangle = [&](FIsoSegment* Segment)
-	{
-		FLoopNode& StartNode = (FLoopNode&)Segment->GetFirstNode();
-		FLoopNode& EndNode = (FLoopNode&)Segment->GetSecondNode();
-		if (StartNode.IsThinZoneNode() || EndNode.IsThinZoneNode())
-		{
-			return;
-		}
-
-		if ((StartNode.GetConnectedSegments().Num() > 2) || (EndNode.GetConnectedSegments().Num() > 2))
-		{
-			return;
-		}
-
-		FIsoSegment& PreviousSegment = StartNode.GetPreviousSegment();
-		FIsoSegment& NextSegment = EndNode.GetNextSegment();
-
-		FIsoNode& PreviousNode = StartNode.GetPreviousNode();
-		FIsoNode& NextNode = EndNode.GetNextNode();
-
-#ifdef DEBUG_FIND_SEGMENT_TO_LINK_OUTER_TO_INNER_LOOP
-		static int32 SegmentIndex = 0;
-		if (bDisplay)
-		{
-			{
-				SegmentIndex++;
-				F3DDebugSession _(FString::Printf(TEXT("Start Segment %d"), SegmentIndex));
-				Grid.DisplayIsoSegment(EGridSpace::UniformScaled, *Segment, SegmentIndex, EVisuProperty::BlueCurve);
-				Grid.DisplayIsoSegment(EGridSpace::UniformScaled, PreviousSegment, SegmentIndex, EVisuProperty::GreenCurve);
-				Grid.DisplayIsoSegment(EGridSpace::UniformScaled, NextSegment, SegmentIndex, EVisuProperty::GreenCurve);
-			}
-		}
-#endif
-
-		const FPoint2D& StartPoint2D = StartNode.Get2DPoint(EGridSpace::UniformScaled, Grid);
-		const FPoint2D& EndPoint2D = EndNode.Get2DPoint(EGridSpace::UniformScaled, Grid);
-
-		const FPoint& StartPoint3D = StartNode.Get3DPoint(Grid);
-		const FPoint& EndPoint3D = EndNode.Get3DPoint(Grid);
-
-		// StartMaxSlope and EndMaxSlope are at most equal to 4, because if the slop with candidate node is biggest to 4, the nez triangle will be inverted
-		double StartReferenceSlope = ComputePositiveSlope(StartPoint2D, EndPoint2D, 0);
-
-		double StartMaxSlope = GetSlopeAtStartNode(StartPoint2D, PreviousNode.Get2DPoint(EGridSpace::UniformScaled, Grid), StartReferenceSlope);
-		StartMaxSlope /= 2.;
-		StartMaxSlope = FMath::Min(StartMaxSlope, 4.);
-
-		double EndReferenceSlope = StartReferenceSlope < 4 ? StartReferenceSlope + 4 : StartReferenceSlope - 4;
-		double EndMaxSlope = GetSlopeAtEndNode(EndPoint2D, NextNode.Get2DPoint(EGridSpace::UniformScaled, Grid), EndReferenceSlope);
-		EndMaxSlope /= 2.;
-		EndMaxSlope = FMath::Min(EndMaxSlope, 4.);
-
-		const double MinSlopeToNotBeAligned = 0.01;
-
-		FLoopNode* CandidatNode = nullptr;
-		double MinCriteria = HUGE_VALUE;
-		double CandidateSlopeAtStartNode = 8.;
-		double CandidateSlopeAtEndNode = 8.;
-
-		for (int32 Index = 1; Index < Cell.SubLoops.Num(); ++Index)
-		{
-			for (FLoopNode* Node : Cell.SubLoops[Index])
-			{
-				if (Node->IsThinZoneNode())
-				{
-					continue;
-				}
-
-				// Check if the node is inside the sector (X) or outside (Z)
-				const FPoint2D& NodePoint2D = Node->Get2DPoint(EGridSpace::UniformScaled, Grid);
-				double SlopeAtStartNode = GetSlopeAtStartNode(StartPoint2D, NodePoint2D, StartReferenceSlope);
-				double SlopeAtEndNode = GetSlopeAtEndNode(EndPoint2D, NodePoint2D, EndReferenceSlope);
-
-				if (SlopeAtStartNode <= 0 || SlopeAtStartNode >= StartMaxSlope)
-				{
-					continue;
-				}
-
-				if (SlopeAtEndNode <= 0 || SlopeAtEndNode >= EndMaxSlope)
-				{
-					continue;
-				}
-
-				double PointCriteria = IsoTriangulatorImpl::IsoscelesCriteriaMax(StartPoint2D, EndPoint2D, NodePoint2D);
-
-				if (
-					// the candidate triangle is inside the current candidate triangle
-					((SlopeAtStartNode < (CandidateSlopeAtStartNode + MinSlopeToNotBeAligned)) && (SlopeAtEndNode < (CandidateSlopeAtEndNode + MinSlopeToNotBeAligned))) ||
-					// or the candidate triangle is better the current candidate triangle and doesn't contain the current candidate triangle
-					((PointCriteria < MinCriteria) && ((SlopeAtStartNode > CandidateSlopeAtStartNode) ^ (SlopeAtEndNode > CandidateSlopeAtEndNode))))
-				{
-
-					if (LoopSegmentsIntersectionTool.DoesIntersect(StartNode, *Node))
-					{
-						continue;
-					}
-
-					if (LoopSegmentsIntersectionTool.DoesIntersect(EndNode, *Node))
-					{
-						continue;
-					}
-
-					MinCriteria = PointCriteria;
-					CandidatNode = Node;
-					CandidateSlopeAtStartNode = SlopeAtStartNode;
-					CandidateSlopeAtEndNode = SlopeAtEndNode;
-				}
-			}
-		}
-
-		if (CandidatNode)
-		{
-			CreateCandidateSegment(StartNode, *CandidatNode);
-			CreateCandidateSegment(EndNode, *CandidatNode);
-		}
-	};
-
-	for (int32 Index = 0; Index < OuterLoop.Num() - 1; ++Index)
-	{
-		FLoopNode* Node = OuterLoop[Index];
-		if (&Node->GetNextNode() != OuterLoop[Index + 1])
-		{
-			continue;
-		}
-
-		FIsoSegment* Segment = Node->GetSegmentConnectedTo(OuterLoop[Index + 1]);
-		if (Segment == nullptr)
-		{
-			continue;
-		}
-		FindBestTriangle(Segment);
-	}
-
-	Cell.SelectSegmentInCandidateSegments(IsoSegmentFactory);
-}
-
-void FIsoTriangulator::FindIsoSegmentToLinkOuterLoopNodes(FCell& Cell)
-{
-	TArray<FLoopNode*>& OuterLoop = Cell.SubLoops[0];
-	int32 NodeCount = OuterLoop.Num();
-
-	{
-#ifdef DEBUG_FIND_ISO_SEGMENT_TO_LINK_OUTER_LOOP_NODES
-		F3DDebugSession _(bDisplay, ("FindIsoSegmentToLinkOuterLoopNodes 1"));
-#endif
-		int SubdivisionCount = Cell.OuterLoopSubdivision.Num();
-		for (int32 Andex = 0; Andex < SubdivisionCount - 1; ++Andex)
-		{
-			TArray<FLoopNode*>& SubLoopA = Cell.OuterLoopSubdivision[Andex];
-			for (int32 Bndex = Andex + 1; Bndex < SubdivisionCount; ++Bndex)
-			{
-				TArray<FLoopNode*>& SubLoopB = Cell.OuterLoopSubdivision[Bndex];
-				TryToConnectTwoSubLoopsWithTheMostIsoSegment(Cell, SubLoopA, SubLoopB);
-			}
-		}
-	}
-
-	{
-#ifdef DEBUG_FIND_ISO_SEGMENT_TO_LINK_OUTER_LOOP_NODES
-		F3DDebugSession _(bDisplay, ("FindIsoSegmentToLinkOuterLoopNodes 2"));
-#endif
-		for (TArray<FLoopNode*>& SubLoop : Cell.OuterLoopSubdivision)
-		{
-			TryToConnectVertexSubLoopWithTheMostIsoSegment(Cell, SubLoop);
-		}
-	}
-
-	Cell.SelectSegmentInCandidateSegments(IsoSegmentFactory);
-}
-
 // =============================================================================================================
 // 	   For each cell
 // 	      - Connect loops together and to cell vertices
@@ -1374,6 +937,11 @@ void FIsoTriangulator::ConnectCellLoops()
 	TArray<FCell> Cells;
 	FindCellContainingBoundaryNodes(Cells);
 
+	for (FCell& Cell : Cells)
+	{
+		Cell.InitLoopConnexions();
+	}
+
 #ifdef DEBUG_CONNECT_CELL_LOOPS
 	F3DDebugSession _(bDisplay, ("ConnectCellLoops"));
 #endif
@@ -1384,37 +952,22 @@ void FIsoTriangulator::ConnectCellLoops()
 #ifdef DEBUG_CONNECT_CELL_LOOPS
 		F3DDebugSession _(bDisplay, FString::Printf(TEXT("Cell %d"), Cell.Id));
 		DisplayCell(Cell);
+		Wait();
 #endif
 
-		ConnectCellSubLoopsByNeighborhood(Cell);
-
-#ifdef DEBUG_CONNECT_CELL_LOOPS
-		Grid.DisplayIsoSegments(TEXT("ConnectCellLoops::Step 0"), EGridSpace::UniformScaled, Cell.FinalSegments, false, false, false, EVisuProperty::YellowCurve);
-#endif
-
-		if (Cell.bHasOuterLoop)
+		if (Cell.LoopCells.Num())
 		{
-			FindIsoSegmentToLinkOuterLoopNodes(Cell);
-#ifdef DEBUG_CONNECT_CELL_LOOPS
-			Grid.DisplayIsoSegments(TEXT("ConnectCellLoops::Step 1"), EGridSpace::UniformScaled, Cell.FinalSegments, false, false, false, EVisuProperty::YellowCurve);
-#endif
+			Cell.ConnectLoopsByNeighborhood();
 
-			if (Cell.CandidateSegments.Num() == 0)
-			{
-				FindSegmentToLinkOuterLoopNodes(Cell);
 #ifdef DEBUG_CONNECT_CELL_LOOPS
-				Grid.DisplayIsoSegments(TEXT("ConnectCellLoops::Step 2"), EGridSpace::UniformScaled, Cell.FinalSegments, false, false, false, EVisuProperty::YellowCurve);
+			Grid.DisplayIsoSegments(TEXT("ConnectCellLoops::Step 0"), EGridSpace::UniformScaled, Cell.FinalSegments, false, false, false, EVisuProperty::YellowCurve);
 #endif
 
-				FindSegmentToLinkOuterToInnerLoopNodes(Cell);
-#ifdef DEBUG_CONNECT_CELL_LOOPS
-				Grid.DisplayIsoSegments(TEXT("ConnectCellLoops::Step 3"), EGridSpace::UniformScaled, Cell.FinalSegments, false, false, false, EVisuProperty::YellowCurve);
-#endif
-			}
+			ConnectCellCornerToLoops(Cell);
 		}
-		ConnectCellCornerToInnerLoop(Cell);
+
 #ifdef DEBUG_CONNECT_CELL_LOOPS
-		Grid.DisplayIsoSegments(TEXT("ConnectCellLoops::Step 4"), EGridSpace::UniformScaled, Cell.FinalSegments, false, false, false, EVisuProperty::YellowCurve);
+		Grid.DisplayIsoSegments(TEXT("ConnectCellLoops::Step 1"), EGridSpace::UniformScaled, Cell.FinalSegments, false, false, false, EVisuProperty::YellowCurve);
 #endif
 
 		FinalToLoops.Append(Cell.FinalSegments);
@@ -1521,7 +1074,7 @@ void FIsoTriangulator::FindCellContainingBoundaryNodes(TArray<FCell>& Cells)
 		{
 			if (CellIndex != NodeToCellIndices[Index])
 			{
-				Cells.Emplace(CellIndex, CellNodes, Grid);
+				Cells.Emplace(CellIndex, CellNodes, *this);
 
 				CellIndex = NodeToCellIndices[Index];
 				CellNodes.Reset(LoopNodeCount);
@@ -1537,16 +1090,18 @@ void FIsoTriangulator::FindCellContainingBoundaryNodes(TArray<FCell>& Cells)
 				CellNodes.Add(&LoopNode);
 			}
 		}
-		Cells.Emplace(CellIndex, CellNodes, Grid);
+		Cells.Emplace(CellIndex, CellNodes, *this);
 
 #ifdef DEBUG_BUILD_CELLS
 		DisplayCell(Cells.Last());
 		Wait(bDisplay);
 #endif
 	}
+
 	FChrono::Elapse(StartTime);
 }
 
+#ifdef OLD
 void FIsoTriangulator::FindCandidateSegmentsToLinkInnerAndLoop()
 {
 	const double FlatAngle = 0.1;
@@ -1558,14 +1113,12 @@ void FIsoTriangulator::FindCandidateSegmentsToLinkInnerAndLoop()
 	{
 		FIsoSegment& SegCandidate = IsoSegmentFactory.New();
 		SegCandidate.Init(*InnerNode, LoopNode, ESegmentType::InnerToLoop);
-		NewTestSegments.Add(&SegCandidate);
 	};
 
 	TFunction<void(FLoopNode&, FLoopNode&)> CreateCandidateBoundarySegment = [&](FLoopNode& StartNode, FLoopNode& EndNode)
 	{
 		FIsoSegment& SegCandidate = IsoSegmentFactory.New();
 		SegCandidate.Init(StartNode, EndNode, ESegmentType::LoopToLoop);
-		NewTestSegments.Add(&SegCandidate);
 	};
 
 	int32 CountU = Grid.GetCuttingCount(EIso::IsoU);
@@ -1722,6 +1275,7 @@ void FIsoTriangulator::FindCandidateSegmentsToLinkInnerAndLoop()
 	Chronos.FindSegmentToLinkInnerToLoopDuration = FChrono::Elapse(StartTime);
 #endif
 }
+#endif
 
 void FIsoTriangulator::ConnectUnconnectedInnerSegments()
 {
@@ -2279,8 +1833,8 @@ void FIsoTriangulator::TriangulateInnerNodes()
 			Grid.DisplayTriangle(EGridSpace::UniformScaled, *GlobalIndexToIsoInnerNodes[OppositIndex], *GlobalIndexToIsoInnerNodes[OppositIndex - 1], *GlobalIndexToIsoInnerNodes[Index]);
 #endif 
 
-			Mesh.AddTriangle(GlobalIndexToIsoInnerNodes[Index]->GetFaceIndex(), GlobalIndexToIsoInnerNodes[Index + 1]->GetFaceIndex(), GlobalIndexToIsoInnerNodes[OppositIndex]->GetFaceIndex());
-			Mesh.AddTriangle(GlobalIndexToIsoInnerNodes[OppositIndex]->GetFaceIndex(), GlobalIndexToIsoInnerNodes[OppositIndex - 1]->GetFaceIndex(), GlobalIndexToIsoInnerNodes[Index]->GetFaceIndex());
+			Mesh.AddTriangle(GlobalIndexToIsoInnerNodes[Index]->GetGlobalIndex(), GlobalIndexToIsoInnerNodes[Index + 1]->GetGlobalIndex(), GlobalIndexToIsoInnerNodes[OppositIndex]->GetGlobalIndex());
+			Mesh.AddTriangle(GlobalIndexToIsoInnerNodes[OppositIndex]->GetGlobalIndex(), GlobalIndexToIsoInnerNodes[OppositIndex - 1]->GetGlobalIndex(), GlobalIndexToIsoInnerNodes[Index]->GetGlobalIndex());
 		}
 		Index++;
 	}
@@ -2289,213 +1843,194 @@ void FIsoTriangulator::TriangulateInnerNodes()
 #endif
 }
 
-void FIsoTriangulator::ConnectCellSubLoopsByNeighborhood(FCell& Cell)
+void FCell::ConnectLoopsByNeighborhood()
 {
-	FTimePoint StartTime = FChrono::Now();
-
-	int32 LoopCount = Cell.SubLoops.Num();
-
-	TArray<TPair<int32, FPoint2D>> LoopBarycenters;
-	LoopBarycenters.Reserve(LoopCount + 4);
+#ifdef DEBUG_CONNECT_CELL_SUB_LOOPS_BY_NEIGHBORHOOD
+	F3DDebugSession _(Grid.bDisplay, FString::Printf(TEXT("Cell % d: Build Segments to Connect inner loops"), Id));
+#endif
 
-	int32 LoopIndex = -1;
-	for (const TArray<FLoopNode*>& Nodes : Cell.SubLoops)
+	for (FLoopConnexion& LoopConnexion : LoopConnexions)
 	{
-		++LoopIndex;
-
-		// the external loop is not processed 
-		if (Nodes[0]->GetLoopIndex() == 0)
-		{
-			continue;
-		}
-
-		TPair<int32, FPoint2D>& BaryCenterObj = LoopBarycenters.Emplace_GetRef(LoopIndex, FPoint2D::ZeroPoint);
-		FPoint2D& BaryCenter = BaryCenterObj.Value;
-		for (const FLoopNode* Node : Nodes)
-		{
-			BaryCenter += Node->Get2DPoint(EGridSpace::UniformScaled, Grid);
-		}
-		BaryCenter /= (double)Nodes.Num();
+		TryToConnectTwoSubLoopsWithShortestSegment(LoopConnexion);
 	}
 
-	TArray<int32> EdgeVertexIndices;
-	if (Cell.bHasOuterLoop && LoopCount < 5)
+	TArray<FLoopConnexion*> LoopConnexionPtrs;
+	LoopConnexionPtrs.Reserve(LoopConnexions.Num());
+	for (FLoopConnexion& Connexion : LoopConnexions)
 	{
-		EdgeVertexIndices.Reserve(6);
-		Cell.BorderLoopIndices.Reserve(3);
-		if (LoopCount == 2)
-		{
-			Cell.BorderLoopIndices.Add(1);
-		}
-		else if (LoopCount == 3)
-		{
-			EdgeVertexIndices.Append({ 1, 2 });
-			Cell.BorderLoopIndices.Append({ 1, 2 });
-		}
-		else if (LoopCount == 4)
-		{
-			EdgeVertexIndices.Append({ 1, 2, 2, 3, 3, 1 });
-			Cell.BorderLoopIndices.Append({ 1, 2, 3 });
-		}
+		LoopConnexionPtrs.Add(&Connexion);
 	}
-	else if (LoopBarycenters.Num() < 4)
-	{
-		EdgeVertexIndices.Reserve(6);
-		Cell.BorderLoopIndices.Reserve(3);
-		if (LoopCount == 1)
+
+	Algo::Sort(LoopConnexionPtrs, [&](const FLoopConnexion* LoopConnexion1, const FLoopConnexion* LoopConnexion2)
 		{
-			Cell.BorderLoopIndices.Add(0);
-		}
-		else if (LoopCount == 2)
+			return LoopConnexion1->MinDistance < LoopConnexion2->MinDistance;
+		});
+
+	{
+#ifdef DEBUG_CONNECT_CELL_SUB_LOOPS_BY_NEIGHBORHOOD
+		F3DDebugSession _(TEXT("Best Path"));
+#endif
+		for (FLoopConnexion* LoopConnexion : LoopConnexionPtrs)
 		{
-			EdgeVertexIndices.Append({ 0, 1 });
-			Cell.BorderLoopIndices.Append({ 0, 1 });
+			if (!LoopConnexion->bIsConnexionWithBorder && LoopConnexion->IsShortestPath())
+			{
+				TryToCreateSegment(*LoopConnexion);
+			}
 		}
-		else if (LoopCount == 3)
+
+		for (FLoopConnexion* LoopConnexion : LoopConnexionPtrs)
 		{
-			EdgeVertexIndices.Append({ 0, 1, 1, 2, 2, 0 });
-			Cell.BorderLoopIndices.Append({ 0, 1, 2 });
+			if (LoopConnexion->bIsConnexionWithBorder)
+			{
+				TryToCreateSegment(*LoopConnexion);
+			}
 		}
 	}
-	else
-	{
-#ifdef DEBUG_BOWYERWATSON
-		FBowyerWatsonTriangulator::bDisplay = bDisplay;
+
+#ifdef DEBUG_CONNECT_CELL_SUB_LOOPS_BY_NEIGHBORHOOD
+	Grid.DisplayIsoSegments(TEXT("Build Segments Connect inner close loops"), EGridSpace::UniformScaled, CandidateSegments, false, false, false, EVisuProperty::YellowCurve);
+	Wait(Grid.bDisplay);
 #endif
-		FBowyerWatsonTriangulator Triangulator(LoopBarycenters, EdgeVertexIndices);
-		Triangulator.Triangulate();
-		Triangulator.GetOuterVertices(Cell.BorderLoopIndices);
-	}
 
-	// Connect inner close loops 
-	// ==========================================================================================
+	if(CandidateSegments.Num())
 	{
+		const bool bSelectWithCellSizeCriteria = true;
+		SelectSegmentInCandidateSegments(Triangulator.IsoSegmentFactory, bSelectWithCellSizeCriteria);
+	}
+
 #ifdef DEBUG_CONNECT_CELL_SUB_LOOPS_BY_NEIGHBORHOOD
-		//F3DDebugSession _(TEXT("Build Segments Connect inner close loops "));
+	Grid.DisplayIsoSegments(TEXT("Build Segments Connect inner close loops"), EGridSpace::UniformScaled, FinalSegments, false, false, false, EVisuProperty::YellowCurve);
+	Wait(Grid.bDisplay);
 #endif
 
-		for (int32 Index = 0; Index < EdgeVertexIndices.Num();)
-		{
-			int32 IndexLoopA = EdgeVertexIndices[Index++];
-			int32 IndexLoopB = EdgeVertexIndices[Index++];
+}
 
-			const TArray<FLoopNode*>& SubLoopA = Cell.SubLoops[IndexLoopA];
-			const TArray<FLoopNode*>& SubLoopB = Cell.SubLoops[IndexLoopB];
+void FCell::TryToCreateSegment(FLoopConnexion& LoopConnexion)
+{
+	const FPoint2D& ACoordinates = LoopConnexion.NodeA->Get2DPoint(EGridSpace::UniformScaled, Grid);
+	const FPoint2D& BCoordinates = LoopConnexion.NodeB->Get2DPoint(EGridSpace::UniformScaled, Grid);
 
-#ifdef DEBUG_CONNECT_CELL_SUB_LOOPS_BY_NEIGHBORHOOD
-			//F3DDebugSession _(TEXT("Segment"));
-#endif
-			TryToConnectTwoSubLoopsWithShortestSegment(Cell, SubLoopA, SubLoopB);
-		}
+	Triangulator.TryToCreateSegment(*this, LoopConnexion.NodeA, ACoordinates, LoopConnexion.NodeB, BCoordinates, 0.1);
+}
 
-		Cell.SelectSegmentInCandidateSegments(IsoSegmentFactory);
-#ifdef DEBUG_CONNECT_CELL_SUB_LOOPS_BY_NEIGHBORHOOD
-		Grid.DisplayIsoSegments(TEXT("Build Segments Connect inner close loops"), EGridSpace::UniformScaled, Cell.FinalSegments, false, false, false, EVisuProperty::YellowCurve);
-		Wait();
-#endif
+
+void FCell::InitLoopConnexions()
+{
+	FTimePoint StartTime = FChrono::Now();
+
+	switch (InnerLoopCount)
+	{
+	case 0:
+	{
+		break;
+	}
+	case 1:
+	{
+		LoopConnexions.Reserve( OuterLoopCount );
+		LoopCellBorderIndices = { OuterLoopCount };
+		break;
 	}
+	case 2:
+	{
+		LoopConnexions.Reserve(1 + 2 * OuterLoopCount);
+		LoopConnexions.Emplace(LoopCells[OuterLoopCount], LoopCells[OuterLoopCount + 1]);
 
+		LoopCellBorderIndices = { OuterLoopCount , OuterLoopCount + 1 };
+		break;
+	}
+	case 3:
+	{
+		LoopConnexions.Reserve(3 + 3 * OuterLoopCount);
+		LoopConnexions.Emplace(LoopCells[OuterLoopCount], LoopCells[OuterLoopCount + 1]);
+		LoopConnexions.Emplace(LoopCells[OuterLoopCount], LoopCells[OuterLoopCount + 2]);
+		LoopConnexions.Emplace(LoopCells[OuterLoopCount + 1], LoopCells[OuterLoopCount + 2]);
 
-	// With Outer loop
-	// ==========================================================================================
-	if (Cell.bHasOuterLoop && Cell.SubLoops.Num() > 1)
+		LoopCellBorderIndices = { OuterLoopCount , OuterLoopCount + 1 , OuterLoopCount + 2 };
+		break;
+	}
+	default:
 	{
-#ifdef DEBUG_CONNECT_CELL_SUB_LOOPS_BY_NEIGHBORHOOD
-		//F3DDebugSession _(TEXT("Build Segments to Outer loop"));
+		TArray<TPair<int32, FPoint2D>> LoopBarycenters = GetLoopBarycenters();
+
+#ifdef DEBUG_BOWYERWATSON
+		FBowyerWatsonTriangulator::bDisplay = bDisplay;
 #endif
-		for (TArray<FLoopNode*>& SubLoopA : Cell.OuterLoopSubdivision)
+		TArray<int32> EdgeVertexIndices;
+		FBowyerWatsonTriangulator BWTriangulator(LoopBarycenters, EdgeVertexIndices);
+		BWTriangulator.Triangulate();
+		LoopCellBorderIndices = BWTriangulator.GetOuterVertices();
+		LoopConnexions.Reserve((EdgeVertexIndices.Num() >> 1) + LoopCellBorderIndices.Num() * OuterLoopCount);
+
+		for (int32 Index = 0; Index < EdgeVertexIndices.Num();)
 		{
-			for (int32 IndexBorderLoop : Cell.BorderLoopIndices)
-			{
-				int32 CandidateSegmentsCount = Cell.CandidateSegments.Num();
-				const TArray<FLoopNode*>& BorderLoop = Cell.SubLoops[IndexBorderLoop];
+			const int32 StartIndex = Index++;
+			const int32 EndIndex = Index++;
+			LoopConnexions.Emplace(LoopCells[EdgeVertexIndices[StartIndex]], LoopCells[EdgeVertexIndices[EndIndex]]);
+		}
+	}
+	}
 
-				TryToConnectTwoSubLoopsWithShortestSegment(Cell, SubLoopA, BorderLoop);
-				if (CandidateSegmentsCount == Cell.CandidateSegments.Num())
-				{
-					// if the subloops have not been connnected with TryToConnectTwoSubLoopsWithShortestSegment
-					// the try to connect them with an isoceles triangle
-					// Loops must be connected together
-					TryToConnectTwoLoopsWithIsocelesTriangle(Cell, SubLoopA, BorderLoop);
-				}
-			}
+	for (int32 Index = 0; Index < OuterLoopCount; ++Index)
+	{
+		for (int32 BIndex : LoopCellBorderIndices)
+		{
+			LoopConnexions.Emplace(LoopCells[Index], LoopCells[BIndex], true);
 		}
+	}
 
-		Cell.SelectSegmentInCandidateSegments(IsoSegmentFactory);
-#ifdef DEBUG_CONNECT_CELL_SUB_LOOPS_BY_NEIGHBORHOOD
-		Grid.DisplayIsoSegments(TEXT("Build Segments Connect inner to outer loop"), EGridSpace::UniformScaled, Cell.FinalSegments, false, false, false, EVisuProperty::YellowCurve);
-		Wait();
-#endif
+	const int32 ConnexionCount = LoopConnexions.Num();
+	for (FLoopCell& Loop : LoopCells)
+	{
+		Loop.Connexions.Reserve(ConnexionCount);
+	}
+
+	for (FLoopConnexion& Connexion : LoopConnexions)
+	{
+		Connexion.LinkToLoop();
 	}
 }
 
-void FIsoTriangulator::TryToConnectTwoSubLoopsWithShortestSegment(FCell& Cell, const TArray<FLoopNode*>& LoopA, const TArray<FLoopNode*>& LoopB)
+void FCell::TryToConnectTwoSubLoopsWithShortestSegment(FLoopConnexion& LoopConnexion)
 {
+	const TArray<FLoopNode*>& LoopA = LoopConnexion.Loop1.Nodes;
+	const TArray<FLoopNode*>& LoopB = LoopConnexion.Loop2.Nodes;
+
+#ifdef DEBUG_CONNECT_CELL_SUB_LOOPS_BY_NEIGHBORHOOD
+	F3DDebugSession A(Grid.bDisplay, TEXT("ConnectTwoSubLoopsWithShortestSegment"));
+#endif
 	double MinDistanceSquare = HUGE_VALUE_SQUARE;
-	int32 MinIndexA = -1;
-	int32 MinIndexB = -1;
 
-	for (int32 IndexA = 0; IndexA < LoopA.Num(); ++IndexA)
+	for (FLoopNode* NodeA : LoopA)
 	{
-		const FLoopNode* NodeA = LoopA[IndexA];
-		if (NodeA->IsDeleteOrThinNode())
-		{
-			continue;
-		}
 		const FPoint2D& ACoordinates = NodeA->Get2DPoint(EGridSpace::UniformScaled, Grid);
 
-		for (int32 IndexB = 0; IndexB < LoopB.Num(); ++IndexB)
+		for (FLoopNode* NodeB : LoopB)
 		{
-			const FLoopNode* NodeB = LoopB[IndexB];
-			if (NodeB->IsDeleteOrThinNode())
-			{
-				continue;
-			}
 			const FPoint2D& BCoordinates = NodeB->Get2DPoint(EGridSpace::UniformScaled, Grid);
 
+#ifdef DEBUG_CONNECT_CELL_SUB_LOOPS_BY_NEIGHBORHOOD_
+			Grid.DisplayIsoSegment(EGridSpace::UniformScaled, *NodeA, *NodeB, 0, EVisuProperty::BlueCurve);
+#endif
+
 			double SquareDistance = ACoordinates.SquareDistance(BCoordinates);
 			if (SquareDistance < MinDistanceSquare)
 			{
 				MinDistanceSquare = SquareDistance;
-				MinIndexA = IndexA;
-				MinIndexB = IndexB;
+				LoopConnexion.NodeA = NodeA;
+				LoopConnexion.NodeB = NodeB;
 			}
 		}
 	}
 
-	if (MinIndexA >= 0 && MinIndexB >= 0)
-	{
-		FLoopNode* NodeA = LoopA[MinIndexA];
-		const FPoint2D& ACoordinates = NodeA->Get2DPoint(EGridSpace::UniformScaled, Grid);
-		FLoopNode* NodeB = LoopB[MinIndexB];
-		const FPoint2D& BCoordinates = NodeB->Get2DPoint(EGridSpace::UniformScaled, Grid);
+	LoopConnexion.MinDistance = FMath::Sqrt(MinDistanceSquare);
 
-		//Display(EGridSpace::UniformScaled, *NodeA, *NodeB, 0, EVisuProperty::BlueCurve);
-		if (TryToCreateSegment(Cell, NodeA, ACoordinates, NodeB, BCoordinates, 0.1))
-		{
-			NodeA = &LoopA[MinIndexA]->GetNextNode();
-			if (!NodeA->IsDeleteOrThinNode() && !NodeB->IsDeleteOrThinNode())
-			{
-			if (Cell.Contains(NodeA))
-			{
-				NodeB = LoopB[MinIndexB];
-				TryToCreateSegment(Cell, NodeA, NodeA->Get2DPoint(EGridSpace::UniformScaled, Grid), NodeB, NodeB->Get2DPoint(EGridSpace::UniformScaled, Grid), 0.1);
-			}
-			}
-
-			NodeB = &LoopB[MinIndexB]->GetNextNode();
-			if (!NodeA->IsDeleteOrThinNode() && !NodeB->IsDeleteOrThinNode())
-			{
-				if (Cell.Contains(NodeB))
-				{
-					NodeA = LoopA[MinIndexA];
-					TryToCreateSegment(Cell, NodeA, NodeA->Get2DPoint(EGridSpace::UniformScaled, Grid), NodeB, NodeB->Get2DPoint(EGridSpace::UniformScaled, Grid), 0.1);
-				}
-			}
-		}
+#ifdef DEBUG_CONNECT_CELL_SUB_LOOPS_BY_NEIGHBORHOOD
+	if (LoopConnexion.NodeA && LoopConnexion.NodeB)
+	{
+		Grid.DisplayIsoSegment(EGridSpace::UniformScaled, *LoopConnexion.NodeA, *LoopConnexion.NodeB, 0, EVisuProperty::RedCurve);
 	}
-
-};
+#endif
+}
 
 void FIsoTriangulator::TryToConnectTwoLoopsWithIsocelesTriangle(FCell& Cell, const TArray<FLoopNode*>& LoopA, const TArray<FLoopNode*>& LoopB)
 {
@@ -2530,9 +2065,10 @@ void FIsoTriangulator::TryToConnectTwoLoopsWithIsocelesTriangle(FCell& Cell, con
 			{
 				TriangleIndex++;
 				F3DDebugSession _(FString::Printf(TEXT("Start Segment %d"), TriangleIndex));
-				Display(EGridSpace::UniformScaled, *Segment);
-				Display(EGridSpace::UniformScaled, StartNode, 0, EVisuProperty::RedPoint);
-				Display(EGridSpace::UniformScaled, EndNode);
+				Grid.DisplayIsoSegment(EGridSpace::UniformScaled, *Segment);
+				Grid.DisplayIsoNode(EGridSpace::UniformScaled, StartNode, 0, EVisuProperty::RedPoint);
+				Grid.DisplayIsoNode(EGridSpace::UniformScaled, EndNode);
+				Wait();
 			}
 		}
 #endif
@@ -2865,7 +2401,7 @@ bool FIsoTriangulator::TryToCreateSegment(FCell& Cell, FLoopNode* NodeA, const F
 	return true;
 };
 
-void FIsoTriangulator::ConnectCellCornerToInnerLoop(FCell& Cell)
+void FIsoTriangulator::ConnectCellCornerToLoops(FCell& Cell)
 {
 	FIsoInnerNode* CellNodes[4];
 	int32 Index = Cell.Id;
@@ -2895,12 +2431,12 @@ void FIsoTriangulator::ConnectCellCornerToInnerLoop(FCell& Cell)
 	F3DDebugSession _(TEXT("With cell corners"));
 #endif
 
-	TFunction<void(int32, FIsoInnerNode*)> FindAndTryCreateCandidateSegmentToLinkLoopToCorner = [&](int32 IndexLoopA, FIsoInnerNode* InnerNode)
+	TFunction<void(FLoopCell&, FIsoInnerNode*)> FindAndTryCreateCandidateSegmentToLinkLoopToCorner = [&](FLoopCell& LoopCell, FIsoInnerNode* InnerNode)
 	{
 
 		const FPoint2D& InnerCoordinates = InnerNode->Get2DPoint(EGridSpace::UniformScaled, Grid);
 
-		const TArray<FLoopNode*>& LoopA = Cell.SubLoops[IndexLoopA];
+		const TArray<FLoopNode*>& LoopA = LoopCell.Nodes;
 
 		double MinDistanceSquare = HUGE_VALUE_SQUARE;
 		int32 MinIndexA = -1;
@@ -2940,25 +2476,22 @@ void FIsoTriangulator::ConnectCellCornerToInnerLoop(FCell& Cell)
 	{
 		if (CellNodes[ICell])
 		{
-			for (int32 IndexLoopA : Cell.BorderLoopIndices)
+			for (FLoopCell& LoopCell : Cell.LoopCells)
 			{
 #ifdef DEBUG_CONNECTCELLCORNERTOINNERLOOP
 				DisplayPoint(CellNodes[ICell]->Get2DPoint(EGridSpace::UniformScaled, Grid), EVisuProperty::GreenPoint);
 #endif
-				FindAndTryCreateCandidateSegmentToLinkLoopToCorner(IndexLoopA, CellNodes[ICell]);
+				FindAndTryCreateCandidateSegmentToLinkLoopToCorner(LoopCell, CellNodes[ICell]);
 			}
 
-			if (Cell.bHasOuterLoop)
-			{
-#ifdef DEBUG_CONNECTCELLCORNERTOINNERLOOP
-				DisplayPoint(CellNodes[ICell]->Get2DPoint(EGridSpace::UniformScaled, Grid), EVisuProperty::GreenPoint);
-#endif
-				FindAndTryCreateCandidateSegmentToLinkLoopToCorner(0, CellNodes[ICell]);
-			}
 		}
 	}
 
-	Cell.SelectSegmentInCandidateSegments(IsoSegmentFactory);
+	if(Cell.CandidateSegments.Num())
+	{
+		const bool bSelectWithCellSizeCriteria = false;
+		Cell.SelectSegmentInCandidateSegments(IsoSegmentFactory, bSelectWithCellSizeCriteria);
+	}
 }
 
 } //namespace UE::CADKernel
\ No newline at end of file
diff --git a/Engine/Source/Runtime/Datasmith/CADKernel/Private/CADKernel/Mesh/Meshers/IsoTriangulator.h b/Engine/Source/Runtime/Datasmith/CADKernel/Private/CADKernel/Mesh/Meshers/IsoTriangulator.h
--- a/Engine/Source/Runtime/Datasmith/CADKernel/Private/CADKernel/Mesh/Meshers/IsoTriangulator.h
+++ b/Engine/Source/Runtime/Datasmith/CADKernel/Private/CADKernel/Mesh/Meshers/IsoTriangulator.h
@@ -26,6 +26,8 @@ struct FCell;
 
 using FMeshPolygonFunc = TFunction<void(const FGrid&, FIsoNode*[], FFaceMesh&)>;
 
+struct FLoopConnexion;
+
 struct CADKERNEL_API FIsoTriangulatorChronos
 {
 	FDuration TriangulateDuration = FChrono::Init();
@@ -77,6 +79,7 @@ class FIsoTriangulator
 	friend class FCycleTriangulator;
 	friend class FLoopCleaner;
 	friend class FParametricMesher;
+	friend struct FCell;
 
 protected:
 
@@ -114,7 +117,7 @@ protected:
 	FIntersectionSegmentTool LoopSegmentsIntersectionTool;
 	FIntersectionSegmentTool InnerSegmentsIntersectionTool;
 	FIntersectionSegmentTool InnerToLoopSegmentsIntersectionTool;
-	FIntersectionNodePairTool InnerToOuterSegmentsIntersectionTool;
+	FIntersectionNodePairTool InnerToOuterIsoSegmentsIntersectionTool;
 	FIntersectionSegmentTool ThinZoneIntersectionTool;
 
 
@@ -137,8 +140,6 @@ protected:
 	 */
 	TArray<FIsoSegment*> CandidateSegments;
 
-	TArray<FIsoSegment*> NewTestSegments;
-
 	bool bDisplay = false;
 
 	bool bNeedCheckOrientation = false;
@@ -207,29 +208,7 @@ public:
 	void ConnectCellLoops();
 	void FindCellContainingBoundaryNodes(TArray<FCell>& Cells);
 
-	/**
-	 * The closest loops are connected together
-	 * To do it, a Delaunay triangulation of the loop barycenter is realized.
-	 * Each edge of this mesh defined a near loops pair
-	 * The shortest segment is then build between this pair of loops
-	 */
-	void ConnectCellSubLoopsByNeighborhood(FCell& cell);
-
-	void FindIsoSegmentToLinkOuterLoopNodes(FCell& Cell);	
-
-	void FindSegmentToLinkOuterLoopNodes(FCell& Cell);
-
-	void FindSegmentToLinkOuterToInnerLoopNodes(FCell& Cell);
-
-	void ConnectCellCornerToInnerLoop(FCell& Cell);
-
-	/**
-	 * The goal of this algorithm is to connect iso U (or V) aligned loop nodes as soon as they are nearly in the same iso V (or U) strip.
-	 * I.e.:
-	 * - Iso U aligned: NodeA.U = NodeB.U +/-TolU
-	 * - In the same strip: each node of the segment has the same index "i" that verify: isoV[i] - TolV < Node.V < isoV[i+1] + TolV
-	 */
-	void FindIsoSegmentToLinkLoopToLoop();
+	void ConnectCellCornerToLoops(FCell& Cell);
 
 	void FindCandidateSegmentsToLinkInnerAndLoop();
 
@@ -247,7 +226,7 @@ public:
 
 	/**
 	 * Find in the network a minimal cycle stating from a segment
-	 * @return false if the new cycle crosses a segement already used
+	 * @return false if the new cycle crosses a segment already used
 	 */
 	bool FindCycle(FIsoSegment* StartSegment, bool bLeftSide, TArray<FIsoSegment*>& Cycle, TArray<bool>& CycleOrientation);
 
@@ -297,6 +276,11 @@ public:
 			});
 	}
 
+	const FGrid& GetGrid() const
+	{
+		return Grid;
+	}
+
 private:
 
 	FIsoSegment* FindNextSegment(EGridSpace Space, const FIsoSegment* StartSegment, const FIsoNode* StartNode, SlopeMethod GetSlop) const;
@@ -305,20 +289,6 @@ private:
 	// 	   Create segments
 	// ==========================================================================================
 
-	/**
-	 *  SubLoopA                  SubLoopB
-	 *      --X---X             X-----X--
-	 *             \           /
-	 *              \         /
-	 *               X=======X
-	 *              /         \
-	 *             /           \
-	 *      --X---X             X-----X--
-	 *
-	 *     ======= ShortestSegment
-	 */
-	void TryToConnectTwoSubLoopsWithShortestSegment(FCell& Cell, const TArray<FLoopNode*>& SubLoopA, const TArray<FLoopNode*>& SubLoopB);
-
 	void TryToConnectTwoLoopsWithIsocelesTriangle(FCell& Cell, const TArray<FLoopNode*>& SubLoopA, const TArray<FLoopNode*>& SubLoopB);
 
 	/**
diff --git a/Engine/Source/Runtime/Datasmith/CADKernel/Private/CADKernel/Mesh/Meshers/IsoTriangulator/IntersectionSegmentTool.cpp b/Engine/Source/Runtime/Datasmith/CADKernel/Private/CADKernel/Mesh/Meshers/IsoTriangulator/IntersectionSegmentTool.cpp
--- a/Engine/Source/Runtime/Datasmith/CADKernel/Private/CADKernel/Mesh/Meshers/IsoTriangulator/IntersectionSegmentTool.cpp
+++ b/Engine/Source/Runtime/Datasmith/CADKernel/Private/CADKernel/Mesh/Meshers/IsoTriangulator/IntersectionSegmentTool.cpp
@@ -14,12 +14,13 @@ namespace UE::CADKernel
 const FIsoSegment* FIntersectionSegmentTool::FindIntersectingSegment(const FIsoSegment& Segment) const
 {
 	using namespace IntersectionSegmentTool;
-	return TIntersectionSegmentTool<IntersectionSegmentTool::FSegment>::FindIntersectingSegment(&Segment.GetFirstNode(), &Segment.GetSecondNode());
+	const IntersectionSegmentTool::FSegment* IntersectingSegment = TIntersectionSegmentTool<IntersectionSegmentTool::FSegment>::FindIntersectingSegment(&Segment.GetFirstNode(), &Segment.GetSecondNode());
+	return IntersectingSegment ? IntersectingSegment->GetIsoSegment() : nullptr;
 }
 
 FIsoSegment* FIntersectionSegmentTool::FindIntersectingSegment(const FIsoSegment& Segment)
 {
-	return const_cast<FIsoSegment*> (static_cast<const FIntersectionSegmentTool*>(this)->TIntersectionSegmentTool<IntersectionSegmentTool::FSegment>::FindIntersectingSegment(&Segment.GetFirstNode(), &Segment.GetSecondNode()));
+	return const_cast<FIsoSegment*> (static_cast<const FIntersectionSegmentTool*>(this)->FindIntersectingSegment(Segment));
 }
 
 bool FIntersectionSegmentTool::DoesIntersect(const FIsoNode* StartNode, const FIsoNode* EndNode, FIsoSegment** Segment) const
diff --git a/Engine/Source/Runtime/Datasmith/CADKernel/Private/CADKernel/Mesh/Meshers/IsoTriangulator/IntersectionSegmentTool.h b/Engine/Source/Runtime/Datasmith/CADKernel/Private/CADKernel/Mesh/Meshers/IsoTriangulator/IntersectionSegmentTool.h
--- a/Engine/Source/Runtime/Datasmith/CADKernel/Private/CADKernel/Mesh/Meshers/IsoTriangulator/IntersectionSegmentTool.h
+++ b/Engine/Source/Runtime/Datasmith/CADKernel/Private/CADKernel/Mesh/Meshers/IsoTriangulator/IntersectionSegmentTool.h
@@ -64,7 +64,7 @@ struct FSegment
 		return nullptr;
 	}
 
-	EConnectionType IsSuperimposed(const FSegment2D& SegmentAB, const FSegment2D& SegmentCD, bool bSameOrientation) const
+	static EConnectionType IsSuperimposed(const FSegment2D& SegmentAB, const FSegment2D& SegmentCD, bool bSameOrientation)
 	{
 		const FPoint2D AB = SegmentAB.GetVector().Normalize();
 		const FPoint2D CD = SegmentCD.GetVector().Normalize();
@@ -158,7 +158,7 @@ struct FSegment
 		return true;
 	}
 
-	bool DoesItIntersect(const FSegment& Segment) const
+	virtual bool DoesItIntersect(const FSegment& Segment) const
 	{
 		if (!CouldItIntersect(Segment.Boundary))
 		{
@@ -167,6 +167,11 @@ struct FSegment
 
 		return DoIntersect(Segment2D, Segment.Segment2D);
 	}
+
+	bool IsParallelWith(const FSegment& Segment) const
+	{
+		return AreParallel(Segment2D, Segment.Segment2D);
+	}
 };
 }
 
@@ -226,6 +231,16 @@ struct FSegment : public IntersectionToolBase::FSegment
 		return EndNode;
 	}
 
+	virtual bool DoesItIntersect(const IntersectionToolBase::FSegment& Segment) const override
+	{
+		if (!CouldItIntersect(Segment.Boundary))
+		{
+			return false;
+		}
+
+		return DoIntersectInside(Segment2D, Segment.Segment2D);
+	}
+
 };
 
 }
@@ -238,11 +253,13 @@ protected:
 
 	TArray<SegmentType> Segments;
 	bool bSegmentsAreSorted;
+	const bool bAllowOverlapping;
 
 public:
-	TIntersectionSegmentTool(const FGrid& InGrid)
+	TIntersectionSegmentTool(const FGrid& InGrid, bool bInAllowOverlapping = false)
 		: Grid(InGrid)
 		, bSegmentsAreSorted(false)
+		, bAllowOverlapping(bInAllowOverlapping)
 	{
 	}
 
@@ -287,12 +304,12 @@ public:
 	}
 
 	template<typename ExtremityType1, typename ExtremityType2>
-	const FIsoSegment* FindIntersectingSegment(const ExtremityType1* StartExtremity, const ExtremityType2* EndExtremity) const
+	const SegmentType* FindIntersectingSegment(const ExtremityType1* StartExtremity, const ExtremityType2* EndExtremity) const
 	{
 		using namespace IntersectionSegmentTool;
-		FSegment InSegment(Grid, *StartExtremity, *EndExtremity);
+		const SegmentType InSegment(Grid, *StartExtremity, *EndExtremity);
 
-		for (const FSegment& Segment : Segments)
+		for (const SegmentType& Segment : Segments)
 		{
 			if (!Segment.IsValid())
 			{
@@ -312,13 +329,18 @@ public:
 				}
 			}
 
+			if (bAllowOverlapping && Segment.IsParallelWith(InSegment))
+			{
+				continue;
+			}
+
 			switch (Segment.DoesItStartFromAndSuperimposed(StartExtremity, EndExtremity, InSegment.Segment2D))
 			{
 			case EConnectionType::SameSegment:
 			case EConnectionType::StartFrom:
 				continue;
 			case EConnectionType::SuperimposedByOrOn:
-				return Segment.IsoSegment;
+				return &Segment;
 			case EConnectionType::DoesntStartFrom:
 			default:
 				break;
@@ -326,7 +348,7 @@ public:
 
 			if (Segment.DoesItIntersect(InSegment))
 			{
-				return Segment.IsoSegment;
+				return &Segment;
 			}
 		}
 
@@ -363,6 +385,11 @@ public:
 				}
 			}
 
+			if (bAllowOverlapping && Segment.IsParallelWith(InSegment))
+			{
+				continue;
+			}
+
 			switch (Segment.DoesItStartFromAndSuperimposed(StartExtremity, EndExtremity, InSegment.Segment2D))
 			{
 			case EConnectionType::StartFrom:
@@ -376,6 +403,7 @@ public:
 					OutIntersectedSegments->Add(Segment.GetIsoSegment());
 				}
 				continue;
+
 			case EConnectionType::DoesntStartFrom:
 			default:
 				break;
@@ -536,7 +564,12 @@ public:
 	 */
 	const FIsoSegment* FindIntersectingSegment(const FIsoNode& StartNode, const FIsoNode& EndNode) const
 	{
-		return TIntersectionSegmentTool<IntersectionSegmentTool::FSegment>::FindIntersectingSegment(&StartNode, &EndNode);
+		const IntersectionSegmentTool::FSegment* IntersectingSegment = TIntersectionSegmentTool<IntersectionSegmentTool::FSegment>::FindIntersectingSegment(&StartNode, &EndNode);
+		if (IntersectingSegment)
+		{
+			return IntersectingSegment->IsoSegment;
+		}
+		return nullptr;
 	}
 
 	bool FindIntersectingSegments(const FIsoNode& StartNode, const FIsoNode& EndNode, TArray<const FIsoSegment*>& OutIntersections) const
@@ -550,8 +583,8 @@ class FIntersectionNodePairTool : public TIntersectionSegmentTool<IntersectionNo
 {
 
 public:
-	FIntersectionNodePairTool(const FGrid& InGrid)
-		: TIntersectionSegmentTool<IntersectionNodePairTool::FSegment>(InGrid)
+	FIntersectionNodePairTool(const FGrid& InGrid, bool bInAllowOverlaping)
+		: TIntersectionSegmentTool<IntersectionNodePairTool::FSegment>(InGrid, bInAllowOverlaping)
 	{
 	}
 
@@ -572,6 +605,10 @@ public:
 		return TIntersectionSegmentTool<IntersectionNodePairTool::FSegment>::FindIntersectingSegments(&StartNode, &EndNode, nullptr);
 	}
 
+	bool DoesIntersect(const FIsoNode& StartNode, const FIsoNode& EndNode) const
+	{
+		return TIntersectionSegmentTool<IntersectionNodePairTool::FSegment>::FindIntersectingSegment(&StartNode, &EndNode) != nullptr;
+	}
 };
 
 } // namespace UE::CADKernel
diff --git a/Engine/Source/Runtime/Datasmith/CADKernel/Private/CADKernel/Mesh/Meshers/IsoTriangulator/IsoCell.h b/Engine/Source/Runtime/Datasmith/CADKernel/Private/CADKernel/Mesh/Meshers/IsoTriangulator/IsoCell.h
--- a/Engine/Source/Runtime/Datasmith/CADKernel/Private/CADKernel/Mesh/Meshers/IsoTriangulator/IsoCell.h
+++ b/Engine/Source/Runtime/Datasmith/CADKernel/Private/CADKernel/Mesh/Meshers/IsoTriangulator/IsoCell.h
@@ -13,117 +13,308 @@
 namespace UE::CADKernel
 {
 class FGrid;
+class FIsoTriangulator;
+
+struct FLoopCell
+{
+	int32 Id;
+	FPoint2D Barycenter;
+	TArray<FLoopNode*> Nodes;
+	TArray<FLoopConnexion*> Connexions;
+	bool bIsOuterLoop = false;
+
+	FLoopCell(const int32 InIndex, TArray<FLoopNode*>& InNodes, const FGrid& Grid)
+		: Id(InIndex)
+		, Barycenter(FPoint::ZeroPoint)
+		, Nodes(InNodes)
+		, bIsOuterLoop(Nodes[0]->GetLoopIndex() == 0)
+	{
+		for (const FLoopNode* Node : Nodes)
+		{
+			Barycenter += Node->Get2DPoint(EGridSpace::UniformScaled, Grid);
+		}
+		Barycenter /= (double)Nodes.Num();
+	}
+};
+
+struct FLoopConnexion
+{
+	bool bIsConnexionWithBorder = false;
+	FLoopCell& Loop1;
+	FLoopCell& Loop2;
+
+	double MinDistance = HUGE_VALUE_SQUARE;
+
+	FLoopNode* NodeA = nullptr;
+	FLoopNode* NodeB = nullptr;
+
+	FIsoSegment* Segment = nullptr;
+
+	FLoopConnexion(FLoopCell& InLoop1, FLoopCell& InLoop2, bool bIsConnectingBoder = false)
+		: bIsConnexionWithBorder(bIsConnectingBoder)
+		, Loop1(InLoop1)
+		, Loop2(InLoop2)
+	{
+	}
+
+	void LinkToLoop()
+	{
+		Loop1.Connexions.Add(this);
+		Loop2.Connexions.Add(this);
+	}
+
+	const FLoopCell* GetOtherLoop(const FLoopCell* Loop)
+	{
+		return (Loop == &Loop1) ? &Loop2 : &Loop1;
+	}
+
+	bool IsShortestPath()
+	{
+		TMap<const FLoopCell*, double> DistanceToLoops;
+		DistanceToLoops.Add(&Loop1, 0.);
+		double MinPathDistance = HUGE_VALUE;
+		while (true)
+		{
+			double DistanceToCurrent = HUGE_VALUE;
+			TPair<const FLoopCell*, double>* CurrentLoop = nullptr;
+			for (TPair<const FLoopCell*, double>& DistanceToLoop : DistanceToLoops)
+			{
+				if (DistanceToLoop.Value >= 0. && DistanceToLoop.Value < DistanceToCurrent)
+				{
+					DistanceToCurrent = DistanceToLoop.Value;
+					CurrentLoop = &DistanceToLoop;
+				}
+			}
+
+			if (!CurrentLoop || CurrentLoop->Value > MinDistance)
+			{
+				return true;
+			}
+
+			if (CurrentLoop->Key == &Loop2)
+			{
+				return (CurrentLoop->Value > MinDistance);
+			}
+
+			for (FLoopConnexion* Connexion : CurrentLoop->Key->Connexions)
+			{
+				if (Connexion == this || Connexion->bIsConnexionWithBorder)
+				{
+					continue;
+				}
+
+				const FLoopCell* NextLoop = Connexion->GetOtherLoop(CurrentLoop->Key);
+				const double DistanceToNextByCurrent = CurrentLoop->Value + Connexion->MinDistance;
+
+				double* DistanceToNextLoop = DistanceToLoops.Find(NextLoop);
+				if (DistanceToNextLoop)
+				{
+					if (*DistanceToNextLoop > DistanceToNextByCurrent)
+					{
+						*DistanceToNextLoop = DistanceToNextByCurrent;
+					}
+				}
+				else
+				{
+					DistanceToLoops.Add(NextLoop, DistanceToNextByCurrent);
+				}
+			}
+
+			CurrentLoop->Value = -HUGE_VALUE;
+		}
+
+		return true;
+	}
+};
+
+struct FCellPath
+{
+	double Length = 0;
+	FLoopCell* CurrentLoop;
+	TArray<FLoopCell> Path;
+};
 
 struct FCell
 {
+	FIsoTriangulator& Triangulator;
 	const FGrid& Grid;
 	int32 Id;
-	int32 LoopNodeCount;
-	bool bHasOuterLoop = false;
-	int32 InnerLoopCount = 0;
-
-	TArray<int32> LoopIndexToIndex;
-	TArray<TArray<FLoopNode*>> SubLoops;
 
-	// the outer loop is subdivide into connected (in the cell) nodes
-	TArray<TArray<FLoopNode*>> OuterLoopSubdivision;
+	int32 InnerLoopCount = 0;
+	int32 OuterLoopCount = 0;
 
 	TArray<FIsoSegment*> CandidateSegments;
 	TArray<FIsoSegment*> FinalSegments;
+
 	FIntersectionSegmentTool IntersectionTool;
 
-	/** Loop that barycenter is on the border of the barycenter mesh */
-	TSet<int32> BorderLoopIndices;
+	TArray<FLoopCell> LoopCells;
+	TArray<FLoopConnexion> LoopConnexions;
+	TArray<int32> LoopCellBorderIndices;
 
-	FCell(const int32 InLoopIndex, const TArray<FLoopNode*>& InNodes, const FGrid& InGrid)
-		: Grid(InGrid)
+	FCell(const int32 InLoopIndex, TArray<FLoopNode*>& InNodes, FIsoTriangulator& InTriangulator)
+		: Triangulator(InTriangulator)
+		, Grid(Triangulator.GetGrid())
 		, Id(InLoopIndex)
-		, LoopNodeCount(InNodes.Num())
-		, IntersectionTool(Grid)
+		, IntersectionTool(Triangulator.GetGrid())
 	{
-		ensureCADKernel(InNodes.Num() > 0);
+		const int32 NodeCount = InNodes.Num();
+		ensureCADKernel(NodeCount > 0);
+
+		// Subdivide InNodes in SubLoop
+		Algo::Sort(InNodes, [&](const FLoopNode* LoopNode1, const FLoopNode* LoopNode2)
+			{
+				return LoopNode1->GetGlobalIndex() < LoopNode2->GetGlobalIndex();
+			});
 
-		int32 FaceLoopCount = Grid.GetFace().LoopCount();
 		int32 LoopCount = 0;
+		FLoopNode* PreviousNode = nullptr;
+		for (FLoopNode* Node : InNodes)
 		{
-			TArray<int32> LoopVertexCount;
-			LoopVertexCount.Init(0, FaceLoopCount);
-			for (const FLoopNode* Node : InNodes)
+			if (&Node->GetPreviousNode() != PreviousNode)
 			{
-				LoopVertexCount[Node->GetLoopIndex()]++;
+				LoopCount++;
 			}
+			PreviousNode = Node;
+		}
+
+		LoopCells.Reserve(LoopCount);
 
-			for (const int32& VertexCount : LoopVertexCount)
+		LoopCount = 0;
+		TArray<FLoopNode*> LoopNodes;
+		LoopNodes.Reserve(NodeCount);
+
+		int32 LoopIndex = -1;
+		FLoopCell* FirstLoopCell = nullptr;
+
+		TFunction<void(TArray<FLoopNode*>&)> MakeLoopCell = [&FirstLoopCell, &LoopCells = LoopCells, &LoopIndex, &Grid = Grid](TArray<FLoopNode*>& LoopNodes)
+		{
+			if (LoopNodes.Num())
 			{
-				if (VertexCount > 0)
+				if ((LoopIndex == LoopNodes[0]->GetLoopIndex()) && (&LoopNodes.Last()->GetNextNode() == FirstLoopCell->Nodes[0]))
 				{
-					LoopCount++;
+					LoopNodes.Append(FirstLoopCell->Nodes);
+					FirstLoopCell->Nodes = LoopNodes;
+				}
+				else
+				{
+					LoopCells.Emplace(LoopCells.Num(), LoopNodes, Grid);
 				}
-			}
 
-			SubLoops.SetNum(LoopCount);
-			int32 Index = 0;
-			for (int32 LoopIndex = 0; LoopIndex < FaceLoopCount; ++LoopIndex)
-			{
-				if (LoopVertexCount[LoopIndex] > 0)
+				if (LoopIndex != LoopNodes[0]->GetLoopIndex())
 				{
-					SubLoops[Index].Reserve(LoopVertexCount[LoopIndex]);
-					LoopVertexCount[LoopIndex] = Index;
-					Index++;
+					LoopIndex = LoopNodes[0]->GetLoopIndex();
+					FirstLoopCell = &LoopCells.Last();
 				}
 			}
-			Swap(LoopIndexToIndex, LoopVertexCount);
-		}
+		};
 
+		PreviousNode = nullptr;
 		for (FLoopNode* Node : InNodes)
 		{
-			int32 Index = LoopIndexToIndex[Node->GetLoopIndex()];
-			SubLoops[Index].Add(Node);
-		}
+			if (Node->IsDelete())
+			{
+				continue;
+			}
 
-		bHasOuterLoop = (SubLoops[0][0]->GetLoopIndex() == 0);
-		InnerLoopCount = bHasOuterLoop ? (SubLoops.Num() - 1) : SubLoops.Num();
+			if (&Node->GetPreviousNode() != PreviousNode)
+			{
+				MakeLoopCell(LoopNodes);
+				LoopNodes.Reset(NodeCount);
+			}
+			LoopNodes.Add(Node);
+			PreviousNode = Node;
+		}
+		MakeLoopCell(LoopNodes);
 
-		if (bHasOuterLoop)
+		for(const FLoopCell& LoopCell : LoopCells)
 		{
-			// Subdivide Loop0 in SubLoop
-			TArray<FLoopNode*>& Loop0 = SubLoops[0];
-			Algo::Sort(Loop0, [&](const FLoopNode* LoopNode1, const FLoopNode* LoopNode2)
-				{
-					return LoopNode1->GetIndex() < LoopNode2->GetIndex();
-				});
+			if (LoopCell.bIsOuterLoop)
+			{
+				OuterLoopCount++;
+			}
+			else
+			{
+				InnerLoopCount++;
+			}
+		}
+	}
+
+	void InitLoopConnexions();
 
-			OuterLoopSubdivision.Reserve(Loop0.Num());
-			FLoopNode* PreviousNode = nullptr;
-			TArray<FLoopNode*>* SubLoop = nullptr;
-			for (FLoopNode* Node : Loop0)
+	TArray<TPair<int32, FPoint2D>> GetLoopBarycenters()
+	{
+		TArray<TPair<int32, FPoint2D>> LoopBarycenters;
+		LoopBarycenters.Reserve(LoopCells.Num());
+		for (const FLoopCell& LoopCell : LoopCells)
+		{
+			if(!LoopCell.bIsOuterLoop)
 			{
-				if (&Node->GetPreviousNode() != PreviousNode)
-				{
-					SubLoop = &OuterLoopSubdivision.Emplace_GetRef();
-				}
-				SubLoop->Add(Node);
-				PreviousNode = Node;
+				LoopBarycenters.Emplace(LoopCell.Id, LoopCell.Barycenter);
 			}
 		}
+		return MoveTemp(LoopBarycenters);
 	}
 
-	void SelectSegmentInCandidateSegments(TFactory<FIsoSegment>& SegmentFactory)
+	void ConnectLoopsByNeighborhood();
+
+	/**
+	 *  SubLoopA                  SubLoopB
+	 *      --X---X             X-----X--
+	 *             \           /
+	 *              \         /
+	 *               X=======X
+	 *              /         \
+	 *             /           \
+	 *      --X---X             X-----X--
+	 *
+	 *     ======= ShortestSegment
+	 */
+	void TryToConnectTwoSubLoopsWithShortestSegment(FLoopConnexion& LoopConnexion);
+
+	void TryToCreateSegment(FLoopConnexion& LoopConnexion);
+
+	void SelectSegmentInCandidateSegments(TFactory<FIsoSegment>& SegmentFactory, bool bSelectWithCellSizeCriteria)
 	{
 #ifdef DEBUG_SELECT_SEGMENT
 		F3DDebugSession _(Grid.bDisplay, TEXT("SelectSegmentInCandidateSegments "));
-		IntersectionTool.Display(TEXT("Cell.IntersectionTool at SelectSegmentInCandidateSegments start"));
-		//Wait();
+		if(Grid.bDisplay)
+		{
+			IntersectionTool.Display(Grid.bDisplay, TEXT("Cell.IntersectionTool at SelectSegmentInCandidateSegments start"));
+			Wait();
+		}
 #endif
 
-		Algo::Sort(CandidateSegments, [&](const FIsoSegment* Segment1, const FIsoSegment* Segment2)
-			{
-				return Segment1->Get2DLengthSquare(EGridSpace::UniformScaled, Grid) < Segment2->Get2DLengthSquare(EGridSpace::UniformScaled, Grid);
-			});
+		double CellSquareLength = HUGE_VALUE;
+		if(bSelectWithCellSizeCriteria)
+		{
+			int32 IndexU;
+			int32 IndexV;
+			Grid.UVIndexFromGlobalIndex(Id, IndexU, IndexV);
 
-		// Validate all candidate segments
+			CellSquareLength = Grid.GetInner2DPoint(EGridSpace::UniformScaled, IndexU, IndexV).SquareDistance(Grid.GetInner2DPoint(EGridSpace::UniformScaled, IndexU, IndexV + 1));
+			const double CellLength2 = Grid.GetInner2DPoint(EGridSpace::UniformScaled, IndexU, IndexV).SquareDistance(Grid.GetInner2DPoint(EGridSpace::UniformScaled, IndexU + 1, IndexV));
+			CellSquareLength = FMath::Min(CellSquareLength, CellLength2);
+		}
+
+		TArray<TPair<double, FIsoSegment*>> LengthOfCandidateSegments;
+		LengthOfCandidateSegments.Reserve(CandidateSegments.Num());
 		for (FIsoSegment* Segment : CandidateSegments)
 		{
+			LengthOfCandidateSegments.Emplace(Segment->Get2DLengthSquare(EGridSpace::UniformScaled, Grid), Segment);
+		}
+
+		Algo::Sort(LengthOfCandidateSegments, [&](const TPair<double, FIsoSegment*>& P1, const TPair < double, FIsoSegment*>& P2) { return P1.Key < P2.Key; });
+
+		// Validate all candidate segments
+		for (const TPair<double, FIsoSegment*>& Candidate : LengthOfCandidateSegments)
+		{
+			if (Candidate.Key > CellSquareLength)
+			{
+				break;
+			}
+
+			FIsoSegment* Segment = Candidate.Value;
 #ifdef DEBUG_SELECT_SEGMENT
 			F3DDebugSession _(Grid.bDisplay, TEXT("Segment"));
 #endif
@@ -158,13 +349,6 @@ struct FCell
 		CandidateSegments.Empty();
 	}
 
-	bool Contains(FLoopNode* NodeToFind)
-	{
-		int32 LoopIndex = NodeToFind->GetLoopIndex();
-		int32 Index = LoopIndexToIndex[LoopIndex];
-		return SubLoops[Index].Find(NodeToFind) != INDEX_NONE;
-	}
-
 };
 
 } // namespace UE::CADKernel
diff --git a/Engine/Source/Runtime/Datasmith/CADKernel/Private/CADKernel/Mesh/Meshers/IsoTriangulator/IsoNode.h b/Engine/Source/Runtime/Datasmith/CADKernel/Private/CADKernel/Mesh/Meshers/IsoTriangulator/IsoNode.h
--- a/Engine/Source/Runtime/Datasmith/CADKernel/Private/CADKernel/Mesh/Meshers/IsoTriangulator/IsoNode.h
+++ b/Engine/Source/Runtime/Datasmith/CADKernel/Private/CADKernel/Mesh/Meshers/IsoTriangulator/IsoNode.h
@@ -63,16 +63,16 @@ protected:
 	TArray<FIsoSegment*> ConnectedSegments;
 	EIsoNodeStates States;
 
-	int32 Index; // Index of the node either in loop nodes either in inner nodes
-	int32 FaceIndex; // Index of the node in the face
+	int32 LocalIndex; // Index of the node either in loop nodes either in inner nodes
+	int32 GlobalIndex; // Index of the node in the face
 	int32 NodeId; // Id of the node in the mesh
 
 public:
 	FIsoNode(int32 InNodeIndex, int32 InFaceIndex, int32 InNodeId)
 		: ConnectedSegments()
 		, States(EIsoNodeStates::None)
-		, Index(InNodeIndex)
-		, FaceIndex(InFaceIndex)
+		, LocalIndex(InNodeIndex)
+		, GlobalIndex(InFaceIndex)
 		, NodeId(InNodeId)
 	{
 		ConnectedSegments.Reserve(5);
@@ -83,7 +83,7 @@ public:
 	virtual void Delete()
 	{
 		ConnectedSegments.Empty();
-		FaceIndex = -1;
+		GlobalIndex = -1;
 
 		States = EIsoNodeStates::Delete;
 	}
@@ -95,12 +95,12 @@ public:
 
 	const int32 GetIndex() const
 	{
-		return Index;
+		return LocalIndex;
 	}
 
-	const int32 GetFaceIndex() const
+	const int32 GetGlobalIndex() const
 	{
-		return FaceIndex;
+		return GlobalIndex;
 	}
 
 	const int32 GetNodeId() const
@@ -254,8 +254,8 @@ protected:
 	FLoopNode* ConnectedLoopNodes[2];
 
 public:
-	FLoopNode(int32 InLoopIndex, int32 InNodeIndex, int32 InFaceIndex, int32 NodeId)
-		: FIsoNode(InNodeIndex, InFaceIndex, NodeId)
+	FLoopNode(int32 InLoopIndex, int32 InNodeLoopIndex, int32 InNodeFaceIndex, int32 NodeId)
+		: FIsoNode(InNodeLoopIndex, InNodeFaceIndex, NodeId)
 		, LoopIndex(InLoopIndex)
 		, ConnectedLoopNodes{ nullptr, nullptr }
 	{
@@ -329,22 +329,22 @@ public:
 	 */
 	virtual const FPoint2D& Get2DPoint(EGridSpace Space, const FGrid& Grid) const override
 	{
-		return Grid.GetLoop2DPoint(Space, LoopIndex, Index);
+		return Grid.GetLoop2DPoint(Space, LoopIndex, LocalIndex);
 	}
 
 	virtual void Set2DPoint(EGridSpace Space, FGrid& Grid, const FPoint2D& NewCoordinate) override
 	{
-		Grid.SetLoop2DPoint(Space, LoopIndex, Index, NewCoordinate);
+		Grid.SetLoop2DPoint(Space, LoopIndex, LocalIndex, NewCoordinate);
 	}
 
 	virtual const FPoint& Get3DPoint(const FGrid& Grid) const override
 	{
-		return Grid.GetLoops3D()[LoopIndex][Index];
+		return Grid.GetLoops3D()[LoopIndex][LocalIndex];
 	}
 
 	virtual const FVector3f& GetNormal(const FGrid& Grid) const override
 	{
-		return Grid.GetLoopNormals()[LoopIndex][Index];
+		return Grid.GetLoopNormals()[LoopIndex][LocalIndex];
 	}
 
 	/**
@@ -357,7 +357,7 @@ public:
 		case EGridSpace::Default2D:
 		case EGridSpace::Scaled:
 		case EGridSpace::UniformScaled:
-			return Grid.GetLoop2DPoint(Space, LoopIndex, Index);
+			return Grid.GetLoop2DPoint(Space, LoopIndex, LocalIndex);
 		}
 		return FPoint::ZeroPoint;
 	}
@@ -502,22 +502,22 @@ public:
 
 	virtual const FPoint2D& Get2DPoint(EGridSpace Space, const FGrid& Grid) const override
 	{
-		return Grid.GetInner2DPoint(Space, Index);
+		return Grid.GetInner2DPoint(Space, LocalIndex);
 	}
 
 	virtual void Set2DPoint(EGridSpace Space, FGrid& Grid, const FPoint2D& NewCoordinate) override
 	{
-		return Grid.SetInner2DPoint(Space, Index, NewCoordinate);
+		return Grid.SetInner2DPoint(Space, LocalIndex, NewCoordinate);
 	}
 
 	virtual const FPoint& Get3DPoint(const FGrid& Grid) const override
 	{
-		return Grid.GetInner3DPoint(Index);
+		return Grid.GetInner3DPoint(LocalIndex);
 	}
 
 	virtual const FVector3f& GetNormal(const FGrid& Grid) const override
 	{
-		return Grid.GetPointNormal(Index);
+		return Grid.GetPointNormal(LocalIndex);
 	}
 
 	/**
@@ -530,7 +530,7 @@ public:
 		case EGridSpace::Default2D:
 		case EGridSpace::Scaled:
 		case EGridSpace::UniformScaled:
-			return Grid.GetInner2DPoint(Space, Index);
+			return Grid.GetInner2DPoint(Space, LocalIndex);
 		}
 		return FPoint::ZeroPoint;
 	}
diff --git a/Engine/Source/Runtime/Datasmith/CADKernel/Private/CADKernel/Mesh/Structure/GridBase.h b/Engine/Source/Runtime/Datasmith/CADKernel/Private/CADKernel/Mesh/Structure/GridBase.h
--- a/Engine/Source/Runtime/Datasmith/CADKernel/Private/CADKernel/Mesh/Structure/GridBase.h
+++ b/Engine/Source/Runtime/Datasmith/CADKernel/Private/CADKernel/Mesh/Structure/GridBase.h
@@ -245,17 +245,6 @@ public:
 		return IndexV * CuttingCount[EIso::IsoU] + IndexU;
 	}
 
-#ifdef WIP
-	/**
-	 * @return the Index of the position in the arrays of a point [IndexU, IndexV] of the grid
-	 */
-	void UVIndexFromGlobalIndex(int32 GLobalIndex, int32& OutIndexU, int32& OutIndexV) const
-	{
-		OutIndexU = GLobalIndex % CuttingCount[EIso::IsoU];
-		OutIndexV = GLobalIndex / CuttingCount[EIso::IsoU];
-	}
-#endif
-
 #ifdef CADKERNEL_DEBUG
 
 	// ======================================================================================================================================================================================================================
diff --git a/Engine/Source/Runtime/Datasmith/CADKernel/Private/CADKernel/Mesh/Structure/ModelMesh.cpp b/Engine/Source/Runtime/Datasmith/CADKernel/Private/CADKernel/Mesh/Structure/ModelMesh.cpp
--- a/Engine/Source/Runtime/Datasmith/CADKernel/Private/CADKernel/Mesh/Structure/ModelMesh.cpp
+++ b/Engine/Source/Runtime/Datasmith/CADKernel/Private/CADKernel/Mesh/Structure/ModelMesh.cpp
@@ -57,50 +57,6 @@ const FVertexMesh* FModelMesh::GetMeshOfVertexNodeId(const int32 Ident) const
 	return nullptr;
 }
 
-#ifdef OLD
-const int32 FModelMesh::GetIndexOfVertexFromId(const int32 Ident) const
-{
-	for (const FVertexMesh* VertexMesh : VertexMeshes)
-	{
-		if (Ident == VertexMesh->GetStartVertexId())
-		{
-			return VertexMesh->GetIndexInMeshModel();
-		}
-	}
-	return -1;
-}
-
-const int32 FModelMesh::GetIndexOfEdgeFromId(const int32 Ident) const
-{
-	for (const FEdgeMesh* EdgeMesh : EdgeMeshes)
-	{
-		if (Ident >= EdgeMesh->GetStartVertexId())
-		{
-			if (Ident <= EdgeMesh->GetLastVertexIndex())
-			{
-				return EdgeMesh->GetIndexInMeshModel();
-			}
-		}
-	}
-	return -1;
-}
-
-const int32 FModelMesh::GetIndexOfSurfaceFromId(const int32 Ident) const
-{
-	for (const TSharedPtr<FFaceMesh>& FaceMesh : FaceMeshes)
-	{
-		if (Ident >= FaceMesh->GetStartVertexId())
-		{
-			if (Ident <= FaceMesh->GetLastVertexIndex())
-			{
-				return FaceMesh->GetIndexInMeshModel();
-			}
-		}
-	}
-	return -1;
-}
-#endif
-
 void FModelMesh::GetNodeCoordinates(TArray<FPoint>& NodeCoordinates) const
 {
 	NodeCoordinates.Reserve(LastIdUsed + 1);
diff --git a/Engine/Source/Runtime/Datasmith/CADKernel/Private/CADKernel/Mesh/Structure/ModelMeshAnalyzer.cpp b/Engine/Source/Runtime/Datasmith/CADKernel/Private/CADKernel/Mesh/Structure/ModelMeshAnalyzer.cpp
--- a/Engine/Source/Runtime/Datasmith/CADKernel/Private/CADKernel/Mesh/Structure/ModelMeshAnalyzer.cpp
+++ b/Engine/Source/Runtime/Datasmith/CADKernel/Private/CADKernel/Mesh/Structure/ModelMeshAnalyzer.cpp
@@ -10,6 +10,7 @@
 #include "CADKernel/Mesh/Structure/FaceMesh.h"
 #include "CADKernel/Mesh/Structure/VertexMesh.h"
 #include "CADKernel/Topo/TopologicalEntity.h"
+#include "CADKernel/UI/Display.h"
 
 #include "Templates/UnrealTemplate.h"
 
@@ -210,6 +211,55 @@ bool FModelMeshAnalyzer::CheckOrientation() const
 	return true;
 }
 
+void FModelMeshAnalyzer::Display()
+{
+	F3DDebugSession A(FString::Printf(TEXT("ModelMesh %d"), ModelMesh.GetId()));
+	{
+		F3DDebugSegment B(0);
+		for (const Analyzer::FTriangle& Triangle : Triangles)
+		{
+			DisplayTriangle(Triangle);
+		}
+	}
+
+	{
+		F3DDebugSegment B(ModelMesh.GetId());
+		for (const Analyzer::FEdge& Edge : Edges)
+		{
+			switch (Edge.TriangleCount)
+			{
+			case 1:
+				DisplayEdge(Edge, EVisuProperty::BorderEdge);
+				break;
+
+			case 2:
+				DisplayEdge(Edge, EVisuProperty::EdgeMesh);
+				break;
+
+			default:
+				DisplayEdge(Edge, EVisuProperty::NonManifoldEdge);
+				break;
+			}
+		}
+	}
+
+}
+
+void FModelMeshAnalyzer::DisplayTriangle(const Analyzer::FTriangle& Triangle)
+{
+	TArray<FPoint> Vertices;
+	Vertices.Reserve(3);
+	Vertices.Add(NodeCoordinates[Triangle.Vertices[0]]);
+	Vertices.Add(NodeCoordinates[Triangle.Vertices[1]]);
+	Vertices.Add(NodeCoordinates[Triangle.Vertices[2]]);
+	DrawElement(2, Vertices);
+}
+
+void FModelMeshAnalyzer::DisplayEdge(const Analyzer::FEdge& Edge, EVisuProperty Property)
+{
+	DrawSegment( NodeCoordinates[Edge.VertexIndices[0]], NodeCoordinates[Edge.VertexIndices[1]], Property);
+}
+
 FPoint Analyzer::FTriangle::ComputeNormal(const TArray<FPoint>& NodeCoordinates) const
 {
 	return UE::CADKernel::FTriangle(NodeCoordinates[Vertices[0]], NodeCoordinates[Vertices[1]], NodeCoordinates[Vertices[2]]).ComputeNormal();
diff --git a/Engine/Source/Runtime/Datasmith/CADKernel/Private/CADKernel/Mesh/Structure/ModelMeshAnalyzer.h b/Engine/Source/Runtime/Datasmith/CADKernel/Private/CADKernel/Mesh/Structure/ModelMeshAnalyzer.h
--- a/Engine/Source/Runtime/Datasmith/CADKernel/Private/CADKernel/Mesh/Structure/ModelMeshAnalyzer.h
+++ b/Engine/Source/Runtime/Datasmith/CADKernel/Private/CADKernel/Mesh/Structure/ModelMeshAnalyzer.h
@@ -3,6 +3,7 @@
 
 #include "CADKernel/Core/HaveStates.h"
 #include "CADKernel/Core/Types.h"
+#include "CADKernel/UI/Visu.h"
 
 #include "CADKernel/Math/Point.h"
 
@@ -97,6 +98,12 @@ public:
 	double ComputeMaxAngle() const;
 	bool CheckOrientation() const;
 
+	void Display();
+
+private:
+	void DisplayTriangle(const Analyzer::FTriangle& Triangle);
+	void DisplayEdge(const Analyzer::FEdge& Edge, EVisuProperty Property);
+
 };
 
 }
diff --git a/Engine/Source/Runtime/Datasmith/CADKernel/Private/CADKernel/Topo/FaceAnalyzer.cpp b/Engine/Source/Runtime/Datasmith/CADKernel/Private/CADKernel/Topo/FaceAnalyzer.cpp
--- a/Engine/Source/Runtime/Datasmith/CADKernel/Private/CADKernel/Topo/FaceAnalyzer.cpp
+++ b/Engine/Source/Runtime/Datasmith/CADKernel/Private/CADKernel/Topo/FaceAnalyzer.cpp
@@ -259,10 +259,10 @@ void FFaceAnalyzer::DisplayLoopSegments(Topo::FThinFaceContext& Context)
 	//Wait();
 #endif
 }
-
+sadf
+#ifdef CADKERNEL_DEV
 void FFaceAnalyzer::DisplayClosedSegments(Topo::FThinFaceContext& Context)
 {
-#ifdef DEBUG_THIN_FACE
 	F3DDebugSession _(TEXT("Closed Segment"));
 	for (const Topo::FEdgeSegment* Segment : Context.LoopSegments)
 	{
@@ -280,8 +280,8 @@ void FFaceAnalyzer::DisplayClosedSegments(Topo::FThinFaceContext& Context)
 		}
 	}
 	//Wait();
-#endif
 }
 #endif
+#endif
 }
 
diff --git a/Engine/Source/Runtime/Datasmith/CADKernel/Private/CADKernel/Topo/Shell.cpp b/Engine/Source/Runtime/Datasmith/CADKernel/Private/CADKernel/Topo/Shell.cpp
--- a/Engine/Source/Runtime/Datasmith/CADKernel/Private/CADKernel/Topo/Shell.cpp
+++ b/Engine/Source/Runtime/Datasmith/CADKernel/Private/CADKernel/Topo/Shell.cpp
@@ -133,6 +133,10 @@ void FShell::SpreadBodyOrientation()
 		{
 			Face.Entity->SetBackOriented();
 		}
+		else
+		{
+			Face.Entity->ResetBackOriented();
+		}
 	}
 }
 
@@ -346,8 +350,8 @@ int32 FShell::Orient()
 					continue;
 				}
 
-				bool EdgesAreWellOriented = (Edge->IsSameDirection(*NeighboringEdge) == (OrientedNeighboringEdge->Direction != OrientedEdge.Direction));
-				bool FacesHaveSameOrientation = Face.IsBackOriented() == NeighboringFace->IsBackOriented();
+				const bool EdgesAreWellOriented = (Edge->IsSameDirection(*NeighboringEdge) == (OrientedNeighboringEdge->Direction != OrientedEdge.Direction));
+				const bool FacesHaveSameOrientation = Face.IsBackOriented() == NeighboringFace->IsBackOriented();
 
 				if (EdgesAreWellOriented != FacesHaveSameOrientation)
 				{
diff --git a/Engine/Source/Runtime/Datasmith/CADKernel/Private/CADKernel/Topo/TopologicalEdge.cpp b/Engine/Source/Runtime/Datasmith/CADKernel/Private/CADKernel/Topo/TopologicalEdge.cpp
--- a/Engine/Source/Runtime/Datasmith/CADKernel/Private/CADKernel/Topo/TopologicalEdge.cpp
+++ b/Engine/Source/Runtime/Datasmith/CADKernel/Private/CADKernel/Topo/TopologicalEdge.cpp
@@ -722,7 +722,7 @@ void FTopologicalEdge::ProjectTwinEdgePointsOn2DCurve(const TSharedRef<FTopologi
 
 		bool bSameDirection = IsSameDirection(*InTwinEdge);
 		const double ToleranceOfProjection = Length3D * 0.1;
-		Curve->ProjectTwinCurvePoints(Points3D, bSameDirection, Coordinates, ToleranceOfProjection);
+		Curve->ProjectTwinCurvePoints(Boundary, Points3D, bSameDirection, Coordinates, ToleranceOfProjection);
 		Curve->Approximate2DPoints(Coordinates, OutPoints2D);
 	}
 }
@@ -1362,11 +1362,11 @@ void FTopologicalEdge::SpawnIdent(FDatabase& Database)
 	}
 }
 
-TSharedPtr<FTopologicalVertex> FTopologicalEdge::SplitAt(double SplittingCoordinate, const FPoint& NewVertexCoordinate, bool bKeepStartVertexConnectivity, TSharedPtr<FTopologicalEdge>& NewEdge)
+FTopologicalVertex* FTopologicalEdge::SplitAt(double SplittingCoordinate, const FPoint& NewVertexCoordinate, bool bKeepStartVertexConnectivity, TSharedPtr<FTopologicalEdge>& NewEdge)
 {
 	if (GetTwinEntityCount() > 1)
 	{
-		return TSharedPtr<FTopologicalVertex>();
+		return nullptr;
 		// TODO
 	}
 
@@ -1384,7 +1384,7 @@ TSharedPtr<FTopologicalVertex> FTopologicalEdge::SplitAt(double SplittingCoordin
 	}
 	if (!NewEdge.IsValid())
 	{
-		return TSharedPtr<FTopologicalVertex>();
+		return nullptr;
 	}
 
 	if (bKeepStartVertexConnectivity)
@@ -1404,7 +1404,7 @@ TSharedPtr<FTopologicalVertex> FTopologicalEdge::SplitAt(double SplittingCoordin
 	ComputeLength();
 
 	Loop->SplitEdge(*this, NewEdge, bKeepStartVertexConnectivity);
-	return MiddelVertex;
+	return &MiddelVertex.Get();
 }
 
 bool FTopologicalEdge::IsSharpEdge() const
diff --git a/Engine/Source/Runtime/Datasmith/CADKernel/Private/CADKernel/Topo/Topomaker.cpp b/Engine/Source/Runtime/Datasmith/CADKernel/Private/CADKernel/Topo/Topomaker.cpp
--- a/Engine/Source/Runtime/Datasmith/CADKernel/Private/CADKernel/Topo/Topomaker.cpp
+++ b/Engine/Source/Runtime/Datasmith/CADKernel/Private/CADKernel/Topo/Topomaker.cpp
@@ -63,6 +63,15 @@ void MergeCoincidentVertices(TArray<FTopologicalVertex*>& VerticesToMerge, doubl
 	TArray<int32> SortedVertexIndices;
 	SortedVertexIndices.Reserve(VertexNum);
 
+#ifdef DEBUG_MERGE_COINCIDENT_VERTICES
+	F3DDebugSession _(TEXT("Merge Coincident Vertices"));
+	for (const FTopologicalVertex* Vertex : VerticesToMerge)
+	{
+		F3DDebugSession A(*FString::Printf(TEXT("Vertex %d"), Vertex->GetId()));
+		Display(*Vertex);
+	}
+#endif
+
 	for (FTopologicalVertex* Vertex : VerticesToMerge)
 	{
 		VerticesWeight.Add(Vertex->GetCoordinates().DiagonalAxisCoordinate());
@@ -144,7 +153,7 @@ void MergeCoincidentVertices(TArray<FTopologicalVertex*>& VerticesToMerge, doubl
 
 }
 
-TSharedPtr<FTopologicalVertex> SplitAndLink(FTopologicalVertex& StartVertex, FTopologicalEdge& EdgeToLink, FTopologicalEdge& EdgeToSplit, double SquareSewTolerance, double SquareMinEdgeLength)
+FTopologicalVertex* SplitAndLink(FTopologicalVertex& StartVertex, FTopologicalEdge& EdgeToLink, FTopologicalEdge& EdgeToSplit, double SquareSewTolerance, double SquareMinEdgeLength)
 {
 	FTopologicalVertex* VertexToLink = EdgeToLink.GetOtherVertex(StartVertex);
 	FTopologicalVertex* EndVertex = EdgeToSplit.GetOtherVertex(StartVertex);
@@ -154,7 +163,7 @@ TSharedPtr<FTopologicalVertex> SplitAndLink(FTopologicalVertex& StartVertex, FTo
 	{
 		VertexToLink->Link(*EndVertex);
 		EdgeToLink.Link(EdgeToSplit);
-		return VertexToLink->GetLinkActiveEntity();
+		return &VertexToLink->GetLinkActiveEntity().Get();
 	}
 
 	double SquareDistanceToStartVertex = StartVertex.SquareDistanceBetweenBarycenters(*VertexToLink);
@@ -162,7 +171,7 @@ TSharedPtr<FTopologicalVertex> SplitAndLink(FTopologicalVertex& StartVertex, FTo
 	{
 		VertexToLink->Link(StartVertex);
 		EdgeToLink.SetAsDegenerated();
-		return VertexToLink->GetLinkActiveEntity();
+		return &VertexToLink->GetLinkActiveEntity().Get();
 	}
 
 	FPoint ProjectedPoint;
@@ -171,7 +180,7 @@ TSharedPtr<FTopologicalVertex> SplitAndLink(FTopologicalVertex& StartVertex, FTo
 	double SquareDistanceToProjectedPoint = ProjectedPoint.SquareDistance(VertexToLink->GetBarycenter());
 	if (SquareDistanceToProjectedPoint > SquareSewTolerance)
 	{
-		return TSharedPtr<FTopologicalVertex>();
+		return nullptr;
 	}
 
 	double SquareDistanceToOtherPoint = ProjectedPoint.SquareDistance(EndVertex->GetBarycenter());
@@ -180,7 +189,7 @@ TSharedPtr<FTopologicalVertex> SplitAndLink(FTopologicalVertex& StartVertex, FTo
 		// the new point is closed to the extremity, a degenerated edge will be created, so the edges are joined
 		VertexToLink->Link(*EndVertex);
 		EdgeToLink.Link(EdgeToSplit);
-		return VertexToLink->GetLinkActiveEntity();
+		return &VertexToLink->GetLinkActiveEntity().Get();
 	}
 
 	SquareDistanceToProjectedPoint = ProjectedPoint.SquareDistance(StartVertex.GetBarycenter());
@@ -188,7 +197,7 @@ TSharedPtr<FTopologicalVertex> SplitAndLink(FTopologicalVertex& StartVertex, FTo
 	{
 		VertexToLink->Link(StartVertex);
 		EdgeToLink.SetAsDegenerated();
-		return VertexToLink->GetLinkActiveEntity();
+		return &VertexToLink->GetLinkActiveEntity().Get();
 	}
 
 	// JoinParallelEdges process all edges connected to startVertex (ConnectedEdges).
@@ -197,10 +206,10 @@ TSharedPtr<FTopologicalVertex> SplitAndLink(FTopologicalVertex& StartVertex, FTo
 	bool bKeepStartVertexConnectivity = (StartVertex.GetLink() == EdgeToSplit.GetStartVertex()->GetLink());
 
 	TSharedPtr<FTopologicalEdge> NewEdge;
-	TSharedPtr<FTopologicalVertex> NewVertex = EdgeToSplit.SplitAt(UProjectedPoint, ProjectedPoint, bKeepStartVertexConnectivity, NewEdge);
-	if (!NewVertex.IsValid())
+	FTopologicalVertex* NewVertex = EdgeToSplit.SplitAt(UProjectedPoint, ProjectedPoint, bKeepStartVertexConnectivity, NewEdge);
+	if (!NewVertex)
 	{
-		return TSharedPtr<FTopologicalVertex>();
+		return nullptr;
 	}
 
 	VertexToLink->Link(*NewVertex);
@@ -208,113 +217,163 @@ TSharedPtr<FTopologicalVertex> SplitAndLink(FTopologicalVertex& StartVertex, FTo
 	return NewVertex;
 }
 
-void StitchParallelEdges(TArray<FTopologicalVertex*>& VerticesToProcess, double SewTolerance, double MinEdgeLength)
+FTopologicalVertex* StitchParallelEdgesFrom(FTopologicalVertex* Vertex, double SewTolerance, double MinEdgeLength)
 {
 	double SquareTolerance = FMath::Square(SewTolerance);
 	double SquareMinEdgeLength = FMath::Square(MinEdgeLength);
 
-	FTimePoint StartTime = FChrono::Now();
+	TArray<FTopologicalEdge*> ConnectedEdges;
+	Vertex->GetConnectedEdges(ConnectedEdges);
+	const int32 ConnectedEdgeCount = ConnectedEdges.Num();
+	if (ConnectedEdgeCount == 1)
+	{
+		return nullptr;
+	}
 
-	for (int32 VertexI = 0; VertexI < VerticesToProcess.Num(); ++VertexI)
+#ifdef DEBUG_STITCH_PARALLEL_EDGES_FROM
 	{
-		FTopologicalVertex* Vertex = VerticesToProcess[VertexI];
+		F3DDebugSession B(*FString::Printf(TEXT("Vertex %d"), Vertex->GetId()));
+		{
+			F3DDebugSession A(*FString::Printf(TEXT("Vertex %d"), Vertex->GetId()));
+			Display(*Vertex);
+		}
+		for (const FTopologicalEdge* Edge : ConnectedEdges)
+		{
+			F3DDebugSession A(*FString::Printf(TEXT("Edge %d"), Edge->GetId()));
+			int32 TwinCount = Edge->GetTwinEntityCount();
+			EVisuProperty Property = EVisuProperty::RedCurve;
+			switch (TwinCount)
+			{
+			case 1:
+				Property = EVisuProperty::BorderEdge;
+				break;
+			case 2:
+				Property = EVisuProperty::BlueCurve;
+				break;
+			default:
+				break;
+			}
+			Display(*Edge, Property);
+		}
+		Wait();
+	}
+#endif
 
-		if (Vertex->IsDeleted() || !Vertex->IsBorderVertex())
+	for (int32 EdgeI = 0; EdgeI < ConnectedEdgeCount - 1; ++EdgeI)
+	{
+		FTopologicalEdge* Edge = ConnectedEdges[EdgeI];
+		ensureCADKernel(Edge->GetLoop() != nullptr);
+
+		if (Edge->IsDegenerated())
 		{
 			continue;
 		}
 
-		TArray<FTopologicalEdge*> ConnectedEdges;
-		Vertex->GetConnectedEdges(ConnectedEdges);
-		const int32 ConnectedEdgeCount = ConnectedEdges.Num();
-		if (ConnectedEdgeCount == 1)
+		if (!Edge->IsActiveEntity())
 		{
 			continue;
 		}
+		bool bFirstEdgeBorder = Edge->IsBorder();
 
-		for (int32 EdgeI = 0; EdgeI < ConnectedEdgeCount - 1; ++EdgeI)
+		for (int32 EdgeJ = EdgeI + 1; EdgeJ < ConnectedEdgeCount; ++EdgeJ)
 		{
-			FTopologicalEdge* Edge = ConnectedEdges[EdgeI];
-			ensureCADKernel(Edge->GetLoop() != nullptr);
+			FTopologicalEdge* SecondEdge = ConnectedEdges[EdgeJ];
+			if (SecondEdge->IsDegenerated())
+			{
+				continue;
+			}
+			if (!SecondEdge->IsActiveEntity())
+			{
+				continue;
+			}
+			bool bSecondEdgeBorder = SecondEdge->IsBorder();
 
-			if (Edge->IsDegenerated())
+			// Process only if at least one edge is Border
+			if (!bFirstEdgeBorder && !bSecondEdgeBorder)
 			{
 				continue;
 			}
 
-			if (!Edge->IsActiveEntity())
+			FPoint StartTangentEdge = Edge->GetTangentAt(*Vertex);
+			FPoint StartTangentOtherEdge = SecondEdge->GetTangentAt(*Vertex);
+
+			double CosAngle = StartTangentEdge.ComputeCosinus(StartTangentOtherEdge);
+			if (CosAngle < UE_DOUBLE_HALF_SQRT_3) // cos(30 deg)
 			{
 				continue;
 			}
-			bool bFirstEdgeBorder = Edge->IsBorder();
 
-			for (int32 EdgeJ = EdgeI + 1; EdgeJ < ConnectedEdgeCount; ++EdgeJ)
+#ifdef DEBUG_STITCH_PARALLEL_EDGES_FROM
 			{
-				FTopologicalEdge* SecondEdge = ConnectedEdges[EdgeJ];
-				if (SecondEdge->IsDegenerated())
+				F3DDebugSession B(*FString::Printf(TEXT("StitchParallelEdges")));
 				{
-					continue;
-				}
-				if (!SecondEdge->IsActiveEntity())
-				{
-					continue;
+					F3DDebugSession A(*FString::Printf(TEXT("Edge %d"), Edge->GetId()));
+					Display(*Edge);
 				}
-				bool bSecondEdgeBorder = SecondEdge->IsBorder();
-
-				// Process only if at least one edge is Border
-				if (!bFirstEdgeBorder && !bSecondEdgeBorder)
 				{
-					continue;
+					F3DDebugSession A(*FString::Printf(TEXT("SecondEdge %d"), SecondEdge->GetId()));
+					Display(*SecondEdge, EVisuProperty::RedCurve);
 				}
+			}
+#endif
 
-				FPoint StartTangentEdge = Edge->GetTangentAt(*Vertex);
-				FPoint StartTangentOtherEdge = SecondEdge->GetTangentAt(*Vertex);
+			FTopologicalVertex& EndVertex = *Edge->GetOtherVertex(*Vertex)->GetLinkActiveEntity();
+			FTopologicalVertex& SecondEdgeEndVertex = *SecondEdge->GetOtherVertex(*Vertex)->GetLinkActiveEntity();
 
-				double CosAngle = StartTangentEdge.ComputeCosinus(StartTangentOtherEdge);
-				if (CosAngle < UE_DOUBLE_HALF_SQRT_3) // cos(30 deg)
-				{
-					continue;
-				}
-
-				FTopologicalVertex& EndVertex = *Edge->GetOtherVertex(*Vertex)->GetLinkActiveEntity();
-				FTopologicalVertex& SecondEdgeEndVertex = *SecondEdge->GetOtherVertex(*Vertex)->GetLinkActiveEntity();
+			if (&EndVertex == &SecondEdgeEndVertex && Edge->IsLinkableTo(*SecondEdge, MinEdgeLength))
+			{
+				// should not happen but, just in case
+				Edge->Link(*SecondEdge);
+			}
+			else
+			{
+				double EdgeLength = Edge->Length();
+				double SecondEdgeLength = SecondEdge->Length();
 
-				if (&EndVertex == &SecondEdgeEndVertex && Edge->IsLinkableTo(*SecondEdge, MinEdgeLength))
+				FTopologicalVertex* OtherVertex = nullptr;
+				FTopologicalVertex* NewVertex = nullptr;
+				if (EdgeLength < SecondEdgeLength)
 				{
-					// should not happen but, just in case
-					Edge->Link(*SecondEdge);
+					OtherVertex = &SecondEdgeEndVertex;
+					NewVertex = SplitAndLink(*Vertex, *Edge, *SecondEdge, SquareTolerance, SquareMinEdgeLength);
 				}
 				else
 				{
-					double EdgeLength = Edge->Length();
-					double SecondEdgeLength = SecondEdge->Length();
-
-					FTopologicalVertex* OtherVertex = nullptr;
-					TSharedPtr<FTopologicalVertex> NewVertex;
-					if (EdgeLength < SecondEdgeLength)
-					{
-						OtherVertex = &SecondEdgeEndVertex;
-						NewVertex = SplitAndLink(*Vertex, *Edge, *SecondEdge, SquareTolerance, SquareMinEdgeLength);
-					}
-					else
-					{
-						OtherVertex = &EndVertex;
-						NewVertex = SplitAndLink(*Vertex, *SecondEdge, *Edge, SquareTolerance, SquareMinEdgeLength);
-					}
-
-					if (NewVertex.IsValid() && !Vertex->IsLinkedTo(*NewVertex) && !OtherVertex->IsLinkedTo(*NewVertex))
-					{
-						VerticesToProcess.Add(&*NewVertex);
-					}
+					OtherVertex = &EndVertex;
+					NewVertex = SplitAndLink(*Vertex, *SecondEdge, *Edge, SquareTolerance, SquareMinEdgeLength);
 				}
 
-				if (Vertex->IsDeleted() || !Vertex->IsBorderVertex())
+				if (NewVertex && !Vertex->IsLinkedTo(*NewVertex) && !OtherVertex->IsLinkedTo(*NewVertex))
 				{
-					break;
+					return NewVertex;
 				}
 			}
 		}
 	}
+	return nullptr;
+}
+
+void StitchParallelEdges(TArray<FTopologicalVertex*>& VerticesToProcess, double SewTolerance, double MinEdgeLength)
+{
+	double SquareTolerance = FMath::Square(SewTolerance);
+	double SquareMinEdgeLength = FMath::Square(MinEdgeLength);
+
+	FTimePoint StartTime = FChrono::Now();
+
+	for (int32 VertexI = 0; VertexI < VerticesToProcess.Num(); ++VertexI)
+	{
+		FTopologicalVertex* Vertex = VerticesToProcess[VertexI];
+
+		if (Vertex->IsDeleted() || !Vertex->IsBorderVertex())
+		{
+			continue;
+		}
+
+		while (Vertex)
+		{
+			Vertex = StitchParallelEdgesFrom(Vertex, SewTolerance, MinEdgeLength);
+		}
+	}
 
 	FDuration StepDuration = FChrono::Elapse(StartTime);
 	FChrono::PrintClockElapse(EVerboseLevel::Log, TEXT("    "), TEXT("Stitch Parallel Edges"), StepDuration);
@@ -418,6 +477,36 @@ void MergeCoincidentEdges(TArray<FTopologicalVertex*>& VerticesToProcess, double
 	FChrono::PrintClockElapse(EVerboseLevel::Log, TEXT("    "), TEXT("Merge coincident edges"), StepDuration);
 }
 
+void GetVertices(const TArray<FTopologicalEdge*>& InEdges, TArray<FTopologicalVertex*>& OutVertices)
+{
+	OutVertices.Reset(InEdges.Num());
+
+	for (const FTopologicalEdge* Edge : InEdges)
+	{
+		Edge->GetStartVertex()->GetLinkActiveEntity()->ResetMarker1();
+		Edge->GetEndVertex()->GetLinkActiveEntity()->ResetMarker1();
+	}
+
+	for (const FTopologicalEdge* Edge : InEdges)
+	{
+		if (!Edge->GetStartVertex()->GetLinkActiveEntity()->HasMarker1())
+		{
+			Edge->GetStartVertex()->GetLinkActiveEntity()->SetMarker1();
+			OutVertices.Add(&*Edge->GetStartVertex()->GetLinkActiveEntity());
+		}
+		if (!Edge->GetEndVertex()->GetLinkActiveEntity()->HasMarker1())
+		{
+			Edge->GetEndVertex()->GetLinkActiveEntity()->SetMarker1();
+			OutVertices.Add(&*Edge->GetEndVertex()->GetLinkActiveEntity());
+		}
+	}
+
+	for (FTopologicalVertex* Vertex : OutVertices)
+	{
+		Vertex->ResetMarker1();
+	}
+}
+
 } // namespace TopomakerTools
 
 FTopomaker::FTopomaker(FSession& InSession, const FTopomakerOptions& InOptions)
@@ -550,28 +639,31 @@ void FTopomaker::Sew()
 	// basic case: merge pair of edges connected together at both extremities i.e. pair of edges with same extremity active vertices.
 	TopomakerTools::MergeCoincidentEdges(BorderVertices, EdgeLengthTolerance);
 
-	if (SewOption::IsRemoveThinFaces(SewOptions))
-	{
-		TArray<FTopologicalEdge*> NewBorderEdges;
-		RemoveThinFaces(NewBorderEdges);
-	}
-
 	MergeUnconnectedSuccessiveEdges();
 
-	if (SewOption::IsForceJoining(SewOptions))
+	// advance case: two partially coincident edges connected at one extremity i.e. coincident along the shortest edge
+	TopomakerTools::StitchParallelEdges(BorderVertices, SewTolerance, EdgeLengthTolerance);
+
+	const bool bForceJoining = SewOption::IsForceJoining(SewOptions);
+	if (bForceJoining)
 	{
 		BorderVertices.Empty(BorderVertices.Num());
 		GetBorderVertices(BorderVertices);
 		TopomakerTools::MergeCoincidentVertices(BorderVertices, SewToleranceToForceJoin);
 
-		CheckSelfConnectedEdge(LargeEdgeLengthTolerance, BorderVertices);
+		CheckSelfConnectedEdge(EdgeLengthTolerance, BorderVertices);
 	}
 
-	const bool bForceJoining = SewOption::IsForceJoining(SewOptions);
-
 	// re process with new edges from MergeUnconnectedSuccessiveEdges and new merged vertices (if bForceJoining)
+	BorderVertices.Reset();
+	GetBorderVertices(BorderVertices);
 	TopomakerTools::MergeCoincidentEdges(BorderVertices, bForceJoining ? LargeEdgeLengthTolerance : EdgeLengthTolerance);
 
+	if (SewOption::IsRemoveThinFaces(SewOptions))
+	{
+		RemoveThinFaces();
+	}
+
 	// advance case: two partially coincident edges connected at one extremity i.e. coincident along the shortest edge
 	TopomakerTools::StitchParallelEdges(BorderVertices, bForceJoining ? SewToleranceToForceJoin : SewTolerance, bForceJoining ? LargeEdgeLengthTolerance : EdgeLengthTolerance);
 
@@ -602,6 +694,7 @@ void FTopomaker::GetVertices(TArray<FTopologicalVertex*>& Vertices)
 			for (const FOrientedEdge& OrientedEdge : Loop->GetEdges())
 			{
 				const TSharedPtr<FTopologicalEdge>& Edge = OrientedEdge.Entity;
+
 				if (!Edge->GetStartVertex()->GetLinkActiveEntity()->HasMarker1())
 				{
 					Edge->GetStartVertex()->GetLinkActiveEntity()->SetMarker1();
@@ -944,9 +1037,7 @@ void FindBorderVertex(const TArray<FTopologicalEdge*>& Edges, TArray<FTopologica
 
 }
 
-
-//#define DEBUG_THIN_FACE
-void FTopomaker::RemoveThinFaces(TArray<FTopologicalEdge*>& NewBorderEdges)
+void FTopomaker::RemoveThinFaces()
 {
 	FTimePoint StartTime = FChrono::Now();
 
@@ -954,7 +1045,8 @@ void FTopomaker::RemoveThinFaces(TArray<FTopologicalEdge*>& NewBorderEdges)
 	F3DDebugSession _(TEXT("RemoveThinFaces"));
 #endif
 
-	TArray<FTopologicalFace*> DeletedFaces;
+	TArray<FTopologicalEdge*> BorderEdges;
+	TArray<FTopologicalVertex*> BorderVertices;
 
 	// Find thin faces
 	for (TSharedPtr<FTopologicalFace> Face : Faces)
@@ -963,34 +1055,25 @@ void FTopomaker::RemoveThinFaces(TArray<FTopologicalEdge*>& NewBorderEdges)
 		double GapSize = 0;
 		if (Analyer.IsThinFace(GapSize))
 		{
-			Face->Remove(&NewBorderEdges);
-			DeletedFaces.Add(Face.Get());
-#ifdef CADKERNEL_DEV
-			Report.AddThinFace();
-#endif
-		}
-	}
-
 #ifdef DEBUG_THIN_FACE
-	for (const FTopologicalFace* Face : DeletedFaces)
-	{
-		F3DDebugSession _(TEXT("Face"));
-		Display(*Face);
-	}
-#endif
-
-	// Remove deleted edges of NewBorderEdges
-	{
-		int32 NewIndex = 0;
-		for (int32 Index = 0; Index < NewBorderEdges.Num(); ++Index)
-		{
-			if (!NewBorderEdges[Index]->HasMarker1())
 			{
-				NewBorderEdges[NewIndex] = NewBorderEdges[Index];
-				NewIndex++;
+				F3DDebugSession _(TEXT("Thin Face"));
+				Display(*Face);
+				Wait();
 			}
+#endif
+			Face->Remove(&BorderEdges);
+
+			TopomakerTools::GetVertices(BorderEdges, BorderVertices);
+			TopomakerTools::MergeCoincidentVertices(BorderVertices, GapSize * 1.2);
+			TopomakerTools::MergeCoincidentEdges(BorderEdges, GapSize * 1.2);
+
+			BorderEdges.Reset();
+
+#ifdef CADKERNEL_DEV
+			Report.AddThinFace();
+#endif
 		}
-		NewBorderEdges.SetNum(NewIndex);
 	}
 
 #ifdef CADKERNEL_DEV
diff --git a/Engine/Source/Runtime/Datasmith/CADKernel/Public/CADKernel/Geo/Curves/RestrictionCurve.h b/Engine/Source/Runtime/Datasmith/CADKernel/Public/CADKernel/Geo/Curves/RestrictionCurve.h
--- a/Engine/Source/Runtime/Datasmith/CADKernel/Public/CADKernel/Geo/Curves/RestrictionCurve.h
+++ b/Engine/Source/Runtime/Datasmith/CADKernel/Public/CADKernel/Geo/Curves/RestrictionCurve.h
@@ -150,9 +150,9 @@ public:
 	 * Project a set of points of a twin curve on the 3D polyline and return the coordinate of the projected point
 	 * @param ToleranceOfProjection: Max error between the both curve to stop the search of projection of a point 
 	 */
-	void ProjectTwinCurvePoints(const TArray<FPoint>& InPointsToProject, bool bSameOrientation, TArray<double>& OutProjectedPointCoords, double ToleranceOfProjection) const
+	void ProjectTwinCurvePoints(const FLinearBoundary& InBoundary, const TArray<FPoint>& InPointsToProject, bool bSameOrientation, TArray<double>& OutProjectedPointCoords, double ToleranceOfProjection) const
 	{
-		Polyline.ProjectCoincidentalPolyline(InPointsToProject, bSameOrientation, OutProjectedPointCoords, ToleranceOfProjection);
+		Polyline.ProjectCoincidentalPolyline(InBoundary, InPointsToProject, bSameOrientation, OutProjectedPointCoords, ToleranceOfProjection);
 	}
 
 	void ComputeIntersectionsWithIsos(const FLinearBoundary& InBoundary, const TArray<double>& InIsoCoordinates, const EIso InTypeIso, const FSurfacicTolerance& ToleranceIso, TArray<double>& OutIntersection) const
diff --git a/Engine/Source/Runtime/Datasmith/CADKernel/Public/CADKernel/Geo/Sampling/PolylineTools.h b/Engine/Source/Runtime/Datasmith/CADKernel/Public/CADKernel/Geo/Sampling/PolylineTools.h
--- a/Engine/Source/Runtime/Datasmith/CADKernel/Public/CADKernel/Geo/Sampling/PolylineTools.h
+++ b/Engine/Source/Runtime/Datasmith/CADKernel/Public/CADKernel/Geo/Sampling/PolylineTools.h
@@ -525,7 +525,7 @@ public:
 	 * @param ToleranceOfProjection: Max error between the both curve to stop the search of projection
 	 *    if ToleranceOfProjection < 0, it's compute with ComputeSquareToleranceForProjection
 	 */
-	void ProjectCoincidentalPolyline(const TArray<PointType>& InPointsToProject, bool bSameOrientation, TArray<double>& OutProjectedPointCoords, const double ToleranceOfProjection) const
+	void ProjectCoincidentalPolyline(const FLinearBoundary& InBoundary, const TArray<PointType>& InPointsToProject, bool bSameOrientation, TArray<double>& OutProjectedPointCoords, const double ToleranceOfProjection) const
 	{
 #ifdef DEBUG_PROJECT_COINCIDENTAL_POLYLINE
 		static int32 Counter = 0;
@@ -547,8 +547,11 @@ public:
 
 		const double SquareTol = ToleranceOfProjection > 0 ? FMath::Square(ToleranceOfProjection) : PolylineTools::ComputeSquareToleranceForProjection(PolylinePoints);
 
-		int32 StartIndex = 0;
-		const int32 EndIndex = PolylinePoints.Num() - 1;
+		int32 BoundaryIndices[2];
+		GetStartEndIndex(InBoundary, BoundaryIndices);
+
+		int32 StartIndex = BoundaryIndices[0];
+		const int32 EndIndex = BoundaryIndices[1] + 1;
 
 		TFunction<double(const PointType&)> ProjectPointToPolyline = [&](const PointType& InPointToProject)
 		{
diff --git a/Engine/Source/Runtime/Datasmith/CADKernel/Public/CADKernel/Geo/Sampling/SurfacicPolyline.h b/Engine/Source/Runtime/Datasmith/CADKernel/Public/CADKernel/Geo/Sampling/SurfacicPolyline.h
--- a/Engine/Source/Runtime/Datasmith/CADKernel/Public/CADKernel/Geo/Sampling/SurfacicPolyline.h
+++ b/Engine/Source/Runtime/Datasmith/CADKernel/Public/CADKernel/Geo/Sampling/SurfacicPolyline.h
@@ -259,10 +259,10 @@ public:
 	 * Project each point of a coincidental polyline on the Polyline.
 	 * @param ToleranceOfProjection: Max error between the both curve to stop the search of projection
 	 */
-	void ProjectCoincidentalPolyline(const TArray<FPoint>& InPointsToProject, bool bSameOrientation, TArray<double>& OutProjectedPointCoordinates, double ToleranceOfProjection) const
+	void ProjectCoincidentalPolyline(const FLinearBoundary& InBoundary, const TArray<FPoint>& InPointsToProject, bool bSameOrientation, TArray<double>& OutProjectedPointCoordinates, double ToleranceOfProjection) const
 	{
 		TPolylineApproximator<FPoint> Approximator3D(Coordinates, Points3D);
-		Approximator3D.ProjectCoincidentalPolyline(InPointsToProject, bSameOrientation, OutProjectedPointCoordinates, ToleranceOfProjection);
+		Approximator3D.ProjectCoincidentalPolyline(InBoundary, InPointsToProject, bSameOrientation, OutProjectedPointCoordinates, ToleranceOfProjection);
 	}
 
 	/**
diff --git a/Engine/Source/Runtime/Datasmith/CADKernel/Public/CADKernel/Math/Geometry.h b/Engine/Source/Runtime/Datasmith/CADKernel/Public/CADKernel/Math/Geometry.h
--- a/Engine/Source/Runtime/Datasmith/CADKernel/Public/CADKernel/Math/Geometry.h
+++ b/Engine/Source/Runtime/Datasmith/CADKernel/Public/CADKernel/Math/Geometry.h
@@ -620,6 +620,14 @@ inline bool FindIntersectionOfLines2D(const FSegment2D& LineAB, const FSegment2D
 CADKERNEL_API bool DoIntersect(const FSegment2D& SegmentAB, const FSegment2D& SegmentCD);
 CADKERNEL_API bool DoIntersectInside(const FSegment2D& SegmentAB, const FSegment2D& SegmentCD);
 
+inline bool AreParallel(const FSegment2D& SegmentAB, const FSegment2D& SegmentCD)
+{
+	const FPoint2D AB = SegmentAB.GetVector().Normalize();
+	const FPoint2D CD = SegmentCD.GetVector().Normalize();
+	const double ParallelCoef = AB ^ CD;
+	return (FMath::IsNearlyZero(ParallelCoef, DOUBLE_KINDA_SMALL_NUMBER));
+};
+
 inline double ComputeCosinus(FVector Vector, FVector OtherVector)
 {
 	Vector.Normalize();
diff --git a/Engine/Source/Runtime/Datasmith/CADKernel/Public/CADKernel/Math/SlopeUtils.h b/Engine/Source/Runtime/Datasmith/CADKernel/Public/CADKernel/Math/SlopeUtils.h
--- a/Engine/Source/Runtime/Datasmith/CADKernel/Public/CADKernel/Math/SlopeUtils.h
+++ b/Engine/Source/Runtime/Datasmith/CADKernel/Public/CADKernel/Math/SlopeUtils.h
@@ -126,6 +126,11 @@ inline double TransformIntoOrientedSlope(double Slope)
 	return WrapTo(Slope, -Slope::PiSlope, Slope::PiSlope, Slope::TwoPiSlope);
 }
 
+inline double TransformIntoClockwiseSlope(double Slope)
+{
+	return Slope::TwoPiSlope - Slope;
+}
+
 /**
  * Transform a positive slope into an unoriented slope [0, 4] i.e. an equivalent angle between [0, Pi]
  * @return a slope between [0, 4]
@@ -281,7 +286,7 @@ inline double ComputePositiveSlope(const FPoint2D& StartPoint, const FPoint2D& E
 
 inline double ClockwiseSlope(const FPoint2D& StartPoint, const FPoint2D& EndPoint, double ReferenceSlope)
 {
-	return Slope::TwoPiSlope - ComputePositiveSlope(StartPoint, EndPoint, ReferenceSlope);
+	return TransformIntoClockwiseSlope(ComputePositiveSlope(StartPoint, EndPoint, ReferenceSlope));
 }
 
 inline double CounterClockwiseSlope(const FPoint2D& StartPoint, const FPoint2D& EndPoint, double ReferenceSlope)
diff --git a/Engine/Source/Runtime/Datasmith/CADKernel/Public/CADKernel/Topo/TopologicalEdge.h b/Engine/Source/Runtime/Datasmith/CADKernel/Public/CADKernel/Topo/TopologicalEdge.h
--- a/Engine/Source/Runtime/Datasmith/CADKernel/Public/CADKernel/Topo/TopologicalEdge.h
+++ b/Engine/Source/Runtime/Datasmith/CADKernel/Public/CADKernel/Topo/TopologicalEdge.h
@@ -720,7 +720,7 @@ public:
 	void ProjectTwinEdgePoints(const TArray<FPoint>& InPointsToProject, bool bSameOrientation, TArray<double>& OutProjectedPointCoords) const
 	{
 		const double ToleranceOfProjection = Length3D * 0.1;
-		Curve->ProjectTwinCurvePoints(InPointsToProject, bSameOrientation, OutProjectedPointCoords, ToleranceOfProjection);
+		Curve->ProjectTwinCurvePoints(Boundary, InPointsToProject, bSameOrientation, OutProjectedPointCoords, ToleranceOfProjection);
 	}
 
 	/**
@@ -771,7 +771,7 @@ public:
 	 * @param bKeepStartVertexConnectivity: if true the new edge is connected to endVertex, otherwise the new edge is connected to start vertex
 	 * @param OutNewEdge, the second edge
 	 */
-	TSharedPtr<FTopologicalVertex> SplitAt(double SplittingCoordinate, const FPoint& NewVertexCoordinate, bool bKeepStartVertexConnectivity, TSharedPtr<FTopologicalEdge>& OutNewEdge);
+	FTopologicalVertex* SplitAt(double SplittingCoordinate, const FPoint& NewVertexCoordinate, bool bKeepStartVertexConnectivity, TSharedPtr<FTopologicalEdge>& OutNewEdge);
 
 	/**
 	 * Extend the Edge to the NewVertex.
diff --git a/Engine/Source/Runtime/Datasmith/CADKernel/Public/CADKernel/Topo/Topomaker.h b/Engine/Source/Runtime/Datasmith/CADKernel/Public/CADKernel/Topo/Topomaker.h
--- a/Engine/Source/Runtime/Datasmith/CADKernel/Public/CADKernel/Topo/Topomaker.h
+++ b/Engine/Source/Runtime/Datasmith/CADKernel/Public/CADKernel/Topo/Topomaker.h
@@ -132,7 +132,7 @@ public:
 	 */
 	void UnlinkEdgesOfSeparateShells();
 
-	void RemoveThinFaces(TArray<FTopologicalEdge*>& NewBorderEdges);
+	void RemoveThinFaces();
 
 #ifdef CADKERNEL_DEV
 	void PrintSewReport()
