fix Issue 313112: Static analyzer warnings in Topomaker.cpp

diff --git a/Engine/Source/Runtime/Datasmith/CADKernel/Private/CADKernel/Topo/Topomaker.cpp b/Engine/Source/Runtime/Datasmith/CADKernel/Private/CADKernel/Topo/Topomaker.cpp
--- a/Engine/Source/Runtime/Datasmith/CADKernel/Private/CADKernel/Topo/Topomaker.cpp
+++ b/Engine/Source/Runtime/Datasmith/CADKernel/Private/CADKernel/Topo/Topomaker.cpp
@@ -364,7 +364,7 @@ void StitchParallelEdges(TArray<FTopologicalVertex*>& VerticesToProcess, double
 	{
 		FTopologicalVertex* Vertex = VerticesToProcess[VertexI];
 
-		if (Vertex->IsDeleted() || !Vertex->IsBorderVertex())
+		if (!Vertex || Vertex->IsDeleted() || !Vertex->IsBorderVertex())
 		{
 			continue;
 		}
