Fix Alias Import None/Heal
CADKernel:
- Fix crash in meshing
- Improve Healing

diff --git a/Engine/Plugins/Enterprise/DatasmithCADImporter/Source/CADTools/Private/CADOption.cpp b/Engine/Plugins/Enterprise/DatasmithCADImporter/Source/CADTools/Private/CADOption.cpp
--- a/Engine/Plugins/Enterprise/DatasmithCADImporter/Source/CADTools/Private/CADOption.cpp
+++ b/Engine/Plugins/Enterprise/DatasmithCADImporter/Source/CADTools/Private/CADOption.cpp
@@ -81,15 +81,14 @@ Only available with CADKernel.\n\
 Default value of ForceFactor is 5.\n"),
 ECVF_Default);
 
-bool FImportParameters::bGStitchingRemoveThinFaces = false;
-/*FAutoConsoleVariableRef GCADTranslatorStitchingRemoveThinFaces(
+bool FImportParameters::bGStitchingRemoveThinFaces = true;
+FAutoConsoleVariableRef GCADTranslatorStitchingRemoveThinFaces(
 	TEXT("ds.CADTranslator.Stitching.RemoveThinFaces"),
 	FImportParameters::bGStitchingRemoveThinFaces,
 	TEXT("During the welding process, Thin faces are removed before the stiching process. The width of the thin faces is equal to the force sew tolerance.\n\
 Only available with CADKernel.\n\
 Default value of RemoveThinFaces is true\n"),
 ECVF_Default);
-*/
 
 bool FImportParameters::bGStitchingRemoveDuplicatedFaces = true;
 FAutoConsoleVariableRef GCADTranslatorStitchingRemoveDuplicatedFaces(
diff --git a/Engine/Plugins/Enterprise/DatasmithCADImporter/Source/DatasmithWireTranslator/Private/AliasModelToCADKernelConverter.cpp b/Engine/Plugins/Enterprise/DatasmithCADImporter/Source/DatasmithWireTranslator/Private/AliasModelToCADKernelConverter.cpp
--- a/Engine/Plugins/Enterprise/DatasmithCADImporter/Source/DatasmithWireTranslator/Private/AliasModelToCADKernelConverter.cpp
+++ b/Engine/Plugins/Enterprise/DatasmithCADImporter/Source/DatasmithWireTranslator/Private/AliasModelToCADKernelConverter.cpp
@@ -330,14 +330,11 @@ bool FAliasModelToCADKernelConverter::AddBRep(AlDagNode& DagNode, const FColor&
 		break;
 	}
 
-	if (CADKernelShell->FaceCount() == 0)
-	{
-		return false;
-	}
 
 	// Create body from faces
 	CADKernelSession.GetModel().Add(CADKernelBody);
-	return true;
+
+	return CADKernelShell->FaceCount() > 0;
 }
 
 }
diff --git a/Engine/Plugins/Enterprise/DatasmithCADImporter/Source/DatasmithWireTranslator/Private/DatasmithWireTranslator.cpp b/Engine/Plugins/Enterprise/DatasmithCADImporter/Source/DatasmithWireTranslator/Private/DatasmithWireTranslator.cpp
--- a/Engine/Plugins/Enterprise/DatasmithCADImporter/Source/DatasmithWireTranslator/Private/DatasmithWireTranslator.cpp
+++ b/Engine/Plugins/Enterprise/DatasmithCADImporter/Source/DatasmithWireTranslator/Private/DatasmithWireTranslator.cpp
@@ -2037,6 +2037,9 @@ void FWireTranslatorImpl::DagForLeavesNoMerge(const TSharedPtr<AlDagNode>& DagNo
 
 TOptional<FMeshDescription> FWireTranslatorImpl::MeshDagNodeWithExternalMesher(AlDagNode& DagNode, TSharedRef<IDatasmithMeshElement> MeshElement, CADLibrary::FMeshParameters& MeshParameters)
 {
+	static uint32 HelloWorld = 0;
+	HelloWorld++;
+
 	// Wire unit is cm
 	CADModelConverter->InitializeProcess();
 
@@ -2059,7 +2062,10 @@ TOptional<FMeshDescription> FWireTranslatorImpl::MeshDagNodeWithExternalMesher(A
 		ColorPtr = &DefaultColor;
 	}
 
-	AliasBRepConverter->AddBRep(DagNode, *ColorPtr, ObjectReference);
+	if (!AliasBRepConverter->AddBRep(DagNode, *ColorPtr, ObjectReference))
+	{
+		return TOptional<FMeshDescription>();
+	}
 
 	CADModelConverter->RepairTopology();
 
diff --git a/Engine/Source/Runtime/Datasmith/CADKernel/Private/CADKernel/Core/Database.cpp b/Engine/Source/Runtime/Datasmith/CADKernel/Private/CADKernel/Core/Database.cpp
--- a/Engine/Source/Runtime/Datasmith/CADKernel/Private/CADKernel/Core/Database.cpp
+++ b/Engine/Source/Runtime/Datasmith/CADKernel/Private/CADKernel/Core/Database.cpp
@@ -306,7 +306,7 @@ void FDatabase::Serialize(FCADKernelArchive& Ar)
 
 	ensure(Ar.IsSaving());
 
-	SpawnEntityIdent(*Model, true);
+	SpawnEntityIdent(GetModel(), true);
 
 	Ar.Session.Serialize(Ar);
 
@@ -374,6 +374,9 @@ void FDatabase::Deserialize(FCADKernelArchive& Ar)
 
 void FDatabase::CleanArchiveEntities()
 {
+	// Initialize the model if not
+	GetModel();
+
 	int32 FaceCount = 0;
 	int32 ShellCount = 0;
 	for (TSharedPtr<FEntity> Entity : ArchiveEntities)
diff --git a/Engine/Source/Runtime/Datasmith/CADKernel/Private/CADKernel/Mesh/Meshers/CycleTriangulator.cpp b/Engine/Source/Runtime/Datasmith/CADKernel/Private/CADKernel/Mesh/Meshers/CycleTriangulator.cpp
--- a/Engine/Source/Runtime/Datasmith/CADKernel/Private/CADKernel/Mesh/Meshers/CycleTriangulator.cpp
+++ b/Engine/Source/Runtime/Datasmith/CADKernel/Private/CADKernel/Mesh/Meshers/CycleTriangulator.cpp
@@ -116,7 +116,7 @@ bool FCycleTriangulator::FindTheCycleToMesh(FIsoSegment* Segment, bool bOrientat
 
 #ifdef DEBUG_FIND_THE_CYCLE_TO_MESH
 	static int32 CycleToMeshIndex = 1;
-	static int32 CycleToMeshIndexToTest = 1403;
+	static int32 CycleToMeshIndexToTest = 0;
 	F3DDebugSession A(Grid.bDisplay, FString::Printf(TEXT("Find The Cycle To Mesh %d"), ++CycleToMeshIndex));
 
 	if (Grid.bDisplay)
@@ -453,11 +453,11 @@ bool FCycleTriangulator::FindCandidateNodes(int32 StartIndex)
 	const double StartMaxSlope = FMath::Min(CandidateNodes.Last().SlopeAtStartNode + Slope::OneDegree, MaxSlope);
 	const double EndMaxSlope = FMath::Min(CandidateNodes[0].SlopeAtEndNode + Slope::OneDegree, MaxSlope);
 
-	if (CandidateNodes.Last().SlopeAtStartNode < MaxSlope)
+	if (CandidateNodes.Last().SlopeAtEndNode < EndMaxSlope)
 	{
 		CandidateNodes.Last().bIsValid = true;
 	}
-	if (CandidateNodes[0].SlopeAtEndNode < MaxSlope)
+	if (CandidateNodes[0].SlopeAtStartNode < StartMaxSlope)
 	{
 		CandidateNodes[0].bIsValid = true;
 	}
@@ -521,9 +521,9 @@ bool FCycleTriangulator::FindTheBestCandidateNode()
 
 #ifdef DEBUG_FIND_THE_BEST_CANDIDATE_NODE
 	static int32 BestCandidateNode = 0;
-	static int32 BestCandidateNodeTarget = 129;
+	static int32 BestCandidateNodeTarget = 100000;
 	++BestCandidateNode;
-	if (Grid.bDisplay /*&& BestCandidateNode >= BestCandidateNodeTarget*/)
+	if (Grid.bDisplay && BestCandidateNode >= BestCandidateNodeTarget)
 	{
 		F3DDebugSession _(FString::Printf(TEXT("Start Segment %d"), BestCandidateNode));
 		Grid.DisplayIsoSegment(EGridSpace::UniformScaled, *FirstSideStartNode, *FirstSideEndNode, 1, EVisuProperty::RedCurve);
@@ -539,7 +539,7 @@ bool FCycleTriangulator::FindTheBestCandidateNode()
 			{
 				break;
 			}
-			IntersectionCountAllowed++;
+			IntersectionCountAllowed = NodeIntersection.Value;
 		}
 
 		FCandidateNode* CNode = NodeIntersection.Key;
@@ -550,7 +550,7 @@ bool FCycleTriangulator::FindTheBestCandidateNode()
 		double PointCriteria = IsoTriangulatorImpl::IsoscelesCriteriaMax(StartPoint2D, EndPoint2D, NodePoint2D);
 
 #ifdef DEBUG_FIND_THE_BEST_CANDIDATE_NODE
-		if (Grid.bDisplay /*&& BestCandidateNode == BestCandidateNodeTarget*/)
+		if (Grid.bDisplay && BestCandidateNode >= BestCandidateNodeTarget)
 		{
 			F3DDebugSession _(TEXT("Candidate triangle"));
 			Grid.DisplayIsoSegment(EGridSpace::UniformScaled, *FirstSideStartNode, *Node, 1, EVisuProperty::RedCurve);
diff --git a/Engine/Source/Runtime/Datasmith/CADKernel/Private/CADKernel/Topo/Body.cpp b/Engine/Source/Runtime/Datasmith/CADKernel/Private/CADKernel/Topo/Body.cpp
--- a/Engine/Source/Runtime/Datasmith/CADKernel/Private/CADKernel/Topo/Body.cpp
+++ b/Engine/Source/Runtime/Datasmith/CADKernel/Private/CADKernel/Topo/Body.cpp
@@ -4,10 +4,6 @@
 #include "CADKernel/Topo/Model.h"
 #include "CADKernel/Topo/Shell.h"
 
-#ifdef CADKERNEL_DEV
-#include "CADKernel/Topo/TopologyReport.h"
-#endif
-
 namespace UE::CADKernel
 {
 
@@ -54,20 +50,4 @@ void FBody::Remove(const FTopologicalShapeEntity* ShellToRemove)
 	Shells.RemoveAt(Index);
 }
 
-#ifdef CADKERNEL_DEV
-FInfoEntity& FBody::GetInfo(FInfoEntity& Info) const
-{
-	return FTopologicalShapeEntity::GetInfo(Info).Add(TEXT("Shells"), Shells);
-}
-
-void FBody::FillTopologyReport(FTopologyReport& Report) const
-{
-	Report.Add(this);
-	for (TSharedPtr<FShell> Shell : Shells)
-	{
-		Shell->FillTopologyReport(Report);
-	}
-}
-#endif
-
 }
diff --git a/Engine/Source/Runtime/Datasmith/CADKernel/Private/CADKernel/Topo/FaceAnalyzer.cpp b/Engine/Source/Runtime/Datasmith/CADKernel/Private/CADKernel/Topo/FaceAnalyzer.cpp
--- a/Engine/Source/Runtime/Datasmith/CADKernel/Private/CADKernel/Topo/FaceAnalyzer.cpp
+++ b/Engine/Source/Runtime/Datasmith/CADKernel/Private/CADKernel/Topo/FaceAnalyzer.cpp
@@ -11,10 +11,6 @@
 namespace UE::CADKernel
 {
 
-#ifdef CADKERNEL_DEV
-FIdent Topo::FEdgeSegment::LastId = 0;
-#endif
-
 void FFaceAnalyzer::FindClosedSegments(Topo::FThinFaceContext& Context)
 {
 #ifdef DEBUG_FIND_THIN_FACE
@@ -227,61 +223,5 @@ bool FFaceAnalyzer::IsThinFace(double& OutGapSize)
 	return false;
 }
 
-#ifdef CADKERNEL_DEV
-void FFaceAnalyzer::DisplayLoopSegments(Topo::FThinFaceContext& Context)
-{
-#ifdef DEBUG_THIN_FACE
-	const FTopologicalEdge* Edge = nullptr;
-	Edge = Context.LoopSegments[0]->GetEdge();
-
-	F3DDebugSession _(TEXT("BoundarySegment"));
-
-#ifdef DEBUG_THIN_FACE_EDGE
-	Open3DDebugSession(FString::Printf(TEXT("Edge %d\n"), Edge->GetId()));
-#endif
-
-	for (Topo::FEdgeSegment* EdgeSegment : Context.LoopSegments)
-	{
-		if (Edge != EdgeSegment->GetEdge())
-		{
-			Edge = EdgeSegment->GetEdge();
-#ifdef DEBUG_THIN_FACE_EDGE
-			Close3DDebugSession();
-			Open3DDebugSession(FString::Printf(TEXT("Edge %d\n"), Edge->GetId()));
-#endif
-		}
-		DisplaySegment(EdgeSegment->GetExtemity(ELimit::Start), EdgeSegment->GetExtemity(ELimit::End), EdgeSegment->GetId(), EVisuProperty::GreenCurve, true);
-		DisplayPoint(EdgeSegment->GetExtemity(ELimit::Start), EVisuProperty::BluePoint);
-	}
-#ifdef DEBUG_THIN_FACE_EDGE
-	Close3DDebugSession();
-#endif
-	//Wait();
-#endif
-}
-sadf
-#ifdef CADKERNEL_DEV
-void FFaceAnalyzer::DisplayClosedSegments(Topo::FThinFaceContext& Context)
-{
-	F3DDebugSession _(TEXT("Closed Segment"));
-	for (const Topo::FEdgeSegment* Segment : Context.LoopSegments)
-	{
-		F3DDebugSession _(TEXT("Closed"));
-		if (Segment->GetClosedSegment())
-		{
-			DisplaySegment(Segment->GetExtemity(ELimit::Start), Segment->GetExtemity(ELimit::End), Segment->GetId(), EVisuProperty::BlueCurve);
-			Topo::FEdgeSegment* Parallel = Segment->GetClosedSegment();
-			DisplaySegment(Parallel->GetExtemity(ELimit::Start), Parallel->GetExtemity(ELimit::End), Parallel->GetId(), EVisuProperty::BlueCurve);
-
-			double Coordinate;
-			FPoint Projection = ProjectPointOnSegment(Segment->GetMiddle(), Parallel->GetExtemity(ELimit::Start), Parallel->GetExtemity(ELimit::End), Coordinate, true);
-
-			DisplaySegment(Projection, Segment->GetMiddle(), Segment->GetId(), EVisuProperty::RedCurve);
-		}
-	}
-	//Wait();
-}
-#endif
-#endif
 }
 
diff --git a/Engine/Source/Runtime/Datasmith/CADKernel/Private/CADKernel/Topo/Model.cpp b/Engine/Source/Runtime/Datasmith/CADKernel/Private/CADKernel/Topo/Model.cpp
--- a/Engine/Source/Runtime/Datasmith/CADKernel/Private/CADKernel/Topo/Model.cpp
+++ b/Engine/Source/Runtime/Datasmith/CADKernel/Private/CADKernel/Topo/Model.cpp
@@ -90,14 +90,6 @@ void FModel::SpreadBodyOrientation()
 	}
 }
 
-#ifdef CADKERNEL_DEV
-FInfoEntity& FModel::GetInfo(FInfoEntity& Info) const
-{
-	return FTopologicalShapeEntity::GetInfo(Info)
-		.Add(TEXT("Bodies"), Bodies);
-}
-#endif
-
 struct FBodyShell
 {
 	TSharedPtr<FBody> Body;
diff --git a/Engine/Source/Runtime/Datasmith/CADKernel/Private/CADKernel/Topo/Shell.cpp b/Engine/Source/Runtime/Datasmith/CADKernel/Private/CADKernel/Topo/Shell.cpp
--- a/Engine/Source/Runtime/Datasmith/CADKernel/Private/CADKernel/Topo/Shell.cpp
+++ b/Engine/Source/Runtime/Datasmith/CADKernel/Private/CADKernel/Topo/Shell.cpp
@@ -529,4 +529,69 @@ int32 FShell::Orient()
 	return ShellSwappedFaceCount;
 }
 
+namespace ShellTools
+{
+
+void UnlinkFromOther(TArray<FTopologicalFace*>& Faces, TArray<FTopologicalVertex*>& VerticesToLink)
+{
+	for (FTopologicalFace* Face : Faces)
+	{
+		Face->SetToProcessMarker();
+	}
+
+	VerticesToLink.Reserve(100);
+
+	for (FTopologicalFace* Face : Faces)
+	{
+		for (const TSharedPtr<FTopologicalLoop>& Loop : Face->GetLoops())
+		{
+			for (const FOrientedEdge& OrientedEdge : Loop->GetEdges())
+			{
+				FTopologicalEdge* Edge = OrientedEdge.Entity.Get();
+				if (!Edge)
+				{
+					continue;
+				}
+
+				for (FTopologicalEdge* TwinEdge : Edge->GetTwinEntities())
+				{
+					if (!TwinEdge || TwinEdge == Edge)
+					{
+						continue;
+					}
+
+					FTopologicalFace* NeighborFace = TwinEdge->GetFace();
+					if (!NeighborFace)
+					{
+						continue;
+					}
+					if (!NeighborFace->IsToProcess())
+					{
+						Edge->Unlink();
+
+						FTopologicalVertex& Vertex = *Edge->GetStartVertex();
+						if(!Vertex.IsProcessed())
+						{
+							VerticesToLink.Add(&Vertex);
+						}
+						break;
+					}
+				}
+			}
+		}
+	}
+
+	for (FTopologicalFace* Face : Faces)
+	{
+		Face->ResetToProcessMarker();
+	}
+
+	for (FTopologicalVertex* Vertex : VerticesToLink)
+	{
+		Vertex->ResetProcessedMarker();
+	}
+}
+}
+
+
 }
\ No newline at end of file
diff --git a/Engine/Source/Runtime/Datasmith/CADKernel/Private/CADKernel/Topo/TopologicalEdge.cpp b/Engine/Source/Runtime/Datasmith/CADKernel/Private/CADKernel/Topo/TopologicalEdge.cpp
--- a/Engine/Source/Runtime/Datasmith/CADKernel/Private/CADKernel/Topo/TopologicalEdge.cpp
+++ b/Engine/Source/Runtime/Datasmith/CADKernel/Private/CADKernel/Topo/TopologicalEdge.cpp
@@ -261,7 +261,7 @@ bool FTopologicalEdge::IsLinkableTo(const FTopologicalEdge& Edge, double EdgeLen
 
 void FTopologicalEdge::LinkIfCoincident(FTopologicalEdge& Twin, double EdgeLengthTolerance, double SquareJoiningTolerance)
 {
-	if(IsDeleted() || Twin.IsDeleted())
+	if (IsDeleted() || Twin.IsDeleted())
 	{
 		return;
 	}
@@ -269,7 +269,7 @@ void FTopologicalEdge::LinkIfCoincident(FTopologicalEdge& Twin, double EdgeLengt
 	// Degenerated twin edges are not linked
 	if (IsDegenerated() || Twin.IsDegenerated())
 	{
-		if(HasSameLengthAs(Twin, EdgeLengthTolerance))
+		if (HasSameLengthAs(Twin, EdgeLengthTolerance))
 		{
 			SetAsDegenerated();
 			Twin.SetAsDegenerated();
@@ -359,6 +359,13 @@ void FTopologicalEdge::Link(FTopologicalEdge& Twin)
 	MakeLink(Twin);
 }
 
+void FTopologicalEdge::Disjoin()
+{
+	RemoveFromLink();
+	GetStartVertex()->RemoveFromLink();
+	GetEndVertex()->RemoveFromLink();
+}
+
 void FTopologicalEdge::Empty()
 {
 	if (StartVertex.IsValid())
@@ -388,7 +395,7 @@ void FTopologicalEdge::Empty()
 
 FTopologicalFace* FTopologicalEdge::GetFace() const
 {
-	if(Loop != nullptr)
+	if (Loop != nullptr)
 	{
 		return Loop->GetFace();
 	}
@@ -435,19 +442,19 @@ void FTopologicalEdge::ComputeCrossingPointCoordinates()
 			double NewDelatU = DeltaU / (Index - IndexMin);
 			for (int32 Andex = IndexMin + 1; Andex < Index; ++Andex)
 			{
-				CrossingPointUs[Andex] = CrossingPointUs[Andex-1] + NewDelatU;
-			}		
+				CrossingPointUs[Andex] = CrossingPointUs[Andex - 1] + NewDelatU;
+			}
 		}
 		IndexMin = Index;
 		LocalUMin = DeltaUMin;
 	}
 
-	if(IndexMin != CrossingPointUs.Num() - 1)
+	if (IndexMin != CrossingPointUs.Num() - 1)
 	{
 		IndexMin--;
 		double DeltaU = CrossingPointUs[CrossingPointUs.Num() - 1] - CrossingPointUs[IndexMin];
 		double NewDelatU = DeltaU / (CrossingPointUs.Num() - 1 - IndexMin);
-		for (int32 Index = IndexMin + 1; Index < CrossingPointUs.Num()-1; ++Index)
+		for (int32 Index = IndexMin + 1; Index < CrossingPointUs.Num() - 1; ++Index)
 		{
 			CrossingPointUs[Index] = CrossingPointUs[Index - 1] + NewDelatU;
 		}
@@ -550,7 +557,7 @@ double FTopologicalEdge::TransformTwinEdgeCoordinateToLocalCoordinate(const FTop
 	{
 		// linear transform
 		const bool bSameDirection = IsSameDirection(TwinEdge);
-		const double Start = bSameDirection ? Boundary.GetMin(): Boundary.GetMax();
+		const double Start = bSameDirection ? Boundary.GetMin() : Boundary.GetMax();
 		const double End = bSameDirection ? Boundary.GetMax() : Boundary.GetMin();
 
 		const double Distance = End - Start;
@@ -630,7 +637,7 @@ void FTopologicalEdge::TransformTwinEdgeCoordinatesToLocalCoordinates(const FTop
 
 		OutLocalCoordinates.Empty(InTwinCoordinates.Num());
 		const double Factor = Distance / TwinDistance;
-		for(double TwinCoordinate : InTwinCoordinates)
+		for (double TwinCoordinate : InTwinCoordinates)
 		{
 			OutLocalCoordinates.Add(Start + (TwinCoordinate - TwinStart) * Factor);
 		}
@@ -826,22 +833,6 @@ bool FTopologicalEdge::IsSameDirection(const FTopologicalEdge& Edge) const
 	return Vertex1Edge == Edge.GetStartVertex()->GetLink();
 }
 
-#ifdef CADKERNEL_DEV
-FInfoEntity& FTopologicalEdge::GetInfo(FInfoEntity& Info) const
-{
-	return FTopologicalEntity::GetInfo(Info)
-		.Add(TEXT("IsDegenerated"), IsDegenerated())
-		.Add(TEXT("Link"), TopologicalLink)
-		.Add(TEXT("Curve"), Curve)
-		.Add(TEXT("Vertex1"), StartVertex)
-		.Add(TEXT("Vertex2"), EndVertex)
-		.Add(TEXT("Boundary"), Boundary)
-		.Add(TEXT("Loop"), Loop)
-		.Add(TEXT("Length"), Length())
-		.Add(TEXT("Mesh"), Mesh);
-}
-#endif
-
 FEdgeMesh& FTopologicalEdge::GetOrCreateMesh(FModelMesh& ShellMesh)
 {
 	if (!IsActiveEntity())
@@ -967,28 +958,14 @@ TArray<double> FTopologicalEdge::GetPreElementLengths() const
 	return PolylineTools::ComputePolylineSegmentLengths(StartNode, InnerNodes, EndNode);
 }
 
-TSharedPtr<FTopologicalEdge> FTopologicalEdge::CreateEdgeByMergingEdges(TArray<FOrientedEdge>& Edges, const TSharedRef<FTopologicalVertex> StartVertex, const TSharedRef<FTopologicalVertex> EndVertex)
+TSharedPtr<FTopologicalEdge> FTopologicalEdge::CreateEdgeByMergingEdges(const double SmallEdgeTolerance, TArray<FOrientedEdge>& Edges, const TSharedRef<FTopologicalVertex>& StartVertex, const TSharedRef<FTopologicalVertex>& EndVertex)
 {
-	// Make merged 2d Nurbs ===================================================
-
-	const double Tolerance3D = Edges[0].Entity->GetTolerance3D();
-	//const double Tolerance2D = Edges[0].Entity->GetCurve()->Get2DCurve()->GetDomainTolerance();
 	TSharedRef<FSurface> CarrierSurface = Edges[0].Entity->GetCurve()->GetCarrierSurface();
 
 	// check if all curves are 2D NURBS
 	bool bAreNurbs = true;
 	int32 NurbsMaxDegree = 0;
 
-	// TODO, Check if some edges are small edges and could be replace by extending next or previous one
-	double NewEdgeLength = 0;
-	for (const FOrientedEdge& Edge : Edges)
-	{
-		NewEdgeLength += Edge.Entity->Length();
-	}
-
-	// MinLength = smaller than this size, the adgacent edge is extend to replace it  
-	const double MinLength = FMath::Min(NewEdgeLength / (Edges.Num() + 3.), FMath::Max(NewEdgeLength / 20., Tolerance3D * 5.)); // Tolerance3D = 0.01mm => MinLength = 0.25mm, 1/20 of the length is small enougth to be a detail.
-
 	TArray<TSharedPtr<FNURBSCurve>> NurbsCurves;
 	NurbsCurves.Reserve(Edges.Num());
 
@@ -1001,7 +978,7 @@ TSharedPtr<FTopologicalEdge> FTopologicalEdge::CreateEdgeByMergingEdges(TArray<F
 		}
 
 		double EdgeLength = Edge.Entity->Length();
-		if (bCanRemove && EdgeLength < MinLength)
+		if (bCanRemove && EdgeLength < SmallEdgeTolerance)
 		{
 			NurbsCurves.Emplace(TSharedPtr<FNURBSCurve>());
 			bCanRemove = false;
@@ -1010,7 +987,7 @@ TSharedPtr<FTopologicalEdge> FTopologicalEdge::CreateEdgeByMergingEdges(TArray<F
 		bCanRemove = true;
 
 
-		// Find the max degree of the nurbs
+		// Find the max degree of the NURBS
 		TSharedPtr<FNURBSCurve> NURBS = NurbsCurves.Emplace_GetRef(StaticCastSharedRef<FNURBSCurve>(Edge.Entity->GetCurve()->Get2DCurve()));
 		int32 NurbsDegree = NURBS->GetDegree();
 		if (NurbsDegree > NurbsMaxDegree)
@@ -1050,7 +1027,7 @@ TSharedPtr<FTopologicalEdge> FTopologicalEdge::CreateEdgeByMergingEdges(TArray<F
 		if (NURBS->GetDegree() < NurbsMaxDegree)
 		{
 			NURBS = BSpline::DuplicateNurbsCurveWithHigherDegree(NurbsMaxDegree, *NURBS);
-			if(!NURBS.IsValid())
+			if (!NURBS.IsValid())
 			{
 				// cancel
 				return TSharedPtr<FTopologicalEdge>();
@@ -1123,7 +1100,7 @@ TSharedPtr<FTopologicalEdge> FTopologicalEdge::CreateEdgeByMergingEdges(TArray<F
 		}
 	}
 
-	if(bIsRational)
+	if (bIsRational)
 	{
 		NewWeights.Reserve(PoleCount + NurbsMaxDegree);
 		for (const TSharedPtr<FNURBSCurve>& NurbsCurve : NurbsCurves)
@@ -1140,7 +1117,7 @@ TSharedPtr<FTopologicalEdge> FTopologicalEdge::CreateEdgeByMergingEdges(TArray<F
 			}
 
 			NewPoles.Append(NurbsCurve->GetPoles());
-			if(NurbsCurve->IsRational())
+			if (NurbsCurve->IsRational())
 			{
 				NewWeights.Append(NurbsCurve->GetWeights());
 			}
@@ -1475,13 +1452,4 @@ TArray<FTopologicalFace*> FTopologicalEdge::GetLinkedFaces() const
 	return MoveTemp(NeighborFaces);
 }
 
-#ifdef CADKERNEL_DEV
-FInfoEntity& FEdgeLink::GetInfo(FInfoEntity& Info) const
-{
-	return FEntity::GetInfo(Info)
-		.Add(TEXT("Active Entity"), ActiveEntity)
-		.Add(TEXT("Twin Entities"), TwinEntities);
-}
-#endif
-
 }
diff --git a/Engine/Source/Runtime/Datasmith/CADKernel/Private/CADKernel/Topo/TopologicalFace.cpp b/Engine/Source/Runtime/Datasmith/CADKernel/Private/CADKernel/Topo/TopologicalFace.cpp
--- a/Engine/Source/Runtime/Datasmith/CADKernel/Private/CADKernel/Topo/TopologicalFace.cpp
+++ b/Engine/Source/Runtime/Datasmith/CADKernel/Private/CADKernel/Topo/TopologicalFace.cpp
@@ -282,10 +282,8 @@ void FTopologicalFace::Disjoin(TArray<FTopologicalEdge*>* NewBorderEdges)
 					NewBorderEdges->Add(TwinEdge);
 				}
 			}
-			Edge->RemoveFromLink();
+			Edge->Disjoin();
 			Edge->SetMarker1();
-			Edge->GetStartVertex()->RemoveFromLink();
-			Edge->GetEndVertex()->RemoveFromLink();
 		}
 	}
 }
@@ -457,19 +455,6 @@ void FTopologicalFace::SpawnIdent(FDatabase& Database)
 	}
 }
 
-#ifdef CADKERNEL_DEV
-FInfoEntity& FTopologicalFace::GetInfo(FInfoEntity& Info) const
-{
-	return FTopologicalShapeEntity::GetInfo(Info)
-		.Add(TEXT("Orientation"), IsBackOriented() ? "Back" : "Front")
-		.Add(TEXT("Carrier Surface"), CarrierSurface)
-		.Add(TEXT("Boundary"), (FSurfacicBoundary&) Boundary)
-		.Add(TEXT("Loops"), Loops)
-		.Add(TEXT("QuadCriteria"), QuadCriteria)
-		.Add(TEXT("Mesh"), Mesh);
-}
-#endif
-
 FFaceMesh& FTopologicalFace::GetOrCreateMesh(FModelMesh& MeshModel)
 {
 	if (!Mesh.IsValid())
diff --git a/Engine/Source/Runtime/Datasmith/CADKernel/Private/CADKernel/Topo/TopologicalLoop.cpp b/Engine/Source/Runtime/Datasmith/CADKernel/Private/CADKernel/Topo/TopologicalLoop.cpp
--- a/Engine/Source/Runtime/Datasmith/CADKernel/Private/CADKernel/Topo/TopologicalLoop.cpp
+++ b/Engine/Source/Runtime/Datasmith/CADKernel/Private/CADKernel/Topo/TopologicalLoop.cpp
@@ -1015,16 +1015,6 @@ void FTopologicalLoop::EnsureLogicalClosing(const double Tolerance3D)
 	}
 }
 
-#ifdef CADKERNEL_DEV
-FInfoEntity& FTopologicalLoop::GetInfo(FInfoEntity& Info) const
-{
-	return FEntity::GetInfo(Info)
-		.Add(TEXT("Edges"), (TArray<TOrientedEntity<FEntity>>&) Edges)
-		.Add(TEXT("IsExternal"), bIsExternal)
-		.Add(TEXT("Hosted by"), Face);
-}
-#endif
-
 void FTopologicalLoop::FindBreaks(TArray<TSharedPtr<FTopologicalVertex>>& OutBreaks, TArray<int32>& OutStartSideIndex, TArray<double>& OutBreakValues) const
 {
 	const double MinCosAngleOfBreak = -0.7;  // 135 deg
diff --git a/Engine/Source/Runtime/Datasmith/CADKernel/Private/CADKernel/Topo/TopologicalShapeEntity.cpp b/Engine/Source/Runtime/Datasmith/CADKernel/Private/CADKernel/Topo/TopologicalShapeEntity.cpp
--- a/Engine/Source/Runtime/Datasmith/CADKernel/Private/CADKernel/Topo/TopologicalShapeEntity.cpp
+++ b/Engine/Source/Runtime/Datasmith/CADKernel/Private/CADKernel/Topo/TopologicalShapeEntity.cpp
@@ -12,14 +12,4 @@ void FTopologicalShapeEntity::CompleteMetadata()
 	}
 }
 
-
-#ifdef CADKERNEL_DEV
-FInfoEntity& FTopologicalShapeEntity::GetInfo(FInfoEntity& Info) const
-{
-	return FTopologicalEntity::GetInfo(Info)
-		.Add(TEXT("Hosted by"), (FEntity*) HostedBy)
-		.Add(Dictionary);
-}
-#endif
-
 }
diff --git a/Engine/Source/Runtime/Datasmith/CADKernel/Private/CADKernel/Topo/TopologicalVertex.cpp b/Engine/Source/Runtime/Datasmith/CADKernel/Private/CADKernel/Topo/TopologicalVertex.cpp
--- a/Engine/Source/Runtime/Datasmith/CADKernel/Private/CADKernel/Topo/TopologicalVertex.cpp
+++ b/Engine/Source/Runtime/Datasmith/CADKernel/Private/CADKernel/Topo/TopologicalVertex.cpp
@@ -9,17 +9,6 @@
 namespace UE::CADKernel
 {
 
-#ifdef CADKERNEL_DEV
-FInfoEntity& FTopologicalVertex::GetInfo(FInfoEntity& Info) const
-{
-	return FTopologicalEntity::GetInfo(Info)
-		.Add(TEXT("Link"), TopologicalLink)
-		.Add(TEXT("Position"), Coordinates)
-		.Add(TEXT("ConnectedEdges"), ConnectedEdges)
-		.Add(TEXT("mesh"), Mesh);
-}
-#endif
-
 void FTopologicalVertex::AddConnectedEdge(FTopologicalEdge& Edge)
 {
 	ConnectedEdges.Add(&Edge);
@@ -212,19 +201,4 @@ void FTopologicalVertex::SpawnIdent(FDatabase& Database)
 }
 
 
-#ifdef CADKERNEL_DEV
-FInfoEntity& TTopologicalLink<FTopologicalVertex>::GetInfo(FInfoEntity& Info) const
-{
-	return FEntity::GetInfo(Info)
-		.Add(TEXT("active Entity"), ActiveEntity)
-		.Add(TEXT("twin Entities"), TwinEntities);
-}
-
-FInfoEntity& FVertexLink::GetInfo(FInfoEntity& Info) const
-{
-	return TTopologicalLink<FTopologicalVertex>::GetInfo(Info)
-		.Add(TEXT("barycenter"), Barycenter);
-}
-#endif
-
 } // namespace UE::CADKernel
diff --git a/Engine/Source/Runtime/Datasmith/CADKernel/Private/CADKernel/Topo/Topomaker.cpp b/Engine/Source/Runtime/Datasmith/CADKernel/Private/CADKernel/Topo/Topomaker.cpp
--- a/Engine/Source/Runtime/Datasmith/CADKernel/Private/CADKernel/Topo/Topomaker.cpp
+++ b/Engine/Source/Runtime/Datasmith/CADKernel/Private/CADKernel/Topo/Topomaker.cpp
@@ -821,7 +821,7 @@ void FTopomaker::MergeUnconnectedSuccessiveEdges()
 			{
 				TSharedRef<FTopologicalVertex> StartVertex = Candidates[0].Direction == EOrientation::Front ? Candidates[0].Entity->GetStartVertex() : Candidates[0].Entity->GetEndVertex();
 				TSharedRef<FTopologicalVertex> EndVertex = Candidates.Last().Direction == EOrientation::Front ? Candidates.Last().Entity->GetEndVertex() : Candidates.Last().Entity->GetStartVertex();
-				FTopologicalEdge::CreateEdgeByMergingEdges(Candidates, StartVertex, EndVertex);
+				FTopologicalEdge::CreateEdgeByMergingEdges(LargeEdgeLengthTolerance * 10., Candidates, StartVertex, EndVertex);
 			}
 		}
 	}
@@ -891,80 +891,6 @@ void SplitVertexLinkByShell(FTopologicalVertex* Vertex)
 	}
 }
 
-void SplitEdgeLinkByShell(FTopologicalEdge& Edge)
-{
-	TMap<const FShell*, TArray<FTopologicalEdge*>> ShellToEdges;
-
-	const TArray<FTopologicalEdge*>& Twins = Edge.GetTwinEntities();
-	ShellToEdges.Reserve(Twins.Num());
-
-	for (FTopologicalEdge* TwinEdge : Twins)
-	{
-		const FShell* OtherShell = (const FShell*)TwinEdge->GetLoop()->GetFace()->GetHost();
-		TArray<FTopologicalEdge*>& Edges = ShellToEdges.FindOrAdd(OtherShell);
-		Edges.Add(TwinEdge);
-	}
-
-	Edge.UnlinkTwinEntities();
-
-	for (TPair<const FShell*, TArray<FTopologicalEdge*>>& Pair : ShellToEdges)
-	{
-		TArray<FTopologicalEdge*>& Edges = Pair.Value;
-		if (Edges.Num() > 1)
-		{
-			FTopologicalEdge* FirstEdge = Edges[0];
-			for (FTopologicalEdge* TwinEdge : Edges)
-			{
-				FirstEdge->Link(*TwinEdge);
-			}
-		}
-	}
-}
-
-void FTopomaker::UnlinkEdgesOfSeparateShells()
-{
-	// Mandatory: UnlinkNonManifoldVertex has to be call before
-	TArray<FTopologicalEdge*> Edges;
-	Edges.Reserve(Faces.Num() * 10);
-	for (TSharedPtr<FTopologicalFace> Face : Faces)
-	{
-		if (Face->IsDeletedOrDegenerated())
-		{
-			continue;
-		}
-
-		for (TSharedPtr<FTopologicalLoop> Loop : Face->GetLoops())
-		{
-			for (const FOrientedEdge& OrientedEdge : Loop->GetEdges())
-			{
-				FTopologicalEdge& Edge = *OrientedEdge.Entity->GetLinkActiveEdge();
-				if (!Edge.IsProcessed())
-				{
-					bool bNeedToSplit = false;
-					Edge.SetProcessedMarker();
-					Edges.Add(&Edge);
-
-					const FShell* Shell = (FShell*)Edge.GetLoop()->GetFace()->GetHost();
-					for (FTopologicalEdge* TwinEdge : Edge.GetTwinEntities())
-					{
-						const FShell* TwinEdgeShell = (FShell*)TwinEdge->GetLoop()->GetFace()->GetHost();
-						if (Shell != TwinEdgeShell)
-						{
-							SplitEdgeLinkByShell(Edge);
-							break;
-						}
-					}
-				}
-			}
-		}
-	}
-	for (FTopologicalEdge* Edge : Edges)
-	{
-		Edge->ResetMarker1();
-	}
-}
-
-
 void FTopomaker::UnlinkNonManifoldVertex()
 {
 	TArray<FTopologicalVertex*> Vertices;
@@ -1313,7 +1239,6 @@ void FTopomaker::SplitIntoConnectedShells()
 	}
 
 	UnlinkNonManifoldVertex();
-	UnlinkEdgesOfSeparateShells();
 
 	RemoveEmptyShells();
 
@@ -1510,4 +1435,33 @@ void FTopomaker::ResetMarkersOfFaces()
 	}
 }
 
+
+void FTopomaker::UnlinkFromOther()
+{
+	TArray<FTopologicalVertex*> VerticesToLink;
+	for (FShell* Shell : Shells)
+	{
+		VerticesToLink.Reset();
+		Shell->UnlinkFromOther(VerticesToLink);
+		TopomakerTools::MergeCoincidentVertices(VerticesToLink, Tolerance);
+	}
+
+	{
+		VerticesToLink.Reset();
+		TArray<FTopologicalFace*> FacePtrs;
+		for (TSharedPtr<FTopologicalFace> Face : Faces)
+		{
+			if (!Face->IsDeleted())
+			{
+				FacePtrs.Add(Face.Get());
+			}
+		}
+
+		ShellTools::UnlinkFromOther(FacePtrs, VerticesToLink);
+		TopomakerTools::MergeCoincidentVertices(VerticesToLink, Tolerance);
+	}
+
+	UnlinkNonManifoldVertex();
+}
+
 }
diff --git a/Engine/Source/Runtime/Datasmith/CADKernel/Public/CADKernel/Topo/Shell.h b/Engine/Source/Runtime/Datasmith/CADKernel/Public/CADKernel/Topo/Shell.h
--- a/Engine/Source/Runtime/Datasmith/CADKernel/Public/CADKernel/Topo/Shell.h
+++ b/Engine/Source/Runtime/Datasmith/CADKernel/Public/CADKernel/Topo/Shell.h
@@ -21,6 +21,11 @@ class FTopologicalFace;
 class FTopologyReport;
 struct FFaceSubset;
 
+namespace ShellTools
+{
+void UnlinkFromOther(TArray<FTopologicalFace*>& Faces, TArray<FTopologicalVertex*>& VerticesToLink);
+}
+
 class CADKERNEL_API FOrientedFace : public TOrientedEntity<FTopologicalFace>
 {
 public:
@@ -138,6 +143,13 @@ public:
 
 	void CheckTopology(TArray<FFaceSubset>& Subshells);
 
+	void UnlinkFromOther(TArray<FTopologicalVertex*>& OutVerticesToLink)
+	{
+		TArray<FTopologicalFace*> Faces;
+		GetFaces(Faces);
+		ShellTools::UnlinkFromOther(Faces, OutVerticesToLink);
+	}
+
 #ifdef CADKERNEL_DEV
 	virtual void FillTopologyReport(FTopologyReport& Report) const override;
 #endif
diff --git a/Engine/Source/Runtime/Datasmith/CADKernel/Public/CADKernel/Topo/TopologicalEdge.h b/Engine/Source/Runtime/Datasmith/CADKernel/Public/CADKernel/Topo/TopologicalEdge.h
--- a/Engine/Source/Runtime/Datasmith/CADKernel/Public/CADKernel/Topo/TopologicalEdge.h
+++ b/Engine/Source/Runtime/Datasmith/CADKernel/Public/CADKernel/Topo/TopologicalEdge.h
@@ -248,6 +248,14 @@ public:
 	 */
 	void Link(FTopologicalEdge& OtherEdge);
 
+	void Disjoin();
+
+	void Unlink()
+	{
+		Disjoin();
+	}
+
+
 	TSharedRef<const FTopologicalEdge> GetLinkActiveEdge() const
 	{
 		return StaticCastSharedRef<const FTopologicalEdge>(GetLinkActiveEntity());
@@ -864,7 +872,7 @@ public:
 	 * The merged edges are deleted
 	 * @return TSharedPtr<FTopologicalEdge>() if failed
 	 */
-	static TSharedPtr<FTopologicalEdge> CreateEdgeByMergingEdges(TArray<FOrientedEdge>& Edges, const TSharedRef<FTopologicalVertex> StartVertex, const TSharedRef<FTopologicalVertex> EndVertex);
+	static TSharedPtr<FTopologicalEdge> CreateEdgeByMergingEdges(const double SmallEdgeTolerance, TArray<FOrientedEdge>& Edges, const TSharedRef<FTopologicalVertex>& StartVertex, const TSharedRef<FTopologicalVertex>& EndVertex);
 };
 
 struct CADKERNEL_API FEdge2DProperties
diff --git a/Engine/Source/Runtime/Datasmith/CADKernel/Public/CADKernel/Topo/Topomaker.h b/Engine/Source/Runtime/Datasmith/CADKernel/Public/CADKernel/Topo/Topomaker.h
--- a/Engine/Source/Runtime/Datasmith/CADKernel/Public/CADKernel/Topo/Topomaker.h
+++ b/Engine/Source/Runtime/Datasmith/CADKernel/Public/CADKernel/Topo/Topomaker.h
@@ -57,6 +57,17 @@ static ESewOption GetFromOptions(bool bGStitchingForceSew, bool bGStitchingRemov
 
 }
 
+namespace TopomakerTools
+{
+
+/**
+ * Merge Border Vertices with other vertices.
+ * @param Vertices: the initial array of active vertices to process, this array is updated at the end of the process
+ */
+void MergeCoincidentVertices(TArray<FTopologicalVertex*>& VerticesToMerge, double Tolerance);
+
+}
+
 struct FTopomakerOptions
 {
 	ESewOption SewOptions;
@@ -113,6 +124,7 @@ public:
 		ThinFaceWidth = Tolerance * ForceJoinFactor;
 	}
 
+
 	void Sew();
 
 	/**
@@ -130,7 +142,7 @@ public:
 	/**
 	 * Mandatory: UnlinkNonManifoldVertex has to be call before
 	 */
-	void UnlinkEdgesOfSeparateShells();
+	void UnlinkFromOther();
 
 	void RemoveThinFaces();
 
