[CADKernel] Improve meshing strategy in thin zone, Improve Sew algorithm

diff --git a/Engine/Source/Runtime/Datasmith/CADKernel/Private/CADKernel/Mesh/Meshers/CycleTriangulator.cpp b/Engine/Source/Runtime/Datasmith/CADKernel/Private/CADKernel/Mesh/Meshers/CycleTriangulator.cpp
--- a/Engine/Source/Runtime/Datasmith/CADKernel/Private/CADKernel/Mesh/Meshers/CycleTriangulator.cpp
+++ b/Engine/Source/Runtime/Datasmith/CADKernel/Private/CADKernel/Mesh/Meshers/CycleTriangulator.cpp
@@ -275,7 +275,7 @@ bool FCycleTriangulator::FindTheBestAcuteTriangle()
 
 			Grid.DisplayIsoSegment(EGridSpace::UniformScaled, *NewSegmentStart, *NextStart, 0, EVisuProperty::RedCurve);
 			Grid.DisplayIsoSegment(EGridSpace::UniformScaled, *NewSegmentEnd, *PrevEnd, 0, EVisuProperty::BlueCurve);
-			if(CandidateTriangleIndex == 42)
+			if (CandidateTriangleIndex == 0)
 			{
 				Wait();
 			}
@@ -306,8 +306,8 @@ bool FCycleTriangulator::FindTheBestAcuteTriangle()
 
 		GetCandidateTriangleNodes(Candidate.Key);
 
-		if(IsCycleInsideCandidate(*CandidatePoint, *StartPoint, *EndPoint, *NextPoint, TransformIntoCounterClockwise)
-		 || IsCycleInsideCandidate(*CandidatePoint, *EndPoint, *StartPoint, *PrevPoint, TransformIntoClockwise))
+		if (IsCycleInsideCandidate(*CandidatePoint, *StartPoint, *EndPoint, *NextPoint, TransformIntoCounterClockwise)
+			|| IsCycleInsideCandidate(*CandidatePoint, *EndPoint, *StartPoint, *PrevPoint, TransformIntoClockwise))
 		{
 			continue;
 		}
@@ -322,7 +322,7 @@ bool FCycleTriangulator::FindTheBestAcuteTriangle()
 		{
 			continue;
 		}
-	
+
 		bAcuteTriangle = true;
 		SetFinalCandidate();
 		return true;
@@ -334,7 +334,7 @@ bool FCycleTriangulator::FindTheBestAcuteTriangle()
 	int32 AcuteNodeIndex = -1;
 	constexpr double MaxSlopeToBeSelected = Slope::ThreeQuaterPiSlope;
 
-	for (;  false && NIndex < SubCycleNodeCount; ++NIndex)
+	for (; false && NIndex < SubCycleNodeCount; ++NIndex)
 	{
 		const TPair<int, double>& Candidate = VertexIndexToSlopes[NIndex];
 
@@ -344,7 +344,7 @@ bool FCycleTriangulator::FindTheBestAcuteTriangle()
 		}
 
 		GetCandidateTriangleNodes(Candidate.Key);
-		
+
 		// Check that the cycle is not inside the selected triangle
 		if (IsCycleInsideCandidate(*CandidatePoint, *StartPoint, *EndPoint, *NextPoint, TransformIntoCounterClockwise)
 			|| IsCycleInsideCandidate(*CandidatePoint, *EndPoint, *StartPoint, *PrevPoint, TransformIntoClockwise))
@@ -364,12 +364,12 @@ bool FCycleTriangulator::FindTheBestAcuteTriangle()
 				if (Length1 > MaxEdgeSquareLength)
 				{
 #ifdef DEBUG_FIND_THE_CYCLE_TO_MESH
-				if (Grid.bDisplay)
-				{
-					F3DDebugSession W(FString::Printf(TEXT("cancel due to Seg Length %f vs %f"), sqrt(Length1), sqrt(MaxEdgeSquareLength)));
-					Grid.DisplayIsoSegment(EGridSpace::UniformScaled, *NewSegmentStart, *CandidateNode, 0, EVisuProperty::YellowCurve);
-					Wait();
-				}
+					if (Grid.bDisplay)
+					{
+						F3DDebugSession W(FString::Printf(TEXT("cancel due to Seg Length %f vs %f"), sqrt(Length1), sqrt(MaxEdgeSquareLength)));
+						Grid.DisplayIsoSegment(EGridSpace::UniformScaled, *NewSegmentStart, *CandidateNode, 0, EVisuProperty::YellowCurve);
+						Wait();
+					}
 #endif
 					continue;
 				}
@@ -379,12 +379,12 @@ bool FCycleTriangulator::FindTheBestAcuteTriangle()
 				if (Length2 > MaxEdgeSquareLength)
 				{
 #ifdef DEBUG_FIND_THE_CYCLE_TO_MESH
-				if (Grid.bDisplay)
-				{
-					F3DDebugSession W(FString::Printf(TEXT("cancel due to Seg Length %f vs %f"), sqrt(Length2), sqrt(MaxEdgeSquareLength)));
-					Grid.DisplayIsoSegment(EGridSpace::UniformScaled, *NewSegmentEnd, *CandidateNode, 0, EVisuProperty::YellowCurve);
-					Wait();
-				}
+					if (Grid.bDisplay)
+					{
+						F3DDebugSession W(FString::Printf(TEXT("cancel due to Seg Length %f vs %f"), sqrt(Length2), sqrt(MaxEdgeSquareLength)));
+						Grid.DisplayIsoSegment(EGridSpace::UniformScaled, *NewSegmentEnd, *CandidateNode, 0, EVisuProperty::YellowCurve);
+						Wait();
+					}
 #endif
 					continue;
 				}
@@ -396,7 +396,7 @@ bool FCycleTriangulator::FindTheBestAcuteTriangle()
 		}
 	}
 
-	if(bAcuteTriangle)
+	if (bAcuteTriangle)
 	{
 		GetCandidateTriangleNodes(AcuteNodeIndex);
 		SetFinalCandidate();
@@ -423,12 +423,13 @@ bool FCycleTriangulator::FindCandidateNodes(int32 StartIndex)
 
 	for (int32 Index = NextIndex(FirstSideEndIndex); ; Index = NextIndex(Index))
 	{
-		FIsoNode* Node = SubCycleNodes[Index];
-		if (Node == FirstSideStartNode)
+		if (Index == FirstSideStartIndex)
 		{
 			break;
 		}
 
+		FIsoNode* Node = SubCycleNodes[Index];
+
 		const FPoint2D& NodePoint2D = Node->Get2DPoint(EGridSpace::UniformScaled, Grid);
 		double SlopeAtStartNode = GetSlopeAtStartNode(StartPoint2D, NodePoint2D, StartReferenceSlope);
 		double SlopeAtEndNode = GetSlopeAtEndNode(EndPoint2D, NodePoint2D, EndReferenceSlope);
@@ -445,7 +446,10 @@ bool FCycleTriangulator::FindCandidateNodes(int32 StartIndex)
 	F3DDebugSession A(Grid.bDisplay, TEXT("Find Candidate Nodes"));
 	if (Grid.bDisplay)
 	{
+		F3DDebugSession A(Grid.bDisplay, TEXT("StartSeg"));
 		Grid.DisplayIsoSegment(EGridSpace::UniformScaled, *FirstSideStartNode, *FirstSideEndNode, 0, EVisuProperty::YellowCurve);
+		Grid.DisplayIsoNode(EGridSpace::UniformScaled, *FirstSideStartNode, 0, EVisuProperty::BluePoint);
+		Grid.DisplayIsoNode(EGridSpace::UniformScaled, *FirstSideEndNode, 0, EVisuProperty::RedPoint);
 	}
 #endif
 
@@ -453,15 +457,26 @@ bool FCycleTriangulator::FindCandidateNodes(int32 StartIndex)
 	const double StartMaxSlope = FMath::Min(CandidateNodes.Last().SlopeAtStartNode + Slope::OneDegree, MaxSlope);
 	const double EndMaxSlope = FMath::Min(CandidateNodes[0].SlopeAtEndNode + Slope::OneDegree, MaxSlope);
 
-	if (CandidateNodes.Last().SlopeAtEndNode < EndMaxSlope)
+	if (CandidateNodes.Last().SlopeAtStartNode < StartMaxSlope)
 	{
 		CandidateNodes.Last().bIsValid = true;
+		ExtremityCandidateNodes[0] = &CandidateNodes.Last();
 	}
-	if (CandidateNodes[0].SlopeAtStartNode < StartMaxSlope)
+	if (CandidateNodes[0].SlopeAtEndNode < EndMaxSlope)
 	{
 		CandidateNodes[0].bIsValid = true;
+		ExtremityCandidateNodes[1] = &CandidateNodes[0];
 	}
 
+#ifdef DEBUG_FIND_CANDIDATE_NODES
+	if (Grid.bDisplay)
+	{
+		F3DDebugSession A(Grid.bDisplay, TEXT("Start End Points"));
+		Grid.DisplayIsoNode(EGridSpace::UniformScaled, *CandidateNodes.Last().Node, 0, EVisuProperty::RedPoint);
+		Grid.DisplayIsoNode(EGridSpace::UniformScaled, *CandidateNodes[0].Node, 0, EVisuProperty::BluePoint);
+	}
+#endif
+
 	for (FCandidateNode& CNode : CandidateNodes)
 	{
 		if (CNode.bIsValid)
@@ -489,12 +504,13 @@ bool FCycleTriangulator::FindCandidateNodes(int32 StartIndex)
 #ifdef DEBUG_FIND_CANDIDATE_NODES
 	if (Grid.bDisplay)
 	{
+		F3DDebugSession A(Grid.bDisplay, TEXT("Points"));
 		for (FCandidateNode& CNode : CandidateNodes)
 		{
 			//F3DDebugSession W(TEXT("Next Segment"));
 			Grid.DisplayIsoNode(EGridSpace::UniformScaled, *CNode.Node, 0, CNode.bIsValid ? EVisuProperty::BluePoint : EVisuProperty::PinkPoint);
 		}
-		Wait();
+		//Wait();
 	}
 #endif
 
@@ -506,16 +522,14 @@ bool FCycleTriangulator::FindTheBestCandidateNode()
 	constexpr double MinSlopeToNotBeAligned = Slope::Epsilon;
 
 	IntersectionCountAllowed = 0;
-	int32 CandidateNodeIndex = -1;
 
 	double MinCriteria = HUGE_VALUE;
 	double CandidateSlopeAtStartNode = Slope::PiSlope - MinSlopeToNotBeAligned;
 	double CandidateSlopeAtEndNode = Slope::PiSlope - MinSlopeToNotBeAligned;
-	FCandidateNode* Candidate = nullptr;
+	FCandidateNode* CandidateNode = nullptr;
 	double CandidateEndSquareDistance = HUGE_VALUE;
 	double CandidateStartSquareDistance = HUGE_VALUE;
 
-
 	const FPoint2D& StartPoint2D = FirstSideStartNode->Get2DPoint(EGridSpace::UniformScaled, Grid);
 	const FPoint2D& EndPoint2D = FirstSideEndNode->Get2DPoint(EGridSpace::UniformScaled, Grid);
 
@@ -527,7 +541,7 @@ bool FCycleTriangulator::FindTheBestCandidateNode()
 	{
 		F3DDebugSession _(FString::Printf(TEXT("Start Segment %d"), BestCandidateNode));
 		Grid.DisplayIsoSegment(EGridSpace::UniformScaled, *FirstSideStartNode, *FirstSideEndNode, 1, EVisuProperty::RedCurve);
-		Wait();
+		//Wait();
 	}
 #endif
 
@@ -535,9 +549,28 @@ bool FCycleTriangulator::FindTheBestCandidateNode()
 	{
 		if (NodeIntersection.Value > IntersectionCountAllowed)
 		{
-			if (CandidateNodeIndex >= 0)
+			if (CandidateNode)
 			{
-				break;
+				// if the candidate triangle is degenerated, allow intersection with iso to find a better candidate
+
+				// check if the angle of the candidate triangle are not degenerated
+				double MinSlope;
+				double MaxSlope;
+				GetMinMax(CandidateSlopeAtStartNode, CandidateSlopeAtEndNode, MinSlope, MaxSlope);
+				constexpr double PiMinus5 = Slope::PiSlope - Slope::FiveDegree;
+				const bool bIsDegeneratedDueSlope = MinSlope < Slope::FiveDegree && MaxSlope > PiMinus5;
+
+				// check if the length of the sides are not disproportionate min side minus than 50 times max length 
+				const double MinSquareLenth = FMath::Min(CandidateEndSquareDistance, CandidateStartSquareDistance);
+				const double SegmentSquareLength2500 = StartPoint2D.SquareDistance(EndPoint2D) * 2500;
+				const bool bIsDegeneratedDueToLength = SegmentSquareLength2500 < MinSquareLenth;
+
+				const bool bIsDegenerated = bIsDegeneratedDueSlope || bIsDegeneratedDueToLength;
+
+				if (!bIsDegenerated)
+				{
+					break;
+				}
 			}
 			IntersectionCountAllowed = NodeIntersection.Value;
 		}
@@ -555,7 +588,7 @@ bool FCycleTriangulator::FindTheBestCandidateNode()
 			F3DDebugSession _(TEXT("Candidate triangle"));
 			Grid.DisplayIsoSegment(EGridSpace::UniformScaled, *FirstSideStartNode, *Node, 1, EVisuProperty::RedCurve);
 			Grid.DisplayIsoSegment(EGridSpace::UniformScaled, *FirstSideEndNode, *Node, 2, EVisuProperty::RedCurve);
-			Wait();
+			//Wait();
 		}
 #endif
 
@@ -606,7 +639,7 @@ bool FCycleTriangulator::FindTheBestCandidateNode()
 			}
 
 			MinCriteria = PointCriteria;
-			CandidateNodeIndex = CNode->Index;
+			CandidateNode = CNode;
 			CandidateSlopeAtStartNode = CNode->SlopeAtStartNode;
 			CandidateSlopeAtEndNode = CNode->SlopeAtEndNode;
 			CandidateEndSquareDistance = CEndSquareDistance;
@@ -614,14 +647,128 @@ bool FCycleTriangulator::FindTheBestCandidateNode()
 		}
 	}
 
+	if (!CandidateNode)
+	{
+		return false;
+	}
+
+#ifdef DEBUG_FIND_THE_BEST_CANDIDATE_NODE
+	static int32 TriangleIndex2 = 0;
+	if (Grid.bDisplay && CandidateNode)
+	{
+		FIsoNode* Node = CandidateNode->Node;
+		F3DDebugSession _(FString::Printf(TEXT("Selected triangle %d"), ++TriangleIndex2));
+		Grid.DisplayIsoSegment(EGridSpace::UniformScaled, *FirstSideStartNode, *Node, 1, EVisuProperty::BlueCurve);
+		Grid.DisplayIsoSegment(EGridSpace::UniformScaled, *FirstSideEndNode, *Node, 2, EVisuProperty::BlueCurve);
+		//Wait();
+	}
+#endif
+
+#ifdef DEBUG_OPIMIZE_FIND_BEST_CANDIDATE
+	TFunction<void() > DisplayContext = [&]()
+	{
+		if (Grid.bDisplay)
+		{
+			F3DDebugSession A(TEXT("Opimizer"));
+			if (Grid.bDisplay)
+			{
+				F3DDebugSession B(TEXT("StartSeg"));
+				Grid.DisplayIsoSegment(EGridSpace::UniformScaled, *FirstSideStartNode, *FirstSideEndNode, 0, EVisuProperty::YellowCurve);
+			}
+
+			{
+				F3DDebugSession C(TEXT("Points"));
+				for (FCandidateNode& CNode : CandidateNodes)
+				{
+					Grid.DisplayIsoNode(EGridSpace::UniformScaled, *CNode.Node, 0, CNode.bIsValid ? EVisuProperty::BluePoint : EVisuProperty::PinkPoint);
+				}
+			}
+
+			{
+				F3DDebugSession _(FString::Printf(TEXT("Selected triangle %d"), ++TriangleIndex2));
+				FIsoNode* Node = CandidateNode->Node;
+				Grid.DisplayIsoSegment(EGridSpace::UniformScaled, *FirstSideStartNode, *Node, 1, EVisuProperty::BlueCurve);
+				Grid.DisplayIsoSegment(EGridSpace::UniformScaled, *FirstSideEndNode, *Node, 2, EVisuProperty::BlueCurve);
+			}
+		}
+	};
+#endif
+
+	TFunction<void(FCandidateNode*, const double, const FPoint2D&, double, double) > CompareCandidateAtSegmentNode = [&CandidateNode, &Grid = Grid](FCandidateNode* CandidatePoint, const double CloseLimit, const FPoint2D& SegmentPoint2D, double CandidateToSegmentSquareDistance, double SlopeDiffAtSegmentExtremity)
+	{
+		const FPoint2D& NodePoint2D = CandidatePoint->Node->Get2DPoint(EGridSpace::UniformScaled, Grid);
+		const double CFirstSquareDistance = SegmentPoint2D.SquareDistance(NodePoint2D) * 2.;
+
+		if (SlopeDiffAtSegmentExtremity < CloseLimit && CFirstSquareDistance < CandidateToSegmentSquareDistance)
+		{
+			CandidateNode = CandidatePoint;
+		}
+	};
+
+	// Optimize selection step: if the selected triangle is close to one of the next segment, check if a triangle from it is not better 
+	if (ExtremityCandidateNodes[0] != CandidateNode && ExtremityCandidateNodes[1] != CandidateNode)
+	{
+		if (ExtremityCandidateNodes[0] && ExtremityCandidateNodes[1])
+		{
+			const double SlopeDiffAtStart = ExtremityCandidateNodes[0]->SlopeAtStartNode - CandidateNode->SlopeAtStartNode;
+			const double SlopeDiffAtEnd = ExtremityCandidateNodes[1]->SlopeAtEndNode - CandidateNode->SlopeAtEndNode;
+
+			if (SlopeDiffAtStart < SlopeDiffAtEnd)
+			{
+				const FNodeIntersectionCount* ExtremityIntersectionCount = NodeToIntersection.FindByPredicate([Extremity = ExtremityCandidateNodes[0]](const FNodeIntersectionCount& Node) { return Node.Key == Extremity; });
+				if (IntersectionCountAllowed >= ExtremityIntersectionCount->Value)
+				{
+					CompareCandidateAtSegmentNode(ExtremityCandidateNodes[0], ExtremityCandidateNodes[0]->SlopeAtStartNode * AThird, StartPoint2D, CandidateStartSquareDistance, SlopeDiffAtStart);
+					printf("");
+				}
+			}
+			else
+			{
+				const FNodeIntersectionCount* ExtremityIntersectionCount = NodeToIntersection.FindByPredicate([Extremity = ExtremityCandidateNodes[1]](const FNodeIntersectionCount& Node) { return Node.Key == Extremity; });
+				if (IntersectionCountAllowed >= ExtremityIntersectionCount->Value)
+				{
+					CompareCandidateAtSegmentNode(ExtremityCandidateNodes[1], ExtremityCandidateNodes[1]->SlopeAtEndNode * AThird, EndPoint2D, CandidateEndSquareDistance, SlopeDiffAtEnd);
+				}
+				printf("");
+			}
+		}
+		else if (ExtremityCandidateNodes[0])
+		{
+			const FNodeIntersectionCount* ExtremityIntersectionCount = NodeToIntersection.FindByPredicate([Extremity = ExtremityCandidateNodes[0]](const FNodeIntersectionCount& Node) { return Node.Key == Extremity; });
+			if (IntersectionCountAllowed >= ExtremityIntersectionCount->Value)
+			{
+				const double SlopeDiffAtStart = ExtremityCandidateNodes[0]->SlopeAtStartNode - CandidateNode->SlopeAtStartNode;
+				if (SlopeDiffAtStart < Slope::FifteenDegree)
+				{
+					CompareCandidateAtSegmentNode(ExtremityCandidateNodes[0], ExtremityCandidateNodes[0]->SlopeAtStartNode * AThird, StartPoint2D, CandidateStartSquareDistance, SlopeDiffAtStart);
+					printf("");
+				}
+			}
+		}
+		else if (ExtremityCandidateNodes[1])
+		{
+			//FLoopNode const* const* StartNode = NodesOfLoop.FindByPredicate([](const FLoopNode* Node) { return !Node->IsDelete(); });
+			const FNodeIntersectionCount* ExtremityIntersectionCount = NodeToIntersection.FindByPredicate([Extremity = ExtremityCandidateNodes[1]](const FNodeIntersectionCount& Node) { return Node.Key == Extremity; });
+			if (IntersectionCountAllowed >= ExtremityIntersectionCount->Value)
+			{
+				const double SlopeDiffAtEnd = ExtremityCandidateNodes[1]->SlopeAtEndNode - CandidateNode->SlopeAtEndNode;
+				if (SlopeDiffAtEnd < Slope::FifteenDegree)
+				{
+					CompareCandidateAtSegmentNode(ExtremityCandidateNodes[1], ExtremityCandidateNodes[1]->SlopeAtEndNode * AThird, EndPoint2D, CandidateEndSquareDistance, SlopeDiffAtEnd);
+					printf("");
+				}
+			}
+		}
+	}
+
 #ifdef DEBUG_FIND_THE_BEST_CANDIDATE_NODE
 	static int32 TriangleIndex = 0;
-	if (Grid.bDisplay && CandidateNodeIndex >= 0)
+	if (Grid.bDisplay && CandidateNode)
 	{
-		FIsoNode* Node = SubCycleNodes[CandidateNodeIndex];
+		FIsoNode* Node = CandidateNode->Node;
 		F3DDebugSession _(FString::Printf(TEXT("Selected triangle %d"), ++TriangleIndex));
-		Grid.DisplayIsoSegment(EGridSpace::UniformScaled, *FirstSideStartNode, *Node, 1, EVisuProperty::RedCurve);
-		Grid.DisplayIsoSegment(EGridSpace::UniformScaled, *FirstSideEndNode, *Node, 2, EVisuProperty::RedCurve);
+		Grid.DisplayIsoSegment(EGridSpace::UniformScaled, *FirstSideStartNode, *Node, 1, EVisuProperty::YellowCurve);
+		Grid.DisplayIsoSegment(EGridSpace::UniformScaled, *FirstSideEndNode, *Node, 2, EVisuProperty::YellowCurve);
 		if (TriangleIndex == 1000)
 		{
 			Wait();
@@ -629,10 +776,10 @@ bool FCycleTriangulator::FindTheBestCandidateNode()
 	}
 #endif
 
-	if (CandidateNodeIndex >= 0)
+	if (CandidateNode)
 	{
-		TriangleThirdIndex = CandidateNodeIndex;
-		TriangleThirdNode = SubCycleNodes[CandidateNodeIndex];
+		TriangleThirdIndex = CandidateNode->Index;
+		TriangleThirdNode = CandidateNode->Node;
 		return true;
 	}
 	return false;
@@ -646,7 +793,7 @@ void FCycleTriangulator::ValidateAddNodesAccordingSlopeWithSide(FAdditionalIso&
 	const FPoint2D& EndPoint = Side.EndNode->Get2DPoint(EGridSpace::UniformScaled, Grid);
 	const double SlopeStartEnd = ComputeSlope(StartPoint, EndPoint);
 
-	for(int32 Index = 0; Index < 2; ++Index)
+	for (int32 Index = 0; Index < 2; ++Index)
 	{
 		const FIsoNode* Candidate = Side.Nodes[Index];
 		if (!Candidate)
@@ -671,7 +818,7 @@ void FCycleTriangulator::ValidateAddNodesAccordingSlopeWithSide(FAdditionalIso&
 		}
 
 #ifdef DEBUG_COMPUTE_SLOTE_CRITERION
-		if(Grid.bDisplay)
+		if (Grid.bDisplay)
 		{
 			F3DDebugSession _(FString::Printf(TEXT("Slope %f"), SlopeAtStart));
 			Grid.DisplayIsoSegment(EGridSpace::UniformScaled, *Side.StartNode, *Candidate, 1, EVisuProperty::BlueCurve);
@@ -680,7 +827,7 @@ void FCycleTriangulator::ValidateAddNodesAccordingSlopeWithSide(FAdditionalIso&
 		{
 			F3DDebugSession _(FString::Printf(TEXT("Slope %f"), SlopeAtEnd));
 			Grid.DisplayIsoSegment(EGridSpace::UniformScaled, *Side.EndNode, *Candidate, 1, EVisuProperty::BlueCurve);
-		    //Wait();
+			//Wait();
 		}
 #endif
 
@@ -696,6 +843,12 @@ void FCycleTriangulator::ValidateAddNodesAccordingSlopeWithSide(FAdditionalIso&
 
 void FCycleTriangulator::SelectFinalNodes(FAdditionalIso& Side1, FAdditionalIso& Side2)
 {
+	// Check if complementary node is candidate triangle vertex
+	{
+		Side1.RemoveCandidateIfPresent(Side2.EndNode);
+		Side2.RemoveCandidateIfPresent(Side1.StartNode);
+	}
+
 	if (!Side1.bForceNodes && !Side2.bForceNodes)
 	{
 		const int32 Side1Count = Side1.CandidateNodeCount();
@@ -706,15 +859,15 @@ void FCycleTriangulator::SelectFinalNodes(FAdditionalIso& Side1, FAdditionalIso&
 			ComputeSideCandidateEquilateralCriteria(Side1);
 			ComputeSideCandidateEquilateralCriteria(Side2);
 			double Criteria[4] = { Slope::NullSlope, Slope::NullSlope, Slope::NullSlope, Slope::NullSlope };
-			Criteria[0] = Side1.EquilateralCriteria[0]; 
+			Criteria[0] = Side1.EquilateralCriteria[0];
 			Criteria[1] = Side1.EquilateralCriteria[1];
 			Criteria[2] = Side2.EquilateralCriteria[0];
 			Criteria[3] = Side2.EquilateralCriteria[1];
 
-			int32 Indices[4] = {0, 1, 2, 3};
+			int32 Indices[4] = { 0, 1, 2, 3 };
 			Algo::Sort(Indices, [Criteria](int32 A, int32 B) { return Criteria[A] < Criteria[B]; });
 
-			for(int32 Index = 0; Index < 4 && CandidateCount != 2; ++Index)
+			for (int32 Index = 0; Index < 4 && CandidateCount != 2; ++Index)
 			{
 				if (Criteria[Indices[Index]] > Slope::PiSlope)
 				{
@@ -723,7 +876,7 @@ void FCycleTriangulator::SelectFinalNodes(FAdditionalIso& Side1, FAdditionalIso&
 
 				if (Indices[Index] > 1)
 				{
-					if(!ValidComplementaryNodeOrDeleteIt(Side2, Indices[Index]-2))
+					if (!ValidComplementaryNodeOrDeleteIt(Side2, Indices[Index] - 2))
 					{
 						CandidateCount--;
 					}
@@ -736,7 +889,7 @@ void FCycleTriangulator::SelectFinalNodes(FAdditionalIso& Side1, FAdditionalIso&
 					}
 				}
 			}
-			
+
 			if (CandidateCount > 2)
 			{
 				CandidateCount = 0;
@@ -844,12 +997,12 @@ bool FCycleTriangulator::ValidComplementaryNodeOrDeleteIt(FAdditionalIso& Side,
 	}
 
 	const int32 OppositeIndex = CandidateIndex ? Side.StartIndex : Side.EndIndex;
-	const FIsoNode* OppositeNode = CandidateIndex ? Side.StartNode : Side.EndNode; 
+	const FIsoNode* OppositeNode = CandidateIndex ? Side.StartNode : Side.EndNode;
 
 	TFunction<int32(int32)> NextFunc = [&CandidateIndex, this](int32 Index)
 	{
-		return CandidateIndex ?  NextIndex(Index) : PreviousIndex(Index);
- 	};
+		return CandidateIndex ? NextIndex(Index) : PreviousIndex(Index);
+	};
 
 	const FPoint2D& CandidatePoint = Candidate->Get2DPoint(EGridSpace::UniformScaled, Grid);
 	const FPoint2D& OppositePoint = OppositeNode->Get2DPoint(EGridSpace::UniformScaled, Grid);
@@ -886,7 +1039,7 @@ bool FCycleTriangulator::ValidComplementaryNodeOrDeleteIt(FAdditionalIso& Side,
 		}
 	}
 #ifdef DEBUG_VALID_COMPLEMENTARY_NODE
-	else if(Grid.bDisplay)
+	else if (Grid.bDisplay)
 	{
 		F3DDebugSession _(TEXT("intersect cycle => candidate is deleted"));
 		Grid.DisplayIsoSegment(EGridSpace::UniformScaled, *Candidate, *OppositeNode, 1, EVisuProperty::YellowCurve);
@@ -974,11 +1127,11 @@ bool FCycleTriangulator::IsInnerSideSegmentInsideCycle(FAdditionalIso& Side)
 
 void FCycleTriangulator::ValidateComplementaryNodesWithInsideAndIntersectionsCriteria(FAdditionalIso& Side)
 {
-	for(int32 Index = 0; Index < 2; ++Index)
+	for (int32 Index = 0; Index < 2; ++Index)
 	{
 		if (Side.Nodes[Index])
 		{
-			if(InnerToOuterIsoSegmentsIntersectionTool.DoesIntersect(Index == 0 ? *Side.EndNode : *Side.StartNode, *Side.Nodes[Index]))
+			if (InnerToOuterIsoSegmentsIntersectionTool.DoesIntersect(Index == 0 ? *Side.EndNode : *Side.StartNode, *Side.Nodes[Index]))
 			{
 				Side.RemoveCandidate(Index);
 				Side.bForceNodes = false;
@@ -1178,7 +1331,7 @@ bool FCycleTriangulator::BuildSmallPolygon(TArray<FIsoNode*>& CandidatNodes, boo
 	case 1:
 	case 2:
 		return false;
-	case 3: 
+	case 3:
 	{
 		return BuildTriangle(CandidatNodes);
 	}
@@ -1405,7 +1558,7 @@ void MeshPentagon(const FGrid& Grid, FIsoNode** Nodes, FFaceMesh& Mesh)
 
 	// Remove the pentagons containing the worst triangle until ValidPentagons while ValidPentagons is not empty
 	EPentagon ValidPentagons = EPentagon::All;
-	for(int32 Index = 0; Index < 10; Index++ )
+	for (int32 Index = 0; Index < 10; Index++)
 	{
 		const EPentagon Valid = TriangleToPentagon[TriangleIndices[Index]];
 		EPentagon PreviousValid = ValidPentagons;
@@ -1815,11 +1968,15 @@ void FCycleTriangulator::CleanContext()
 	SubCycleNodeCount = 0;
 	SubCycleNodes.Reset();
 	CandidateNodes.Reset();
+
+	ExtremityCandidateNodes[0] = nullptr;
+	ExtremityCandidateNodes[1] = nullptr;
+
 	PolygonNodes.Reset();
 	NodeToIntersection.Reset();
 
 	VertexIndexToSlopes.Reset();
-	
+
 	FirstSideStartIndex = -1;
 	FirstSideEndIndex = -1;
 	TriangleThirdIndex = -1;
diff --git a/Engine/Source/Runtime/Datasmith/CADKernel/Private/CADKernel/Mesh/Meshers/CycleTriangulator.h b/Engine/Source/Runtime/Datasmith/CADKernel/Private/CADKernel/Mesh/Meshers/CycleTriangulator.h
--- a/Engine/Source/Runtime/Datasmith/CADKernel/Private/CADKernel/Mesh/Meshers/CycleTriangulator.h
+++ b/Engine/Source/Runtime/Datasmith/CADKernel/Private/CADKernel/Mesh/Meshers/CycleTriangulator.h
@@ -95,6 +95,18 @@ struct FAdditionalIso
 		EquilateralCriteria[1] = Slope::TwoPiSlope;
 		bForceNodes = false;
 	}
+
+	void RemoveCandidateIfPresent(FIsoNode* ForbiddenNode)
+	{
+		for (int32 Index = 0; Index < 2; Index++)
+		{
+			if (Nodes[Index] == ForbiddenNode)
+			{
+ 				RemoveCandidate(Index);
+				bForceNodes = false;
+			}
+		}
+	}
 };
 
 using FNodeIntersectionCount = TPair<FCandidateNode*, int32>;
@@ -142,6 +154,7 @@ private:
 
 	// Candidates nodes for FindCandidateNodes
 	TArray<FCandidateNode> CandidateNodes;
+	FCandidateNode* ExtremityCandidateNodes[2] = { nullptr, nullptr };
 	TArray<FNodeIntersectionCount> NodeToIntersection;
 
 	// Final polygon nodes
diff --git a/Engine/Source/Runtime/Datasmith/CADKernel/Private/CADKernel/Mesh/Meshers/IsoTriangulator.cpp b/Engine/Source/Runtime/Datasmith/CADKernel/Private/CADKernel/Mesh/Meshers/IsoTriangulator.cpp
--- a/Engine/Source/Runtime/Datasmith/CADKernel/Private/CADKernel/Mesh/Meshers/IsoTriangulator.cpp
+++ b/Engine/Source/Runtime/Datasmith/CADKernel/Private/CADKernel/Mesh/Meshers/IsoTriangulator.cpp
@@ -952,7 +952,7 @@ void FIsoTriangulator::ConnectCellLoops()
 #ifdef DEBUG_CONNECT_CELL_LOOPS
 		F3DDebugSession _(bDisplay, FString::Printf(TEXT("Cell %d"), Cell.Id));
 		DisplayCell(Cell);
-		Wait();
+		Wait(bDisplay);
 #endif
 
 		if (Cell.LoopCells.Num())
@@ -2422,13 +2422,13 @@ void FIsoTriangulator::ConnectCellCornerToLoops(FCell& Cell)
 		}
 		if (ICell == 4)
 		{
-			// All Cell corners are not null
+			// All Cell corners are null
 			return;
 		}
 	}
 
 #ifdef DEBUG_CONNECTCELLCORNERTOINNERLOOP
-	F3DDebugSession _(TEXT("With cell corners"));
+	F3DDebugSession _(bDisplay, TEXT("With cell corners"));
 #endif
 
 	TFunction<void(FLoopCell&, FIsoInnerNode*)> FindAndTryCreateCandidateSegmentToLinkLoopToCorner = [&](FLoopCell& LoopCell, FIsoInnerNode* InnerNode)
@@ -2476,20 +2476,31 @@ void FIsoTriangulator::ConnectCellCornerToLoops(FCell& Cell)
 	{
 		if (CellNodes[ICell])
 		{
-			for (FLoopCell& LoopCell : Cell.LoopCells)
-			{
 #ifdef DEBUG_CONNECTCELLCORNERTOINNERLOOP
+			if (bDisplay)
+			{
+				F3DDebugSession _(bDisplay, TEXT("Cell corner"));
 				DisplayPoint(CellNodes[ICell]->Get2DPoint(EGridSpace::UniformScaled, Grid), EVisuProperty::GreenPoint);
+			}
 #endif
+			for (FLoopCell& LoopCell : Cell.LoopCells)
+			{
 				FindAndTryCreateCandidateSegmentToLinkLoopToCorner(LoopCell, CellNodes[ICell]);
 			}
 
+#ifdef DEBUG_CONNECTCELLCORNERTOINNERLOOP
+			if (bDisplay)
+			{
+				Grid.DisplayIsoSegments(TEXT("Cell CandidateSegments"), EGridSpace::UniformScaled, Cell.CandidateSegments, false, false, false, EVisuProperty::YellowCurve);
+				Wait();
+			}
+#endif
 		}
 	}
 
 	if(Cell.CandidateSegments.Num())
 	{
-		const bool bSelectWithCellSizeCriteria = false;
+		const bool bSelectWithCellSizeCriteria = true;
 		Cell.SelectSegmentInCandidateSegments(IsoSegmentFactory, bSelectWithCellSizeCriteria);
 	}
 }
diff --git a/Engine/Source/Runtime/Datasmith/CADKernel/Private/CADKernel/Mesh/Meshers/IsoTriangulator/IsoCell.h b/Engine/Source/Runtime/Datasmith/CADKernel/Private/CADKernel/Mesh/Meshers/IsoTriangulator/IsoCell.h
--- a/Engine/Source/Runtime/Datasmith/CADKernel/Private/CADKernel/Mesh/Meshers/IsoTriangulator/IsoCell.h
+++ b/Engine/Source/Runtime/Datasmith/CADKernel/Private/CADKernel/Mesh/Meshers/IsoTriangulator/IsoCell.h
@@ -306,7 +306,7 @@ struct FCell
 
 		Algo::Sort(LengthOfCandidateSegments, [&](const TPair<double, FIsoSegment*>& P1, const TPair < double, FIsoSegment*>& P2) { return P1.Key < P2.Key; });
 
-		// Validate all candidate segments
+		// Validate the first candidate segments
 		for (const TPair<double, FIsoSegment*>& Candidate : LengthOfCandidateSegments)
 		{
 			if (Candidate.Key > CellSquareLength)
diff --git a/Engine/Source/Runtime/Datasmith/CADKernel/Private/CADKernel/Mesh/Meshers/ParametricFaceMesher.cpp b/Engine/Source/Runtime/Datasmith/CADKernel/Private/CADKernel/Mesh/Meshers/ParametricFaceMesher.cpp
--- a/Engine/Source/Runtime/Datasmith/CADKernel/Private/CADKernel/Mesh/Meshers/ParametricFaceMesher.cpp
+++ b/Engine/Source/Runtime/Datasmith/CADKernel/Private/CADKernel/Mesh/Meshers/ParametricFaceMesher.cpp
@@ -1116,9 +1116,19 @@ void FParametricFaceMesher::Mesh(FTopologicalEdge& InEdge, bool bFinalMeshing)
 		F3DDebugSession A(bDisplay, FString::Printf(TEXT("Edge %d"), InEdge.GetId()));
 #endif
 		Display2DWithScale(InEdge);
-		TArray<FPoint> MeshVertices = InEdge.GetMesh()->GetNodeCoordinates();
-		MeshVertices.Append(InEdge.GetStartVertex()->GetMesh()->GetNodeCoordinates());
-		MeshVertices.Append(InEdge.GetEndVertex()->GetMesh()->GetNodeCoordinates());
+		TArray<FPoint> MeshVertices;
+		if (InEdge.GetMesh())
+		{
+			MeshVertices = InEdge.GetMesh()->GetNodeCoordinates();
+		}
+		if(InEdge.GetStartVertex()->GetMesh())
+		{
+			MeshVertices.Append(InEdge.GetStartVertex()->GetMesh()->GetNodeCoordinates());
+		}
+		if(InEdge.GetEndVertex()->GetMesh())
+		{
+			MeshVertices.Append(InEdge.GetEndVertex()->GetMesh()->GetNodeCoordinates());
+		}
 		if (MeshVertices.Num())
 		{
 			TArray<double> CoordinatesOfMesh;
diff --git a/Engine/Source/Runtime/Datasmith/CADKernel/Private/CADKernel/Mesh/Structure/LoopCleaner.cpp b/Engine/Source/Runtime/Datasmith/CADKernel/Private/CADKernel/Mesh/Structure/LoopCleaner.cpp
--- a/Engine/Source/Runtime/Datasmith/CADKernel/Private/CADKernel/Mesh/Structure/LoopCleaner.cpp
+++ b/Engine/Source/Runtime/Datasmith/CADKernel/Private/CADKernel/Mesh/Structure/LoopCleaner.cpp
@@ -1749,8 +1749,14 @@ void FLoopCleaner::FindLoopIntersections()
 
 	LoopSegmentsIntersectionTool.Empty(NodesOfLoopCount);
 
+	const int32 LimitOfNodesOfLoopCount = NodesOfLoopCount + 5;
 	Node = GetNext(StartNode);
-	for (; Node != StartNode; Node = GetNext(Node), ++SegmentCount);
+	for (; Node != StartNode && SegmentCount < LimitOfNodesOfLoopCount; Node = GetNext(Node), ++SegmentCount);
+
+	if (SegmentCount == LimitOfNodesOfLoopCount)
+	{
+		return;
+	}
 
 	Node = StartNode;
 	NextNode = GetNext(Node);
diff --git a/Engine/Source/Runtime/Datasmith/CADKernel/Private/CADKernel/Mesh/Structure/ThinZone2DFinder.cpp b/Engine/Source/Runtime/Datasmith/CADKernel/Private/CADKernel/Mesh/Structure/ThinZone2DFinder.cpp
--- a/Engine/Source/Runtime/Datasmith/CADKernel/Private/CADKernel/Mesh/Structure/ThinZone2DFinder.cpp
+++ b/Engine/Source/Runtime/Datasmith/CADKernel/Private/CADKernel/Mesh/Structure/ThinZone2DFinder.cpp
@@ -1356,7 +1356,7 @@ bool FThinZone2DFinder::BuildLoopSegments()
 			{
 				for (int32 PointIndex = 1; PointIndex < ScaledPoints.Num(); PointIndex++)
 				{
-					while (ScaledPoints.Num() > 1 && ScaledPoints[PointIndex - 1].Distance(ScaledPoints[PointIndex]) < GeometricTolerance)
+					while (PointIndex < ScaledPoints.Num() && ScaledPoints[PointIndex - 1].Distance(ScaledPoints[PointIndex]) < GeometricTolerance)
 					{
 						ScaledPoints.RemoveAt(PointIndex);
 						Coordinates.RemoveAt(PointIndex);
diff --git a/Engine/Source/Runtime/Datasmith/CADKernel/Private/CADKernel/Topo/FaceAnalyzer.cpp b/Engine/Source/Runtime/Datasmith/CADKernel/Private/CADKernel/Topo/FaceAnalyzer.cpp
--- a/Engine/Source/Runtime/Datasmith/CADKernel/Private/CADKernel/Topo/FaceAnalyzer.cpp
+++ b/Engine/Source/Runtime/Datasmith/CADKernel/Private/CADKernel/Topo/FaceAnalyzer.cpp
@@ -13,18 +13,49 @@ namespace UE::CADKernel
 
 void FFaceAnalyzer::FindClosedSegments(Topo::FThinFaceContext& Context)
 {
-#ifdef DEBUG_FIND_THIN_FACE
-	F3DDebugSession _(TEXT("Closed Segment"));
+#ifdef DEBUG_THIN_FACE
+	F3DDebugSession _(TEXT("Find Closed Segment"));
 #endif
+	{
+		bool bEdgeIsConnectedToAnotherEdgeOfTheFace = false;
+		const FTopologicalEdge* Edge = nullptr;
+		for (Topo::FEdgeSegment* Segment : Context.LoopSegments)
+		{
+			if (Segment->GetEdge() != Edge)
+			{
+				Edge = Segment->GetEdge();
+				bEdgeIsConnectedToAnotherEdgeOfTheFace = false;
+				for (FTopologicalEdge* Twin : Edge->GetTwinEntities())
+				{
+					if (Twin == Edge)
+					{
+						continue;
+					}
+					if (Twin->GetFace() == &Face)
+					{
+						bEdgeIsConnectedToAnotherEdgeOfTheFace = true;
+						break;
+					}
+				}
+			}
+			if (bEdgeIsConnectedToAnotherEdgeOfTheFace)
+			{
+				Segment->SetAsThinZone();
+			}
+		}
+	}
 
 	// For each segment, the nearest segment is search
 	// If segment is from an inner loop, the nearest could not be from the same loop
 	for (Topo::FEdgeSegment* Segment : Context.LoopSegments)
 	{
+		if (Segment->IsThinZone())
+		{
+			continue;
+		}
+
 		const FPoint& Point = Segment->GetMiddle();
-		const FTopologicalEdge* Edge = Segment->GetEdge();
 
-		//FPoint ClosedPoint;
 		Topo::FEdgeSegment* ClosedSegment = nullptr;
 
 		double MinSquareDistance = HUGE_VALUE;
@@ -32,7 +63,7 @@ void FFaceAnalyzer::FindClosedSegments(Topo::FThinFaceContext& Context)
 		for (Topo::FEdgeSegment* Candidate : Context.LoopSegments)
 		{
 			// to avoid to define cylinder or cone as a thin surface
-			if (Candidate->GetEdge()->IsLinkedTo(*Edge))
+			if (Candidate->IsThinZone())
 			{
 				continue;
 			}
@@ -73,7 +104,7 @@ void FFaceAnalyzer::FindClosedSegments(Topo::FThinFaceContext& Context)
 		{
 			Segment->SetClosedSegment(ClosedSegment, MinSquareDistance);
 
-#ifdef DEBUG_FIND_THIN_FACE
+#ifdef DEBUG_THIN_FACE
 			//F3DDebugSession _(TEXT("Segment"));
 			DisplaySegment(Segment->GetExtemity(ELimit::Start), Segment->GetExtemity(ELimit::End), 0, EVisuProperty::BlueCurve);
 			DisplaySegment(ClosedSegment->GetExtemity(ELimit::Start), ClosedSegment->GetExtemity(ELimit::End), 0, EVisuProperty::RedCurve);
@@ -90,6 +121,7 @@ void FFaceAnalyzer::Analyze(Topo::FThinFaceContext& Context)
 	double MaxSquareDistance = 0;
 	double MedSquareDistance = 0;
 	double EdgeLength = 0;
+	bool bIsThinZone = false;
 
 	TFunction<void()> SetEdgeMaxGap = [&]()
 	{
@@ -97,17 +129,23 @@ void FFaceAnalyzer::Analyze(Topo::FThinFaceContext& Context)
 		{
 			MedSquareDistance /= EdgeLength;
 		}
-		else
+		else if(!bIsThinZone)
 		{
 			MedSquareDistance = DOUBLE_BIG_NUMBER;
 			MaxSquareDistance = DOUBLE_BIG_NUMBER;
 		}
 
+		if (bIsThinZone)
+		{
+			Context.EdgeSquareDistance.Add(0.);
+			Context.EdgeMaxSquareDistance.Add(0.);
+			return;
+		}
+		
 		if (MaxSquareDistance < 2 * SquareTolerance)
 		{
 			Context.EdgeSquareDistance.Add(MedSquareDistance);
 			Context.EdgeMaxSquareDistance.Add(MaxSquareDistance);
-
 			Context.MaxSquareDistance = FMath::Max(Context.MaxSquareDistance, MaxSquareDistance);
 			Context.ThinSideEdgeLength += Edge->Length();
 		}
@@ -130,12 +168,17 @@ void FFaceAnalyzer::Analyze(Topo::FThinFaceContext& Context)
 			SetEdgeMaxGap();
 
 			Edge = Segment->GetEdge();
-			MaxSquareDistance = 0;
-			MedSquareDistance = 0;
-			EdgeLength = 0;
+			MaxSquareDistance = 0.;
+			MedSquareDistance = 0.;
+			EdgeLength = 0.;
+			bIsThinZone = false;
 		}
 
-		if (Segment->GetClosedSquareDistance() > 0)
+		if (Segment->IsThinZone())
+		{
+			bIsThinZone = true;
+		}
+		else if (Segment->GetClosedSquareDistance() > 0)
 		{
 			MedSquareDistance += Segment->GetLength() * Segment->GetClosedSquareDistance();
 			EdgeLength += Segment->GetLength();
@@ -161,6 +204,11 @@ void FFaceAnalyzer::BuildLoopSegments(Topo::FThinFaceContext& Context)
 	for (const FOrientedEdge& OrientedEdge : Edges)
 	{
 		const FTopologicalEdge* Edge = OrientedEdge.Entity.Get();
+		if (Edge->IsDeletedOrDegenerated())
+		{
+			continue;
+		}
+
 		TArray<double> Coordinates;
 		TArray<FPoint> Points;
 		Edge->GetCurve()->GetDiscretizationPoints(Edge->GetBoundary(), Coordinates, Points);
@@ -206,14 +254,14 @@ bool FFaceAnalyzer::IsThinFace(double& OutGapSize)
 	FindClosedSegments(Context);
 	Chronos.FindClosedSegmentTime = FChrono::Elapse(StartTime);
 #ifdef CADKERNEL_DEV
-	DisplayClosedSegments(Context);
+	DisplayCloseSegments(Context);
 #endif
 
 	StartTime = FChrono::Now();
 	Analyze(Context);
 	Chronos.AnalyzeClosedSegmentTime = FChrono::Elapse(StartTime);
 
-	if (Context.ThinSideEdgeLength > Context.OppositSideEdgeLength && Context.MaxSquareDistance < SquareTolerance )
+	if (Context.OppositSideEdgeLength < MaxOppositSideLength || (Context.ThinSideEdgeLength > Context.OppositSideEdgeLength && Context.MaxSquareDistance < SquareTolerance ) )
 	{
 		OutGapSize = sqrt(Context.MaxSquareDistance);
 		return true;
diff --git a/Engine/Source/Runtime/Datasmith/CADKernel/Private/CADKernel/Topo/FaceAnalyzer.h b/Engine/Source/Runtime/Datasmith/CADKernel/Private/CADKernel/Topo/FaceAnalyzer.h
--- a/Engine/Source/Runtime/Datasmith/CADKernel/Private/CADKernel/Topo/FaceAnalyzer.h
+++ b/Engine/Source/Runtime/Datasmith/CADKernel/Private/CADKernel/Topo/FaceAnalyzer.h
@@ -45,6 +45,8 @@ private:
 
 	double Length;
 
+	bool bIsThinZone = false;
+
 #ifdef CADKERNEL_DEV
 	FIdent Id;
 	static FIdent LastId;
@@ -142,6 +144,16 @@ public:
 	{
 		return ProjectPointOnSegment(PointToProject, Points[ELimit::Start], Points[ELimit::End], SegmentU, true);
 	}
+
+	bool IsThinZone()
+	{
+		return bIsThinZone;
+	}
+
+	void SetAsThinZone()
+	{
+		bIsThinZone = true;
+	}
 };
 
 struct FThinFaceContext
@@ -174,8 +186,9 @@ public:
 	Topo::FFaceAnalyzerChronos Chronos;
 
 protected:
-	double Tolerance;
-	double SquareTolerance;
+	const double Tolerance;
+	const double SquareTolerance;
+	const double MaxOppositSideLength;
 
 	FTopologicalFace& Face;
 
@@ -184,6 +197,7 @@ public:
 	FFaceAnalyzer(FTopologicalFace& InFace, double InTol)
 		: Tolerance(InTol)
 		, SquareTolerance(FMath::Square(InTol))
+		, MaxOppositSideLength(4. * InTol)
 		, Face(InFace)
 	{
  	}
@@ -197,7 +211,7 @@ private:
 	void Analyze(Topo::FThinFaceContext& Context);
 
 #ifdef CADKERNEL_DEV
-	void DisplayClosedSegments(Topo::FThinFaceContext& Context);
+	void DisplayCloseSegments(Topo::FThinFaceContext& Context);
 	void DisplayLoopSegments(Topo::FThinFaceContext& Context);
 #endif
 };
diff --git a/Engine/Source/Runtime/Datasmith/CADKernel/Private/CADKernel/Topo/Shell.cpp b/Engine/Source/Runtime/Datasmith/CADKernel/Private/CADKernel/Topo/Shell.cpp
--- a/Engine/Source/Runtime/Datasmith/CADKernel/Private/CADKernel/Topo/Shell.cpp
+++ b/Engine/Source/Runtime/Datasmith/CADKernel/Private/CADKernel/Topo/Shell.cpp
@@ -543,6 +543,18 @@ void UnlinkFromOther(TArray<FTopologicalFace*>& Faces, TArray<FTopologicalVertex
 
 	for (FTopologicalFace* Face : Faces)
 	{
+		TFunction<void(FTopologicalVertex&)> AddTwinVertex = [&Face, &VerticesToLink](FTopologicalVertex& Vertex)
+		{
+			for (FTopologicalVertex* TwinVertex : Vertex.GetTwinEntities())
+			{
+				if (!TwinVertex->IsProcessed() && TwinVertex->GetFace() == Face)
+				{
+					VerticesToLink.Add(TwinVertex);
+					TwinVertex->SetProcessedMarker();
+				}
+			}
+		};
+
 		for (const TSharedPtr<FTopologicalLoop>& Loop : Face->GetLoops())
 		{
 			for (const FOrientedEdge& OrientedEdge : Loop->GetEdges())
@@ -565,15 +577,17 @@ void UnlinkFromOther(TArray<FTopologicalFace*>& Faces, TArray<FTopologicalVertex
 					{
 						continue;
 					}
+
 					if (!NeighborFace->IsToProcess())
 					{
+						FTopologicalVertex& StartVertex = Edge->GetStartVertex().Get();
+						AddTwinVertex(StartVertex);
+
+						FTopologicalVertex& EndVertex = Edge->GetEndVertex().Get();
+						AddTwinVertex(EndVertex);
+
 						Edge->Unlink();
 
-						FTopologicalVertex& Vertex = *Edge->GetStartVertex();
-						if(!Vertex.IsProcessed())
-						{
-							VerticesToLink.Add(&Vertex);
-						}
 						break;
 					}
 				}
diff --git a/Engine/Source/Runtime/Datasmith/CADKernel/Private/CADKernel/Topo/TopologicalEdge.cpp b/Engine/Source/Runtime/Datasmith/CADKernel/Private/CADKernel/Topo/TopologicalEdge.cpp
--- a/Engine/Source/Runtime/Datasmith/CADKernel/Private/CADKernel/Topo/TopologicalEdge.cpp
+++ b/Engine/Source/Runtime/Datasmith/CADKernel/Private/CADKernel/Topo/TopologicalEdge.cpp
@@ -1073,6 +1073,10 @@ TSharedPtr<FTopologicalEdge> FTopologicalEdge::CreateEdgeByMergingEdges(const do
 				FPoint2D Target = Edges[Index].Entity->Approximate2DPoint(TargetCoordinate);
 
 				int32 NeigborIndex = PreviousLength > NextLength ? Index - 1 : Index + 1;
+				if (NeigborIndex < 0 || NeigborIndex >= NurbsCurves.Num())
+				{
+					return TSharedPtr<FTopologicalEdge>();
+				}
 				NurbsCurves[NeigborIndex]->ExtendTo(Target);
 			}
 		}
diff --git a/Engine/Source/Runtime/Datasmith/CADKernel/Private/CADKernel/Topo/TopologicalFace.cpp b/Engine/Source/Runtime/Datasmith/CADKernel/Private/CADKernel/Topo/TopologicalFace.cpp
--- a/Engine/Source/Runtime/Datasmith/CADKernel/Private/CADKernel/Topo/TopologicalFace.cpp
+++ b/Engine/Source/Runtime/Datasmith/CADKernel/Private/CADKernel/Topo/TopologicalFace.cpp
@@ -277,7 +277,7 @@ void FTopologicalFace::Disjoin(TArray<FTopologicalEdge*>* NewBorderEdges)
 			const TArray<FTopologicalEdge*> Twins = Edge->GetTwinEntities();
 			for (FTopologicalEdge* TwinEdge : Twins)
 			{
-				if (NewBorderEdges && TwinEdge != Edge)
+				if (NewBorderEdges && TwinEdge != Edge && TwinEdge->GetFace() != this)
 				{
 					NewBorderEdges->Add(TwinEdge);
 				}
diff --git a/Engine/Source/Runtime/Datasmith/CADKernel/Private/CADKernel/Topo/TopologicalVertex.cpp b/Engine/Source/Runtime/Datasmith/CADKernel/Private/CADKernel/Topo/TopologicalVertex.cpp
--- a/Engine/Source/Runtime/Datasmith/CADKernel/Private/CADKernel/Topo/TopologicalVertex.cpp
+++ b/Engine/Source/Runtime/Datasmith/CADKernel/Private/CADKernel/Topo/TopologicalVertex.cpp
@@ -31,7 +31,7 @@ void FTopologicalVertex::RemoveConnectedEdge(FTopologicalEdge& Edge)
 	}
 }
 
-bool FTopologicalVertex::IsBorderVertex()
+bool FTopologicalVertex::IsBorderVertex() const
 {
 	for (FTopologicalVertex* Vertex : GetTwinEntities())
 	{
@@ -52,6 +52,16 @@ bool FTopologicalVertex::IsBorderVertex()
 	return false;
 }
 
+const FTopologicalFace* FTopologicalVertex::GetFace() const
+{
+	if (ConnectedEdges.IsEmpty())
+	{
+		return nullptr;
+	}
+	return ConnectedEdges[0]->GetFace();
+}
+
+
 void FTopologicalVertex::GetConnectedEdges(const FTopologicalVertex& OtherVertex, TArray<FTopologicalEdge*>& OutEdges) const
 {
 	OutEdges.Reserve(GetTwinEntityCount());
diff --git a/Engine/Source/Runtime/Datasmith/CADKernel/Private/CADKernel/Topo/Topomaker.cpp b/Engine/Source/Runtime/Datasmith/CADKernel/Private/CADKernel/Topo/Topomaker.cpp
--- a/Engine/Source/Runtime/Datasmith/CADKernel/Private/CADKernel/Topo/Topomaker.cpp
+++ b/Engine/Source/Runtime/Datasmith/CADKernel/Private/CADKernel/Topo/Topomaker.cpp
@@ -483,12 +483,22 @@ void GetVertices(const TArray<FTopologicalEdge*>& InEdges, TArray<FTopologicalVe
 
 	for (const FTopologicalEdge* Edge : InEdges)
 	{
+		if (!Edge || Edge->IsDeleted())
+		{
+			continue;
+		}
+
 		Edge->GetStartVertex()->GetLinkActiveEntity()->ResetMarker1();
 		Edge->GetEndVertex()->GetLinkActiveEntity()->ResetMarker1();
 	}
 
 	for (const FTopologicalEdge* Edge : InEdges)
 	{
+		if (!Edge || Edge->IsDeleted())
+		{
+			continue;
+		}
+
 		if (!Edge->GetStartVertex()->GetLinkActiveEntity()->HasMarker1())
 		{
 			Edge->GetStartVertex()->GetLinkActiveEntity()->SetMarker1();
@@ -642,6 +652,8 @@ void FTopomaker::Sew()
 	MergeUnconnectedSuccessiveEdges();
 
 	// advance case: two partially coincident edges connected at one extremity i.e. coincident along the shortest edge
+	BorderVertices.Reset();
+	GetBorderVertices(BorderVertices);
 	TopomakerTools::StitchParallelEdges(BorderVertices, SewTolerance, EdgeLengthTolerance);
 
 	const bool bForceJoining = SewOption::IsForceJoining(SewOptions);
@@ -665,6 +677,10 @@ void FTopomaker::Sew()
 	}
 
 	// advance case: two partially coincident edges connected at one extremity i.e. coincident along the shortest edge
+	BorderVertices.Reset();
+	GetBorderVertices(BorderVertices);
+
+	TopomakerTools::MergeCoincidentEdges(BorderVertices, bForceJoining ? LargeEdgeLengthTolerance : EdgeLengthTolerance);
 	TopomakerTools::StitchParallelEdges(BorderVertices, bForceJoining ? SewToleranceToForceJoin : SewTolerance, bForceJoining ? LargeEdgeLengthTolerance : EdgeLengthTolerance);
 
 	if (SewOption::IsRemoveDuplicatedFaces(SewOptions))
@@ -821,7 +837,7 @@ void FTopomaker::MergeUnconnectedSuccessiveEdges()
 			{
 				TSharedRef<FTopologicalVertex> StartVertex = Candidates[0].Direction == EOrientation::Front ? Candidates[0].Entity->GetStartVertex() : Candidates[0].Entity->GetEndVertex();
 				TSharedRef<FTopologicalVertex> EndVertex = Candidates.Last().Direction == EOrientation::Front ? Candidates.Last().Entity->GetEndVertex() : Candidates.Last().Entity->GetStartVertex();
-				FTopologicalEdge::CreateEdgeByMergingEdges(LargeEdgeLengthTolerance * 10., Candidates, StartVertex, EndVertex);
+				FTopologicalEdge::CreateEdgeByMergingEdges(LargeEdgeLengthTolerance, Candidates, StartVertex, EndVertex);
 			}
 		}
 	}
@@ -861,21 +877,46 @@ void FTopomaker::RemoveIsolatedEdges()
 	FMessage::Printf(EVerboseLevel::Log, TEXT("\n\nIsolatedEdges count %d\n\n\n"), IsolatedEdges.Num());
 }
 
-void SplitVertexLinkByShell(FTopologicalVertex* Vertex)
+void SplitVertexLinkByShell(FTopologicalVertex* InVertex)
 {
 	TMap<const FShell*, TArray<FTopologicalVertex*>> ShellToVertices;
 
-	const TArray<FTopologicalVertex*>& Twins = Vertex->GetTwinEntities();
+	const TArray<FTopologicalVertex*>& Twins = InVertex->GetTwinEntities();
 	ShellToVertices.Reserve(Twins.Num());
 
 	for (FTopologicalVertex* TwinVertex : Twins)
 	{
-		const FShell* OtherShell = (const FShell*)TwinVertex->GetDirectConnectedEdges()[0]->GetLoop()->GetFace()->GetHost();
+		const FShell* OtherShell = (const FShell*)TwinVertex->GetDirectConnectedEdges()[0]->GetFace()->GetHost();
 		TArray<FTopologicalVertex*>& Vertices = ShellToVertices.FindOrAdd(OtherShell);
 		Vertices.Add(TwinVertex);
 	}
 
-	Vertex->UnlinkTwinEntities();
+#ifdef DEBUG_Split_Vertex_Link_By_Shell
+	{
+		F3DDebugSession B(*FString::Printf(TEXT("Shells")));
+		for (TPair<const FShell*, TArray<FTopologicalVertex*>>& Pair : ShellToVertices)
+		{
+			F3DDebugSession B(*FString::Printf(TEXT("Shell to vertex")));
+			{
+				F3DDebugSession B(*FString::Printf(TEXT("Shell")));
+				Display(*Pair.Key);
+			}
+			TArray<FTopologicalVertex*>& Vertices = Pair.Value;
+			for (const FTopologicalVertex* Vertex : Vertices)
+			{
+				{
+					F3DDebugSession A(*FString::Printf(TEXT("Vertex %d"), Vertex->GetId()));
+					Display(*Vertex);
+					Display(*Vertex->GetDirectConnectedEdges()[0], EVisuProperty::RedPoint);
+					Display(*Vertex->GetFace());
+				}
+			}
+		}
+		Wait();
+	}
+#endif
+
+	InVertex->UnlinkTwinEntities();
 
 	for (TPair<const FShell*, TArray<FTopologicalVertex*>>& Pair : ShellToVertices)
 	{
@@ -896,42 +937,38 @@ void FTopomaker::UnlinkNonManifoldVertex()
 	TArray<FTopologicalVertex*> Vertices;
 	GetVertices(Vertices);
 
-	// check that each vertex belongs to only one shell
-	for (const FTopologicalVertex* Vertex : Vertices)
-	{
-		for (const FTopologicalVertex* TwinVertex : Vertex->GetTwinEntities())
-		{
-			const FShell* Shell = nullptr;
-			{
-				FTopologicalEdge* Edge = TwinVertex->GetDirectConnectedEdges()[0];
-				Shell = (FShell*)Edge->GetLoop()->GetFace()->GetHost();
-			}
-			if (Shell != nullptr)
-			{
-				for (FTopologicalEdge* Edge : TwinVertex->GetDirectConnectedEdges())
-				{
-					ensureCADKernel(Edge->GetLoop()->GetFace()->GetHost() == Shell);
-				}
-			}
-		}
-	}
+	TArray<FTopologicalVertex*> VerticesToSplit;
+	VerticesToSplit.Reserve(Vertices.Num());
 
 	for (FTopologicalVertex* Vertex : Vertices)
 	{
 		const TArray<FTopologicalVertex*>& Twins = Vertex->GetTwinEntities();
-		const FTopologicalEdge* Edge = Twins[0]->GetDirectConnectedEdges()[0];
-		const FShell* Shell = (const FShell*)Edge->GetLoop()->GetFace()->GetHost();
+		const FShell* Shell = (const FShell*) Twins[0]->GetFace()->GetHost();
 
 		for (const FTopologicalVertex* TwinVertex : Twins)
 		{
-			const FShell* OtherShell = (const FShell*)TwinVertex->GetDirectConnectedEdges()[0]->GetLoop()->GetFace()->GetHost();
+			const FShell* OtherShell = (const FShell*)TwinVertex->GetDirectConnectedEdges()[0]->GetFace()->GetHost();
 			if (OtherShell != Shell)
 			{
-				SplitVertexLinkByShell(Vertex);
-				break;
+				TSharedPtr<FTopologicalVertex> ActiveVertex = Vertex->GetLinkActiveEntity();
+				if(!ActiveVertex->IsProcessed())
+				{
+					VerticesToSplit.Add(ActiveVertex.Get());
+					ActiveVertex->SetProcessedMarker();
+				}
 			}
 		}
 	}
+
+	for (const FTopologicalVertex* Vertex : VerticesToSplit)
+	{
+		Vertex->ResetProcessedMarker();
+	}
+
+	for (FTopologicalVertex* Vertex : VerticesToSplit)
+	{
+		SplitVertexLinkByShell(Vertex);
+	}
 }
 
 void FindBorderVertex(const TArray<FTopologicalEdge*>& Edges, TArray<FTopologicalVertex*>& OutBorderVertices)
@@ -1238,6 +1275,7 @@ void FTopomaker::SplitIntoConnectedShells()
 		}
 	}
 
+	UnlinkFromOther();
 	UnlinkNonManifoldVertex();
 
 	RemoveEmptyShells();
@@ -1438,12 +1476,32 @@ void FTopomaker::ResetMarkersOfFaces()
 
 void FTopomaker::UnlinkFromOther()
 {
+	TFunction<void(TArray<FTopologicalVertex*>&)> MergeCoincidents = [&Tolerance = Tolerance](TArray<FTopologicalVertex*>& Vertices)
+	{
+		TArray<FTopologicalVertex*> ActiveVerticesToLink;
+		ActiveVerticesToLink.Reserve(Vertices.Num());
+		for (FTopologicalVertex* Vertex : Vertices)
+		{
+			if (!Vertex->IsProcessed())
+			{
+				ActiveVerticesToLink.Add(&Vertex->GetLinkActiveEntity().Get());
+				Vertex->SetProcessedMarker();
+			}
+		}
+		for (FTopologicalVertex* Vertex : Vertices)
+		{
+			Vertex->ResetProcessedMarker();
+		}
+
+		TopomakerTools::MergeCoincidentVertices(ActiveVerticesToLink, Tolerance);
+	};
+
 	TArray<FTopologicalVertex*> VerticesToLink;
 	for (FShell* Shell : Shells)
 	{
 		VerticesToLink.Reset();
 		Shell->UnlinkFromOther(VerticesToLink);
-		TopomakerTools::MergeCoincidentVertices(VerticesToLink, Tolerance);
+		MergeCoincidents(VerticesToLink);
 	}
 
 	{
@@ -1458,10 +1516,8 @@ void FTopomaker::UnlinkFromOther()
 		}
 
 		ShellTools::UnlinkFromOther(FacePtrs, VerticesToLink);
-		TopomakerTools::MergeCoincidentVertices(VerticesToLink, Tolerance);
+		MergeCoincidents(VerticesToLink);
 	}
-
-	UnlinkNonManifoldVertex();
 }
 
 }
diff --git a/Engine/Source/Runtime/Datasmith/CADKernel/Private/CADKernel/UI/Display.cpp b/Engine/Source/Runtime/Datasmith/CADKernel/Private/CADKernel/UI/Display.cpp
--- a/Engine/Source/Runtime/Datasmith/CADKernel/Private/CADKernel/UI/Display.cpp
+++ b/Engine/Source/Runtime/Datasmith/CADKernel/Private/CADKernel/UI/Display.cpp
@@ -492,11 +492,11 @@ void DisplayAABB2D(const FAABB2D& aabb, FIdent Ident)
 #endif
 }
 
-void Display(const FTopologicalVertex& Vertex)
+void Display(const FTopologicalVertex& Vertex, EVisuProperty Property)
 {
 #ifdef CADKERNEL_DEBUG
 	F3DDebugSegment G(Vertex.GetId());
-	DrawPoint(Vertex.GetCoordinates());
+	DrawPoint(Vertex.GetCoordinates(), Property);
 #endif
 }
 
@@ -1106,8 +1106,11 @@ void DisplayMesh(const FShell& Shell)
 			FTopologicalFace& Face = *OrientedFace.Entity;
 			if (!Face.IsDeletedOrDegenerated() && Face.IsMeshed())
 			{
-				FFaceMesh& Mesh = *Face.GetMesh();
-				DisplayMesh(Mesh);
+				const FFaceMesh* Mesh = Face.GetMesh();
+				if(Mesh)
+				{
+					DisplayMesh(*Mesh);
+				}
 			}
 		}
 	}
diff --git a/Engine/Source/Runtime/Datasmith/CADKernel/Public/CADKernel/Topo/TopologicalEdge.h b/Engine/Source/Runtime/Datasmith/CADKernel/Public/CADKernel/Topo/TopologicalEdge.h
--- a/Engine/Source/Runtime/Datasmith/CADKernel/Public/CADKernel/Topo/TopologicalEdge.h
+++ b/Engine/Source/Runtime/Datasmith/CADKernel/Public/CADKernel/Topo/TopologicalEdge.h
@@ -482,8 +482,11 @@ public:
 		{
 			return GetLinkActiveEdge()->GetMesh();
 		}
-		ensureCADKernel(Mesh.IsValid());
-		return Mesh.Get();
+		if (Mesh.IsValid())
+		{
+			return Mesh.Get();
+		}
+		return nullptr;
 	}
 
 	FEdgeMesh& GetOrCreateMesh(FModelMesh& MeshModel);
diff --git a/Engine/Source/Runtime/Datasmith/CADKernel/Public/CADKernel/Topo/TopologicalFace.h b/Engine/Source/Runtime/Datasmith/CADKernel/Public/CADKernel/Topo/TopologicalFace.h
--- a/Engine/Source/Runtime/Datasmith/CADKernel/Public/CADKernel/Topo/TopologicalFace.h
+++ b/Engine/Source/Runtime/Datasmith/CADKernel/Public/CADKernel/Topo/TopologicalFace.h
@@ -308,6 +308,11 @@ public:
 		CarrierSurface.Reset();
 		for (TSharedPtr<FTopologicalLoop>& Loop : Loops)
 		{
+			for (FOrientedEdge& Edge : Loop->GetEdges())
+			{
+				Edge.Entity->Delete();
+			}
+
 			Loop->Empty();
 		}
 		Loops.Empty();
@@ -363,10 +368,13 @@ public:
 		return Mesh.IsValid();
 	}
 
-	const TSharedRef<FFaceMesh> GetMesh() const
+	FFaceMesh* GetMesh() const
 	{
-		ensureCADKernel(Mesh.IsValid());
-		return Mesh.ToSharedRef();
+		if (Mesh.IsValid())
+		{
+			return Mesh.Get();
+		}
+		return nullptr;
 	}
 
 	void InitDeltaUs();
diff --git a/Engine/Source/Runtime/Datasmith/CADKernel/Public/CADKernel/Topo/TopologicalShapeEntity.h b/Engine/Source/Runtime/Datasmith/CADKernel/Public/CADKernel/Topo/TopologicalShapeEntity.h
--- a/Engine/Source/Runtime/Datasmith/CADKernel/Public/CADKernel/Topo/TopologicalShapeEntity.h
+++ b/Engine/Source/Runtime/Datasmith/CADKernel/Public/CADKernel/Topo/TopologicalShapeEntity.h
@@ -66,6 +66,11 @@ public:
 		return HostedBy;
 	}
 
+	const FTopologicalShapeEntity* GetHost() const
+	{
+		return HostedBy;
+	}
+
 	void SetHost(FTopologicalShapeEntity* Body)
 	{
 		HostedBy = Body;
diff --git a/Engine/Source/Runtime/Datasmith/CADKernel/Public/CADKernel/Topo/TopologicalVertex.h b/Engine/Source/Runtime/Datasmith/CADKernel/Public/CADKernel/Topo/TopologicalVertex.h
--- a/Engine/Source/Runtime/Datasmith/CADKernel/Public/CADKernel/Topo/TopologicalVertex.h
+++ b/Engine/Source/Runtime/Datasmith/CADKernel/Public/CADKernel/Topo/TopologicalVertex.h
@@ -137,7 +137,11 @@ public:
 		{
 			return GetLinkActiveEntity()->GetMesh();
 		}
-		return Mesh.Get();
+		if (Mesh.IsValid())
+		{
+			return Mesh.Get();
+		}
+		return nullptr;
 	}
 
 	void Link(FTopologicalVertex& InEntity);
@@ -178,7 +182,7 @@ public:
 		TLinkable<FTopologicalVertex, FVertexLink>::Empty();
 	}
 
-	bool IsBorderVertex();
+	bool IsBorderVertex() const;
 
 	void AddConnectedEdge(FTopologicalEdge& Edge);
 	void RemoveConnectedEdge(FTopologicalEdge& Edge);
@@ -198,6 +202,8 @@ public:
 		return ConnectedEdges;
 	}
 
+	const FTopologicalFace* GetFace() const;
+
 	void GetConnectedEdges(TArray<FTopologicalEdge*>& OutConnectedEdges) const
 	{
 		if (!TopologicalLink.IsValid())
diff --git a/Engine/Source/Runtime/Datasmith/CADKernel/Public/CADKernel/UI/Display.h b/Engine/Source/Runtime/Datasmith/CADKernel/Public/CADKernel/UI/Display.h
--- a/Engine/Source/Runtime/Datasmith/CADKernel/Public/CADKernel/UI/Display.h
+++ b/Engine/Source/Runtime/Datasmith/CADKernel/Public/CADKernel/UI/Display.h
@@ -226,7 +226,7 @@ CADKERNEL_API void Display(const FShell& Shell);
 CADKERNEL_API void Display(const FTopologicalEdge& Edge, EVisuProperty Property = EVisuProperty::BlueCurve);
 CADKERNEL_API void Display(const FTopologicalFace& Face);
 CADKERNEL_API void Display(const FTopologicalLoop& Loop);
-CADKERNEL_API void Display(const FTopologicalVertex& Vertex);
+CADKERNEL_API void Display(const FTopologicalVertex& Vertex, EVisuProperty Property = EVisuProperty::BluePoint);
 
 CADKERNEL_API void Display2D(const FTopologicalEdge& Edge, EVisuProperty Property = EVisuProperty::BlueCurve);
 CADKERNEL_API void Display2D(const FTopologicalFace& Face);
