Fix crash: Add check to avoid empty FTopologicalLoop in CADKernel

diff --git a/Engine/Plugins/Enterprise/DatasmithCADImporter/Source/CADInterfaces/Private/TechSoftBridge.cpp b/Engine/Plugins/Enterprise/DatasmithCADImporter/Source/CADInterfaces/Private/TechSoftBridge.cpp
--- a/Engine/Plugins/Enterprise/DatasmithCADImporter/Source/CADInterfaces/Private/TechSoftBridge.cpp
+++ b/Engine/Plugins/Enterprise/DatasmithCADImporter/Source/CADInterfaces/Private/TechSoftBridge.cpp
@@ -49,7 +49,7 @@ template<typename... InArgTypes>
 A3DStatus GetSurfaceAsNurbs(const A3DSurfBase* A3DSurface, A3DSurfNurbsData* DataPtr, InArgTypes&&... Args)
 {
 	return TechSoftInterface::GetSurfaceAsNurbs(A3DSurface, DataPtr, Forward<InArgTypes>(Args)...);
-}; 
+};
 
 UE::CADKernel::FMatrixH CreateCoordinateSystem(const A3DMiscCartesianTransformationData& Transformation, double UnitScale = 1.0)
 {
@@ -151,7 +151,7 @@ UE::CADKernel::FSurfacicBoundary GetSurfacicBoundary(A3DDomainData& Domain, cons
 	UE::CADKernel::FPoint2D Min(Domain.m_sMin.m_dX, Domain.m_sMin.m_dY);
 	UE::CADKernel::FPoint2D Max(Domain.m_sMax.m_dX, Domain.m_sMax.m_dY);
 
-	if(UVReparameterization.GetNeedApply())
+	if (UVReparameterization.GetNeedApply())
 	{
 		UVReparameterization.Apply(Min);
 		UVReparameterization.Apply(Max);
@@ -200,7 +200,7 @@ const A3DRiBrepModel* FTechSoftBridge::GetA3DBody(UE::CADKernel::FBody* BRepMode
 	const A3DEntity** BodyPtr = CADKernelToTechSoft.Find(BRepModel);
 	if (BodyPtr)
 	{
-		return (A3DRiBrepModel*) *BodyPtr;
+		return (A3DRiBrepModel*)*BodyPtr;
 	}
 	return nullptr;
 }
@@ -208,7 +208,7 @@ const A3DRiBrepModel* FTechSoftBridge::GetA3DBody(UE::CADKernel::FBody* BRepMode
 UE::CADKernel::FBody* FTechSoftBridge::GetBody(A3DRiBrepModel* A3DBRepModel)
 {
 	TSharedPtr<UE::CADKernel::FBody>* BodyPtr = TechSoftToCADKernel.Find(A3DBRepModel);
-	if (BodyPtr &&!(*BodyPtr)->IsDeleted())
+	if (BodyPtr && !(*BodyPtr)->IsDeleted())
 	{
 		return (*BodyPtr).Get();
 	}
@@ -227,11 +227,11 @@ UE::CADKernel::FBody* FTechSoftBridge::AddBody(A3DRiBrepModel* A3DBRepModel, TMa
 	FArchiveCADObject BRepMetaData;
 	Parser.ExtractMetaData(A3DBRepModel, BRepMetaData);
 	FString* Name = MetaData.Find(TEXT("Name"));
-	if(Name != nullptr)
+	if (Name != nullptr)
 	{
 		BRepMetaData.MetaData.FindOrAdd(TEXT("Name")) = *Name;
 	}
-	
+
 	if (TSharedPtr<UE::CADKernel::FBody>* BodyPtr = TechSoftToCADKernel.Find(A3DBRepModel))
 	{
 		return (*BodyPtr)->IsDeleted() ? nullptr : (*BodyPtr).Get();
@@ -579,7 +579,7 @@ TSharedPtr<UE::CADKernel::FCurve> AddCurveNurbsFromData(A3DCrvNurbsData& A3DNurb
 	Nurbs.Degree = A3DNurbs.m_uiDegree;
 
 	TechSoftUtils::FillPointArray(A3DNurbs.m_uiCtrlSize, A3DNurbs.m_pCtrlPts, Nurbs.Poles);
-	if(Nurbs.Dimension == 2)
+	if (Nurbs.Dimension == 2)
 	{
 		UVReparameterization.Process(Nurbs.Poles);
 	}
@@ -755,7 +755,7 @@ TSharedPtr<UE::CADKernel::FTopologicalLoop> FTechSoftBridge::AddLoop(const A3DTo
 			}
 		}
 	}
-	 
+
 	return Loop;
 }
 
@@ -764,7 +764,7 @@ void FTechSoftBridge::AddFace(const A3DTopoFace* A3DFace, UE::CADKernel::EOrient
 	using namespace UE::CADKernel;
 
 #ifdef CADKERNEL_DEV
-	UE::CADKernel::FCADFileReport::Get().FaceCount++;
+	FCADFileReport::Get().FaceCount++;
 #endif
 
 	FArchiveCADObject MetaData;
@@ -774,7 +774,7 @@ void FTechSoftBridge::AddFace(const A3DTopoFace* A3DFace, UE::CADKernel::EOrient
 	if (!TopoFaceData.IsValid())
 	{
 #ifdef CADKERNEL_DEV
-		UE::CADKernel::FCADFileReport::Get().FailedFaceCount++;
+		FCADFileReport::Get().FailedFaceCount++;
 #endif
 		return;
 	}
@@ -785,8 +785,8 @@ void FTechSoftBridge::AddFace(const A3DTopoFace* A3DFace, UE::CADKernel::EOrient
 	if (!SurfacePtr.IsValid())
 	{
 #ifdef CADKERNEL_DEV
-		UE::CADKernel::FCADFileReport::Get().DegeneratedSurfaceCount++;
-		UE::CADKernel::FCADFileReport::Get().FailedFaceCount++; 
+		FCADFileReport::Get().DegeneratedSurfaceCount++;
+		FCADFileReport::Get().FailedFaceCount++;
 #endif
 		return;
 	}
@@ -848,7 +848,7 @@ void FTechSoftBridge::AddFace(const A3DTopoFace* A3DFace, UE::CADKernel::EOrient
 			Loops.Add(Loop);
 		}
 
-		if(Loops.Num() > 1)
+		if (Loops.Num() > 1)
 		{
 			// Check external loop
 			TSharedPtr<FTopologicalLoop> ExternalLoop;
@@ -857,7 +857,7 @@ void FTechSoftBridge::AddFace(const A3DTopoFace* A3DFace, UE::CADKernel::EOrient
 			for (const TSharedPtr<FTopologicalLoop>& Loop : Loops)
 			{
 				// fast but not accurate test to check if the loop is inside an other loop based of the bounding box 
-				switch(Loop->Boundary.IsInside(ExternalBoundary, Surface->GetIsoTolerances()))
+				switch (Loop->Boundary.IsInside(ExternalBoundary, Surface->GetIsoTolerances()))
 				{
 				case ESituation::Undefined:
 				{
@@ -882,16 +882,16 @@ void FTechSoftBridge::AddFace(const A3DTopoFace* A3DFace, UE::CADKernel::EOrient
 				}
 			}
 
-			if(!ExternalLoop->IsExternal())
+			if (!ExternalLoop->IsExternal())
 			{
 #ifdef CADKERNEL_DEV
-				UE::CADKernel::FCADFileReport::Get().WrongExternalLoopCount++;
+				FCADFileReport::Get().WrongExternalLoopCount++;
 				FMessage::Printf(Log, TEXT("ERROR: Face %d had wrong external loop\n"), Face->GetId());
 #endif
-	
+
 				for (TSharedPtr<FTopologicalLoop>& Loop : Loops)
 				{
-					if(Loop->IsExternal())
+					if (Loop->IsExternal())
 					{
 						Loop->SetInternal();
 						break;
@@ -904,7 +904,7 @@ void FTechSoftBridge::AddFace(const A3DTopoFace* A3DFace, UE::CADKernel::EOrient
 		if (Loops.Num() == 0)
 		{
 #ifdef CADKERNEL_DEV
-			UE::CADKernel::FCADFileReport::Get().FailedFaceCount++;
+			FCADFileReport::Get().FailedFaceCount++;
 #endif
 			Face->SetAsDegenerated();
 			Face->Delete();
@@ -915,11 +915,22 @@ void FTechSoftBridge::AddFace(const A3DTopoFace* A3DFace, UE::CADKernel::EOrient
 			int32 DoubtfulLoopOrientationCount = 0;
 			Face->AddLoops(Loops, DoubtfulLoopOrientationCount);
 #ifdef CADKERNEL_DEV
-			UE::CADKernel::FCADFileReport::Get().DoubtfulLoopOrientationCount += DoubtfulLoopOrientationCount;
+			FCADFileReport::Get().DoubtfulLoopOrientationCount += DoubtfulLoopOrientationCount;
 #endif
 		}
 	}
 
+	if (Face->GetLoops().Num() == 0)
+	{
+#ifdef CADKERNEL_DEV
+		FMessage::Printf(EVerboseLevel::Log, TEXT("A Face is degenerate, this face is ignored\n"));
+		FCADFileReport::Get().FailedFaceCount++;
+#endif
+		Face->SetAsDegenerated();
+		Face->Delete();
+		return;
+	}
+
 	AddMetadata(MetaData, *Face);
 	Face->SetHostId(ShellIndex);
 	Shell->Add(Face, Orientation);
@@ -944,7 +955,7 @@ TSharedPtr<UE::CADKernel::FSurface> FTechSoftBridge::AddSurface(const A3DSurfBas
 		{
 		case kA3DTypeSurfBlend01:
 			return AddBlend01Surface(A3DSurface, OutUVReparameterization);
-			
+
 		case kA3DTypeSurfBlend02:
 			return AddBlend02Surface(A3DSurface, OutUVReparameterization);
 
@@ -1105,7 +1116,7 @@ TSharedPtr<UE::CADKernel::FSurface> FTechSoftBridge::AddPlaneSurface(const A3DSu
 
 	OutUVReparameterization.AddUVTransform(A3DPlaneData->m_sParam);
 	OutUVReparameterization.ScaleUVTransform(BodyScale, BodyScale);
-	
+
 	UE::CADKernel::FMatrixH CoordinateSystem = TechSoftUtils::CreateCoordinateSystem(A3DPlaneData->m_sTrsf, BodyScale);
 	UE::CADKernel::FSurfacicBoundary Boundary = TechSoftUtils::GetSurfacicBoundary(A3DPlaneData->m_sParam.m_sUVDomain, OutUVReparameterization);
 	return UE::CADKernel::FSurface::MakePlaneSurface(GeometricTolerance, CoordinateSystem, Boundary);
@@ -1400,7 +1411,7 @@ TSharedPtr<UE::CADKernel::FSurface> FTechSoftBridge::AddSurfaceAsNurbs(const A3D
 void FTechSoftBridge::AddMetadata(FArchiveCADObject& MetaData, UE::CADKernel::FTopologicalShapeEntity& Entity)
 {
 	FString* Name = MetaData.MetaData.Find(TEXT("Name"));
-	if(Name != nullptr)
+	if (Name != nullptr)
 	{
 		Entity.SetName(*Name);
 	}
diff --git a/Engine/Plugins/Enterprise/DatasmithCADImporter/Source/CADLibrary/Private/CADKernelTools.cpp b/Engine/Plugins/Enterprise/DatasmithCADImporter/Source/CADLibrary/Private/CADKernelTools.cpp
--- a/Engine/Plugins/Enterprise/DatasmithCADImporter/Source/CADLibrary/Private/CADKernelTools.cpp
+++ b/Engine/Plugins/Enterprise/DatasmithCADImporter/Source/CADLibrary/Private/CADKernelTools.cpp
@@ -373,7 +373,8 @@ void FCADKernelTools::GetBodyTessellation(const UE::CADKernel::FModelMesh& Model
 			FObjectDisplayDataId FaceMaterial = ShellMaterial;
 			GetDisplayDataIds(*Face.Entity, FaceMaterial);
 
-			uint32 TriangleNum = GetFaceTessellation(*Face.Entity->GetMesh(), OutBodyMesh, FaceMaterial);
+			UE::CADKernel::FFaceMesh* FaceMesh = Face.Entity->GetMesh();
+			uint32 TriangleNum = GetFaceTessellation(*FaceMesh, OutBodyMesh, FaceMaterial);
 			OutBodyMesh.TriangleCount += TriangleNum;
 		}
 	}
diff --git a/Engine/Plugins/Enterprise/DatasmithCADImporter/Source/DatasmithOpenNurbsTranslator/Private/OpenNurbsBRepToCADKernelConverter.cpp b/Engine/Plugins/Enterprise/DatasmithCADImporter/Source/DatasmithOpenNurbsTranslator/Private/OpenNurbsBRepToCADKernelConverter.cpp
--- a/Engine/Plugins/Enterprise/DatasmithCADImporter/Source/DatasmithOpenNurbsTranslator/Private/OpenNurbsBRepToCADKernelConverter.cpp
+++ b/Engine/Plugins/Enterprise/DatasmithCADImporter/Source/DatasmithOpenNurbsTranslator/Private/OpenNurbsBRepToCADKernelConverter.cpp
@@ -290,6 +290,15 @@ TSharedPtr<UE::CADKernel::FTopologicalFace> FOpenNurbsBRepToCADKernelConverter::
 		}
 	}
 
+	if (Face->GetLoops().Num() == 0)
+	{
+		Face->SetAsDegenerated();
+		Face->Delete();
+
+		FMessage::Printf(EVerboseLevel::Log, TEXT("A Face is degenerate, this face is ignored\n"));
+		return TSharedPtr<FTopologicalFace>();
+	}
+
 	return Face;
 }
 
diff --git a/Engine/Plugins/Enterprise/DatasmithCADImporter/Source/DatasmithWireTranslator/Private/AliasModelToCADKernelConverter.cpp b/Engine/Plugins/Enterprise/DatasmithCADImporter/Source/DatasmithWireTranslator/Private/AliasModelToCADKernelConverter.cpp
--- a/Engine/Plugins/Enterprise/DatasmithCADImporter/Source/DatasmithWireTranslator/Private/AliasModelToCADKernelConverter.cpp
+++ b/Engine/Plugins/Enterprise/DatasmithCADImporter/Source/DatasmithWireTranslator/Private/AliasModelToCADKernelConverter.cpp
@@ -235,6 +235,16 @@ TSharedPtr<UE::CADKernel::FTopologicalFace> FAliasModelToCADKernelConverter::Add
 
 	int32 DoubtfulLoopOrientationCount = 0;
 	Face->AddLoops(Loops, DoubtfulLoopOrientationCount);
+
+	if (Face->GetLoops().Num() == 0)
+	{
+		Face->SetAsDegenerated();
+		Face->Delete();
+
+		UE::CADKernel::FMessage::Printf(UE::CADKernel::EVerboseLevel::Log, TEXT("The Face %s is degenerate, this face is ignored\n"), TrimRegion.name());
+		return TSharedPtr<UE::CADKernel::FTopologicalFace>();
+	}
+
 	return Face;
 }
 
diff --git a/Engine/Source/Runtime/Datasmith/CADKernel/Private/CADKernel/Geo/Sampling/SurfacicPolyline.cpp b/Engine/Source/Runtime/Datasmith/CADKernel/Private/CADKernel/Geo/Sampling/SurfacicPolyline.cpp
--- a/Engine/Source/Runtime/Datasmith/CADKernel/Private/CADKernel/Geo/Sampling/SurfacicPolyline.cpp
+++ b/Engine/Source/Runtime/Datasmith/CADKernel/Private/CADKernel/Geo/Sampling/SurfacicPolyline.cpp
@@ -159,7 +159,7 @@ void FSurfacicPolyline::ComputeIntersectionsWithIsos(const FLinearBoundary& InBo
 			IsoCoordinate = IsoCoordinates[IsoCoordinateIndex];
 		}
 
-		while (IsoCoordinates[IsoCoordinateIndex] < UMin && IsoCoordinateIndex < IsoCoordinates.Num())
+		while (IsoCoordinateIndex < IsoCoordinates.Num() - 1 && IsoCoordinates[IsoCoordinateIndex] < UMin)
 		{
 			IsoCoordinateIndex++;
 		}
diff --git a/Engine/Source/Runtime/Datasmith/CADKernel/Private/CADKernel/Mesh/Meshers/CycleTriangulator.cpp b/Engine/Source/Runtime/Datasmith/CADKernel/Private/CADKernel/Mesh/Meshers/CycleTriangulator.cpp
--- a/Engine/Source/Runtime/Datasmith/CADKernel/Private/CADKernel/Mesh/Meshers/CycleTriangulator.cpp
+++ b/Engine/Source/Runtime/Datasmith/CADKernel/Private/CADKernel/Mesh/Meshers/CycleTriangulator.cpp
@@ -510,7 +510,7 @@ bool FCycleTriangulator::FindCandidateNodes(int32 StartIndex)
 			//F3DDebugSession W(TEXT("Next Segment"));
 			Grid.DisplayIsoNode(EGridSpace::UniformScaled, *CNode.Node, 0, CNode.bIsValid ? EVisuProperty::BluePoint : EVisuProperty::PinkPoint);
 		}
-		//Wait();
+		Wait(false);
 	}
 #endif
 
@@ -660,7 +660,7 @@ bool FCycleTriangulator::FindTheBestCandidateNode()
 		F3DDebugSession _(FString::Printf(TEXT("Selected triangle %d"), ++TriangleIndex2));
 		Grid.DisplayIsoSegment(EGridSpace::UniformScaled, *FirstSideStartNode, *Node, 1, EVisuProperty::BlueCurve);
 		Grid.DisplayIsoSegment(EGridSpace::UniformScaled, *FirstSideEndNode, *Node, 2, EVisuProperty::BlueCurve);
-		//Wait();
+		Wait(false);
 	}
 #endif
 
@@ -769,7 +769,7 @@ bool FCycleTriangulator::FindTheBestCandidateNode()
 		F3DDebugSession _(FString::Printf(TEXT("Selected triangle %d"), ++TriangleIndex));
 		Grid.DisplayIsoSegment(EGridSpace::UniformScaled, *FirstSideStartNode, *Node, 1, EVisuProperty::YellowCurve);
 		Grid.DisplayIsoSegment(EGridSpace::UniformScaled, *FirstSideEndNode, *Node, 2, EVisuProperty::YellowCurve);
-		if (TriangleIndex == 1000)
+		if (TriangleIndex > 1000)
 		{
 			Wait();
 		}
@@ -810,11 +810,13 @@ void FCycleTriangulator::ValidateAddNodesAccordingSlopeWithSide(FAdditionalIso&
 		if (SlopeAtStart < Slope::NullSlope)
 		{
 			SlopeAtStart = Slope::PiSlope;
+			Side.bForceNodes = false;
 		}
 
 		if (SlopeAtEnd < Slope::NullSlope)
 		{
 			SlopeAtEnd = Slope::PiSlope;
+			Side.bForceNodes = false;
 		}
 
 #ifdef DEBUG_COMPUTE_SLOTE_CRITERION
diff --git a/Engine/Source/Runtime/Datasmith/CADKernel/Private/CADKernel/Mesh/Meshers/IsoTriangulator.cpp b/Engine/Source/Runtime/Datasmith/CADKernel/Private/CADKernel/Mesh/Meshers/IsoTriangulator.cpp
--- a/Engine/Source/Runtime/Datasmith/CADKernel/Private/CADKernel/Mesh/Meshers/IsoTriangulator.cpp
+++ b/Engine/Source/Runtime/Datasmith/CADKernel/Private/CADKernel/Mesh/Meshers/IsoTriangulator.cpp
@@ -62,7 +62,6 @@ FIsoTriangulator::FIsoTriangulator(FGrid& InGrid, FFaceMesh& OutMesh, const FMes
 	, Mesh(OutMesh)
 	, LoopSegmentsIntersectionTool(InGrid)
 	, InnerSegmentsIntersectionTool(InGrid)
-	, InnerToLoopSegmentsIntersectionTool(InGrid)
 	, InnerToOuterIsoSegmentsIntersectionTool(InGrid, /*AllowOverlaping*/ true)
 	, ThinZoneIntersectionTool(InGrid)
 	, Tolerances(InTolerances)
@@ -172,14 +171,18 @@ bool FIsoTriangulator::Triangulate()
 	// =============================================================================================================
 #ifdef DEBUG_ISOTRIANGULATOR
 	Grid.DisplayIsoSegments(TEXT("FIsoTrianguler::Final To Loops Before"), DisplaySpace, FinalToLoops, false, false, false, EVisuProperty::YellowCurve);
-	//Wait(bDisplay);
 #endif
 
 	ConnectCellLoops();
 
 #ifdef DEBUG_ISOTRIANGULATOR
-	Grid.DisplayIsoSegments(TEXT("FIsoTrianguler::Final Iso ToLink Inner To Loops"), DisplaySpace, FinalToLoops, false, false, false, EVisuProperty::YellowCurve);
-	//Wait(bDisplay);
+	Grid.DisplayIsoSegments(TEXT("FIsoTrianguler::Final Iso ToLink Inner To Loops"), DisplaySpace, FinalToLoops, false, false, true, EVisuProperty::YellowCurve);
+#endif
+
+	AddCandidateSegmentsToLinkInnerAndLoop();
+
+#ifdef DEBUG_ISOTRIANGULATOR
+	Grid.DisplayIsoSegments(TEXT("FIsoTrianguler::Final Segments"), DisplaySpace, FinalToLoops, false, false, true, EVisuProperty::YellowCurve);
 #endif
 
 	// =============================================================================================================
@@ -942,6 +945,8 @@ void FIsoTriangulator::ConnectCellLoops()
 		Cell.InitLoopConnexions();
 	}
 
+	InnerToLoopCandidateSegments.Reserve(Cells.Num() * 2);
+
 #ifdef DEBUG_CONNECT_CELL_LOOPS
 	F3DDebugSession _(bDisplay, ("ConnectCellLoops"));
 #endif
@@ -952,7 +957,6 @@ void FIsoTriangulator::ConnectCellLoops()
 #ifdef DEBUG_CONNECT_CELL_LOOPS
 		F3DDebugSession _(bDisplay, FString::Printf(TEXT("Cell %d"), Cell.Id));
 		DisplayCell(Cell);
-		Wait(bDisplay);
 #endif
 
 		if (Cell.LoopCells.Num())
@@ -1101,228 +1105,6 @@ void FIsoTriangulator::FindCellContainingBoundaryNodes(TArray<FCell>& Cells)
 	FChrono::Elapse(StartTime);
 }
 
-#ifdef OLD
-void FIsoTriangulator::FindCandidateSegmentsToLinkInnerAndLoop()
-{
-	const double FlatAngle = 0.1;
-
-#ifdef CADKERNEL_DEV
-	FTimePoint StartTime = FChrono::Now();
-#endif
-	TFunction<void(FIsoInnerNode*, FLoopNode&)> CreateCandidateSegment = [&](FIsoNode* InnerNode, FLoopNode& LoopNode)
-	{
-		FIsoSegment& SegCandidate = IsoSegmentFactory.New();
-		SegCandidate.Init(*InnerNode, LoopNode, ESegmentType::InnerToLoop);
-	};
-
-	TFunction<void(FLoopNode&, FLoopNode&)> CreateCandidateBoundarySegment = [&](FLoopNode& StartNode, FLoopNode& EndNode)
-	{
-		FIsoSegment& SegCandidate = IsoSegmentFactory.New();
-		SegCandidate.Init(StartNode, EndNode, ESegmentType::LoopToLoop);
-	};
-
-	int32 CountU = Grid.GetCuttingCount(EIso::IsoU);
-	int32 CountV = Grid.GetCuttingCount(EIso::IsoV);
-
-	// find cell containing boundary nodes
-	TArray<int32> NodeToCellIndices;
-	TArray<int32> SortedIndex;
-	{
-		//F3DDebugSession _(TEXT("New Test"));
-
-		const TArray<double>& IsoUCoordinates = Grid.GetUniformCuttingCoordinatesAlongIso(EIso::IsoU);
-		const TArray<double>& IsoVCoordinates = Grid.GetUniformCuttingCoordinatesAlongIso(EIso::IsoV);
-
-		NodeToCellIndices.Reserve(LoopNodeCount);
-		int32 IndexU = 0;
-		int32 IndexV = 0;
-		int32 Index = 0;
-		for (const FLoopNode& LoopNode : LoopNodes)
-		{
-			if (LoopNode.IsDelete())
-			{
-				continue;
-			}
-
-			const FPoint2D& Coordinate = LoopNode.Get2DPoint(EGridSpace::UniformScaled, Grid);
-
-			ArrayUtils::FindCoordinateIndex(IsoUCoordinates, Coordinate.U, IndexU);
-			ArrayUtils::FindCoordinateIndex(IsoVCoordinates, Coordinate.V, IndexV);
-
-			NodeToCellIndices.Emplace(IndexV * CountU + IndexU);
-			SortedIndex.Emplace(Index++);
-		}
-
-		Algo::Sort(SortedIndex, [&](const int32& Index1, const int32& Index2)
-			{
-				return (NodeToCellIndices[Index1] < NodeToCellIndices[Index2]);
-			});
-	}
-
-
-	int32 CellIndex = -1;
-
-	FIsoInnerNode* Cell[4] = { nullptr, nullptr, nullptr, nullptr };
-	TFunction<void(int32 CellIndex)> GetCellIsoNode = [&](int32 CellIndex)
-	{
-		int32 Index = CellIndex;
-		Cell[0] = GlobalIndexToIsoInnerNodes[Index++];
-		Cell[1] = GlobalIndexToIsoInnerNodes[Index];
-		Index += CountU;
-		Cell[2] = GlobalIndexToIsoInnerNodes[Index--];
-		Cell[3] = GlobalIndexToIsoInnerNodes[Index];
-	};
-
-	// create segment between a boundary node and a cell border
-	//Open3DDebugSession(FString::Printf(TEXT("Cell %d"), CellIndex));
-	for (int32 Index : SortedIndex)
-	{
-		if (CellIndex != NodeToCellIndices[Index])
-		{
-			CellIndex = NodeToCellIndices[Index];
-			GetCellIsoNode(CellIndex);
-
-			//Close3DDebugSession();
-			//Open3DDebugSession(FString::Printf(TEXT("Cell %d"), CellIndex));
-		}
-
-		FLoopNode& LoopPoint = LoopNodes[Index];
-
-		for (int32 ICell = 0; ICell < 4; ++ICell)
-		{
-			if (Cell[ICell])
-			{
-				if (LoopPoint.IsSegmentBeInsideFace(Cell[ICell]->Get2DPoint(EGridSpace::UniformScaled, Grid), Grid, FlatAngle))
-				{
-					continue;
-				}
-
-				if (LoopSegmentsIntersectionTool.DoesIntersect(*Cell[ICell], LoopPoint))
-				{
-					continue;
-				}
-
-				CreateCandidateSegment(Cell[ICell], LoopPoint);
-			}
-		}
-
-		//DisplayPoint(LoopPoint.Get2DPoint(EGridSpace::UniformScaled, Grid), EVisuProperty::YellowPoint);
-	}
-	//Close3DDebugSession();
-
-	// create segment between two boundary nodes
-	CellIndex = -1;
-	//Open3DDebugSession(TEXT("Find in Cell "));
-	//Open3DDebugSession(FString::Printf(TEXT("Cell %d"), 0));
-	for (int32 Index = 0; Index < SortedIndex.Num() - 1; ++Index)
-	{
-		int32 ISortedIndex = SortedIndex[Index];
-		FLoopNode& StartLoop = LoopNodes[ISortedIndex];
-		const FPoint2D& StartPoint = StartLoop.Get2DPoint(EGridSpace::UniformScaled, Grid);
-
-		//if (CellIndex != NodeToCellIndices[ISortedIndex])
-		//{
-		//	Close3DDebugSession();
-		//	Open3DDebugSession(FString::Printf(TEXT("Cell %d"), NodeToCellIndices[ISortedIndex]));
-		//	DisplayPoint(StartPoint, EVisuProperty::YellowPoint);
-		//}
-
-		CellIndex = NodeToCellIndices[ISortedIndex];
-
-		for (int32 Jndex = Index + 1; Jndex < SortedIndex.Num(); ++Jndex)
-		{
-			int32 JSortedIndex = SortedIndex[Jndex];
-			if (CellIndex != NodeToCellIndices[JSortedIndex])
-			{
-				break;
-			}
-			FLoopNode& EndLoop = LoopNodes[JSortedIndex];
-			//DisplayPoint(EndLoop.Get2DPoint(EGridSpace::UniformScaled, Grid), EVisuProperty::YellowPoint);
-
-			if (&EndLoop.GetPreviousNode() == &StartLoop || &EndLoop.GetNextNode() == &StartLoop)
-			{
-				continue;
-			}
-
-			const FPoint2D& EndPoint = EndLoop.Get2DPoint(EGridSpace::UniformScaled, Grid);
-
-			if (StartLoop.IsSegmentBeInsideFace(EndPoint, Grid, FlatAngle))
-			{
-				continue;
-			}
-
-			if (EndLoop.IsSegmentBeInsideFace(StartPoint, Grid, FlatAngle))
-			{
-				continue;
-			}
-
-			if (LoopSegmentsIntersectionTool.DoesIntersect(StartLoop, EndLoop))
-			{
-				continue;
-			}
-
-			CreateCandidateBoundarySegment(StartLoop, EndLoop);
-		}
-	}
-	//Close3DDebugSession();
-	//Close3DDebugSession();
-	//Wait();
-
-	//Display(EGridSpace::UniformScaled, TEXT("TEST"), NewTestSegments, false);
-	//Wait(true);
-
-#ifdef CADKERNEL_DEV
-	Chronos.FindSegmentToLinkInnerToLoopDuration = FChrono::Elapse(StartTime);
-#endif
-}
-#endif
-
-void FIsoTriangulator::ConnectUnconnectedInnerSegments()
-{
-	TArray<FIsoNode*> UnconnectedNode;
-	UnconnectedNode.Reserve(FinalInnerSegments.Num());
-	for (FIsoSegment* Segment : FinalInnerSegments)
-	{
-		if (Segment->GetFirstNode().GetConnectedSegments().Num() == 1)
-		{
-			UnconnectedNode.Add(&Segment->GetFirstNode());
-		}
-	}
-
-	for (FIsoNode* Node : UnconnectedNode)
-	{
-		double MinDistance = HUGE_VALUE;
-		FLoopNode* Candidate = nullptr;
-		for (FLoopNode& LoopNode : LoopNodes)
-		{
-			if (LoopNode.IsDelete())
-			{
-				continue;
-			}
-
-			double Distance = LoopNode.Get2DPoint(EGridSpace::Scaled, Grid).SquareDistance(Node->Get2DPoint(EGridSpace::Scaled, Grid));
-			if (Distance < MinDistance)
-			{
-				if (!InnerToLoopSegmentsIntersectionTool.DoesIntersect(*Node, LoopNode))
-				{
-					MinDistance = Distance;
-					Candidate = &LoopNode;
-				}
-			}
-		}
-
-		if (Candidate)
-		{
-			FIsoSegment& Segment = IsoSegmentFactory.New();
-			Segment.Init(*Node, *Candidate, ESegmentType::InnerToLoop);
-			Segment.ConnectToNode();
-			FinalToLoops.Add(&Segment);
-			InnerToLoopSegmentsIntersectionTool.AddSegment(Segment);
-			InnerToLoopSegmentsIntersectionTool.Sort();
-		}
-	}
-}
-
 bool FIsoTriangulator::CanCycleBeMeshed(const TArray<FIsoSegment*>& Cycle, FIntersectionSegmentTool& CycleIntersectionTool)
 {
 	bool bHasSelfIntersection = true;
@@ -1513,7 +1295,7 @@ void FIsoTriangulator::TriangulateOverCycle(const EGridSpace Space)
 
 #ifdef FIND_CYCLE
 	int32 OverCycleIndex = 0;
-	if (Grid.GetFace()->GetId() == FaceToDebug)
+	if (Grid.bDisplay)
 	{
 		Open3DDebugSession(TEXT("Triangulate Over Cycle"));
 	}
@@ -1528,11 +1310,11 @@ void FIsoTriangulator::TriangulateOverCycle(const EGridSpace Space)
 			bool bLeftSide = true;
 			FindCycle(Segment, bLeftSide, Cycle, CycleOrientation);
 #ifdef FIND_CYCLE
-			if (Grid.GetFace()->GetId() == FaceToDebug)
+			if (Grid.bDisplay)
 			{
 				F3DDebugSession G(TEXT("Find & mesh cycles"));
 				FString Message = FString::Printf(TEXT("MeshCycle - cycle %d"), OverCycleIndex++);
-				Display(EGridSpace::UniformScaled, *Message, Cycle, false);
+				Grid.DisplayIsoSegments(Message, EGridSpace::UniformScaled, Cycle, false);
 			}
 #endif
 			MeshCycle(Cycle, CycleOrientation);
@@ -1551,11 +1333,11 @@ void FIsoTriangulator::TriangulateOverCycle(const EGridSpace Space)
 				continue;
 			}
 #ifdef FIND_CYCLE
-			if (Grid.GetFace()->GetId() == FaceToDebug)
+			if (Grid.bDisplay)
 			{
 				Open3DDebugSession(TEXT("Find & mesh cycles"));
 				FString Message = FString::Printf(TEXT("MeshCycle - cycle %d"), OverCycleIndex++);
-				Display(EGridSpace::UniformScaled, *Message, Cycle, false);
+				Grid.DisplayIsoSegments(Message, EGridSpace::UniformScaled, Cycle, false);
 				Close3DDebugSession();
 			}
 #endif
@@ -1572,11 +1354,11 @@ void FIsoTriangulator::TriangulateOverCycle(const EGridSpace Space)
 				continue;
 			}
 #ifdef FIND_CYCLE
-			if (Grid.GetFace()->GetId() == FaceToDebug)
+			if (Grid.bDisplay)
 			{
 				Open3DDebugSession(TEXT("Find & mesh cycles"));
 				FString Message = FString::Printf(TEXT("MeshCycle - cycle %d"), OverCycleIndex++);
-				Display(EGridSpace::UniformScaled, *Message, Cycle, false);
+				Grid.DisplayIsoSegments(Message, EGridSpace::UniformScaled, Cycle, false);
 				Close3DDebugSession();
 			}
 #endif
@@ -1602,8 +1384,8 @@ bool FIsoTriangulator::FindCycle(FIsoSegment* StartSegment, bool LeftSide, TArra
 
 #ifdef DEBUG_FIND_CYCLE
 	CycleIndex++;
-	//CycleId = CycleIndex;
-	if (Grid.GetFace()->GetId() == FaceToDebug)
+
+	if (Grid.bDisplay)
 	{
 		CycleIndex = CycleId;
 		bDisplay = true;
@@ -1849,6 +1631,7 @@ void FCell::ConnectLoopsByNeighborhood()
 	F3DDebugSession _(Grid.bDisplay, FString::Printf(TEXT("Cell % d: Build Segments to Connect inner loops"), Id));
 #endif
 
+
 	for (FLoopConnexion& LoopConnexion : LoopConnexions)
 	{
 		TryToConnectTwoSubLoopsWithShortestSegment(LoopConnexion);
@@ -1866,13 +1649,15 @@ void FCell::ConnectLoopsByNeighborhood()
 			return LoopConnexion1->MinDistance < LoopConnexion2->MinDistance;
 		});
 
+
 	{
 #ifdef DEBUG_CONNECT_CELL_SUB_LOOPS_BY_NEIGHBORHOOD
 		F3DDebugSession _(TEXT("Best Path"));
 #endif
+		const int32 LoopCount = LoopCells.Num();
 		for (FLoopConnexion* LoopConnexion : LoopConnexionPtrs)
 		{
-			if (!LoopConnexion->bIsConnexionWithBorder && LoopConnexion->IsShortestPath())
+			if (!LoopConnexion->bIsConnexionWithBorder && LoopConnexion->IsShortestPath(LoopCount))
 			{
 				TryToCreateSegment(*LoopConnexion);
 			}
@@ -2352,11 +2137,6 @@ bool FIsoTriangulator::TryToCreateSegment(FCell& Cell, FLoopNode* NodeA, const F
 		return false;
 	}
 
-	if (InnerToLoopSegmentsIntersectionTool.DoesIntersect(*NodeA, *NodeB))
-	{
-		return false;
-	}
-
 	if (Cell.IntersectionTool.DoesIntersect(*NodeA, *NodeB))
 	{
 		return false;
@@ -2442,7 +2222,7 @@ void FIsoTriangulator::ConnectCellCornerToLoops(FCell& Cell)
 		int32 MinIndexA = -1;
 		for (int32 IndexA = 0; IndexA < LoopA.Num(); ++IndexA)
 		{
-			const FLoopNode* NodeA = LoopA[IndexA];
+			FLoopNode* NodeA = LoopA[IndexA];
 			if (NodeA->IsThinZoneNode())
 			{
 				continue;
@@ -2450,6 +2230,19 @@ void FIsoTriangulator::ConnectCellCornerToLoops(FCell& Cell)
 
 			const FPoint2D& ACoordinates = NodeA->Get2DPoint(EGridSpace::UniformScaled, Grid);
 
+			// if the candidate segment is to close to the cell border, we cannot check if in the other side of the cell there are not also a parallel candidate segment
+			// Add it in InnerToLoopCandidateSegments to be processed at the end
+			double SlopeVsCell = ComputeSlopeRelativeToNearestAxis(ACoordinates, InnerCoordinates);
+			if (SlopeVsCell < Slope::OneDegree)
+			{
+				FIsoSegment& Segment = IsoSegmentFactory.New();
+				Segment.Init(*NodeA, *InnerNode, ESegmentType::InnerToLoopU);
+				Segment.SetCandidate();
+
+				InnerToLoopCandidateSegments.Add(&Segment);
+				continue;
+			}
+
 			double SquareDistance = ACoordinates.SquareDistance(InnerCoordinates);
 			if (SquareDistance < MinDistanceSquare)
 			{
@@ -2480,7 +2273,7 @@ void FIsoTriangulator::ConnectCellCornerToLoops(FCell& Cell)
 			if (bDisplay)
 			{
 				F3DDebugSession _(bDisplay, TEXT("Cell corner"));
-				DisplayPoint(CellNodes[ICell]->Get2DPoint(EGridSpace::UniformScaled, Grid), EVisuProperty::GreenPoint);
+				Grid.DisplayIsoNode(EGridSpace::UniformScaled, *CellNodes[ICell], 0, EVisuProperty::RedPoint);
 			}
 #endif
 			for (FLoopCell& LoopCell : Cell.LoopCells)
@@ -2505,4 +2298,62 @@ void FIsoTriangulator::ConnectCellCornerToLoops(FCell& Cell)
 	}
 }
 
+
+void FIsoTriangulator::AddCandidateSegmentsToLinkInnerAndLoop()
+{
+#ifdef DEBUG_ADD_CANDIDATE_SEGMENTS_TO_LINK_INNER_AND_LOOP
+	F3DDebugSession _(Grid.bDisplay, TEXT("AddCandidateSegmentsToLinkInnerAndLoop "));
+	Grid.DisplayIsoSegments(TEXT("CandidateSegments To Link Inner To Loops"), EGridSpace::UniformScaled, InnerToLoopCandidateSegments, false, false, false, EVisuProperty::RedCurve);
+	Wait();
+#endif
+
+	LoopSegmentsIntersectionTool.AddSegments(FinalToLoops);
+#ifdef DEBUG_ADD_CANDIDATE_SEGMENTS_TO_LINK_INNER_AND_LOOP
+	LoopSegmentsIntersectionTool.Display(Grid.bDisplay, TEXT("LoopSegmentsIntersectionTool with final to loop"));
+#endif
+
+	TArray<TPair<double, FIsoSegment*>> LengthOfCandidateSegments;
+	LengthOfCandidateSegments.Reserve(InnerToLoopCandidateSegments.Num());
+	for (FIsoSegment* Segment : InnerToLoopCandidateSegments)
+	{
+		LengthOfCandidateSegments.Emplace(Segment->Get2DLengthSquare(EGridSpace::UniformScaled, Grid), Segment);
+	}
+
+	FIntersectionSegmentTool IntersectionTool(Grid);
+	IntersectionTool.Reserve(InnerToLoopCandidateSegments.Num());
+
+	Algo::Sort(LengthOfCandidateSegments, [&](const TPair<double, FIsoSegment*>& P1, const TPair < double, FIsoSegment*>& P2) { return P1.Key < P2.Key; });
+
+	// Validate the first candidate segments
+	for (const TPair<double, FIsoSegment*>& Candidate : LengthOfCandidateSegments)
+	{
+		FIsoSegment* Segment = Candidate.Value;
+
+		if (IntersectionTool.DoesIntersect(*Segment))
+		{
+			IsoSegmentFactory.DeleteEntity(Segment);
+			continue;
+		}
+
+
+		if (LoopSegmentsIntersectionTool.DoesIntersect(*Segment))
+		{
+			IsoSegmentFactory.DeleteEntity(Segment);
+			continue;
+		}
+
+		if (FIsoSegment::IsItAlreadyDefined(&Segment->GetFirstNode(), &Segment->GetSecondNode()))
+		{
+			IsoSegmentFactory.DeleteEntity(Segment);
+			continue;
+		}
+
+		FinalToLoops.Add(Segment);
+		IntersectionTool.AddSegment(*Segment);
+		Segment->SetSelected();
+		Segment->ConnectToNode();
+	}
+	CandidateSegments.Empty();
+}
+
 } //namespace UE::CADKernel
\ No newline at end of file
diff --git a/Engine/Source/Runtime/Datasmith/CADKernel/Private/CADKernel/Mesh/Meshers/IsoTriangulator.h b/Engine/Source/Runtime/Datasmith/CADKernel/Private/CADKernel/Mesh/Meshers/IsoTriangulator.h
--- a/Engine/Source/Runtime/Datasmith/CADKernel/Private/CADKernel/Mesh/Meshers/IsoTriangulator.h
+++ b/Engine/Source/Runtime/Datasmith/CADKernel/Private/CADKernel/Mesh/Meshers/IsoTriangulator.h
@@ -109,6 +109,8 @@ protected:
 	TArray<FIsoSegment*> FinalInnerSegments;
 	TArray<FIsoSegment*> InnerToOuterSegments;
 
+	TArray<FIsoSegment*> InnerToLoopCandidateSegments;
+
 	/**
 	 * Tool to check if a segment intersect or not existing segments.
 	 * To be optimal, depending on the segment, only a subset of segment is used.
@@ -116,7 +118,6 @@ protected:
 	 */
 	FIntersectionSegmentTool LoopSegmentsIntersectionTool;
 	FIntersectionSegmentTool InnerSegmentsIntersectionTool;
-	FIntersectionSegmentTool InnerToLoopSegmentsIntersectionTool;
 	FIntersectionNodePairTool InnerToOuterIsoSegmentsIntersectionTool;
 	FIntersectionSegmentTool ThinZoneIntersectionTool;
 
@@ -210,12 +211,7 @@ public:
 
 	void ConnectCellCornerToLoops(FCell& Cell);
 
-	void FindCandidateSegmentsToLinkInnerAndLoop();
-
-	/**
-	 * The goal of this algorithm is to connect unconnected inner segment extremity i.e. extremity with one or less connected segment to the closed boundary node
-	 */
-	void ConnectUnconnectedInnerSegments();
+	void AddCandidateSegmentsToLinkInnerAndLoop();
 
 	/**
 	 * Finalize the tessellation between inner grid boundary and loops.
diff --git a/Engine/Source/Runtime/Datasmith/CADKernel/Private/CADKernel/Mesh/Meshers/IsoTriangulator/IsoCell.h b/Engine/Source/Runtime/Datasmith/CADKernel/Private/CADKernel/Mesh/Meshers/IsoTriangulator/IsoCell.h
--- a/Engine/Source/Runtime/Datasmith/CADKernel/Private/CADKernel/Mesh/Meshers/IsoTriangulator/IsoCell.h
+++ b/Engine/Source/Runtime/Datasmith/CADKernel/Private/CADKernel/Mesh/Meshers/IsoTriangulator/IsoCell.h
@@ -68,9 +68,10 @@ struct FLoopConnexion
 		return (Loop == &Loop1) ? &Loop2 : &Loop1;
 	}
 
-	bool IsShortestPath()
+	bool IsShortestPath(const int32 MaxLoopCount)
 	{
 		TMap<const FLoopCell*, double> DistanceToLoops;
+		DistanceToLoops.Reserve(MaxLoopCount);
 		DistanceToLoops.Add(&Loop1, 0.);
 		double MinPathDistance = HUGE_VALUE;
 		while (true)
diff --git a/Engine/Source/Runtime/Datasmith/CADKernel/Private/CADKernel/Mesh/Structure/Grid.cpp b/Engine/Source/Runtime/Datasmith/CADKernel/Private/CADKernel/Mesh/Structure/Grid.cpp
--- a/Engine/Source/Runtime/Datasmith/CADKernel/Private/CADKernel/Mesh/Structure/Grid.cpp
+++ b/Engine/Source/Runtime/Datasmith/CADKernel/Private/CADKernel/Mesh/Structure/Grid.cpp
@@ -1150,21 +1150,24 @@ void FGrid::GetMeshOfThinZone(const FThinZone2D& ThinZone)
 
 	TFunction<void(TArray<FPoint2D>&)> RemoveDuplicatedNode = [](TArray<FPoint2D>& ThinZoneMesh)
 	{
-		if (ThinZoneMesh.Num())
+		if (ThinZoneMesh.Num() > 1)
 		{
 			// Remove Duplicated Node
 			constexpr double Tolerance = DOUBLE_SMALL_NUMBER;
-			FPoint2D* Next = &ThinZoneMesh[0];
-			for (int32 Index = ThinZoneMesh.Num() - 1; Index >= 0; --Index)
+			for (int32 Index = ThinZoneMesh.Num() - 1; Index > 0; --Index)
 			{
-				double SquareDist = Next->SquareDistance(ThinZoneMesh[Index]);
+				double SquareDist = ThinZoneMesh[Index-1].SquareDistance(ThinZoneMesh[Index]);
 				if (SquareDist < Tolerance)
 				{
 					ThinZoneMesh.RemoveAt(Index);
 				}
-				else
+			}
+			if (ThinZoneMesh.Num() > 2)
+			{
+				double SquareDist = ThinZoneMesh[0].SquareDistance(ThinZoneMesh.Last());
+				if (SquareDist < Tolerance)
 				{
-					Next = &ThinZoneMesh[Index];
+					ThinZoneMesh.RemoveAt(ThinZoneMesh.Num()-1);
 				}
 			}
 		}
diff --git a/Engine/Source/Runtime/Datasmith/CADKernel/Private/CADKernel/Topo/Body.cpp b/Engine/Source/Runtime/Datasmith/CADKernel/Private/CADKernel/Topo/Body.cpp
--- a/Engine/Source/Runtime/Datasmith/CADKernel/Private/CADKernel/Topo/Body.cpp
+++ b/Engine/Source/Runtime/Datasmith/CADKernel/Private/CADKernel/Topo/Body.cpp
@@ -13,30 +13,24 @@ void FBody::AddShell(TSharedRef<FShell> Shell)
 	Shell->SetHost(this);
 }
 
-void FBody::RemoveEmptyShell(FModel& Model)
+void FBody::RemoveEmptyShell()
 {
 	TArray<TSharedPtr<FShell>> NewShells;
 	NewShells.Reserve(Shells.Num());
 	for (TSharedPtr<FShell> Shell : Shells)
 	{
-		if (Shell->FaceCount() > 0)
+		if (!Shell->IsDeleted() && Shell->FaceCount() > 0)
 		{
 			NewShells.Emplace(Shell);
 		}
 		else
 		{
 			Shell->Delete();
+			Shell->ResetHost();
 		}
 	}
-	if (NewShells.IsEmpty())
-	{
-		Delete();
-		Model.Remove(this);
-	}
-	else
-	{
-		Swap(NewShells, Shells);
-	}
+
+	Shells = MoveTemp(NewShells);
 }
 
 void FBody::Remove(const FTopologicalShapeEntity* ShellToRemove)
@@ -47,7 +41,10 @@ void FBody::Remove(const FTopologicalShapeEntity* ShellToRemove)
 	}
 
 	int32 Index = Shells.IndexOfByPredicate([&](const TSharedPtr<FShell>& Shell) { return (Shell.Get() == ShellToRemove); });
-	Shells.RemoveAt(Index);
+	if(Index >= 0)
+	{
+		Shells.RemoveAt(Index);
+	}
 }
 
 }
diff --git a/Engine/Source/Runtime/Datasmith/CADKernel/Private/CADKernel/Topo/FaceAnalyzer.cpp b/Engine/Source/Runtime/Datasmith/CADKernel/Private/CADKernel/Topo/FaceAnalyzer.cpp
--- a/Engine/Source/Runtime/Datasmith/CADKernel/Private/CADKernel/Topo/FaceAnalyzer.cpp
+++ b/Engine/Source/Runtime/Datasmith/CADKernel/Private/CADKernel/Topo/FaceAnalyzer.cpp
@@ -239,8 +239,13 @@ bool FFaceAnalyzer::IsThinFace(double& OutGapSize)
 #ifdef DEBUG_THIN_FACE
 	F3DDebugSession _(TEXT("Thin Surface"));
 #endif
+	FTopologicalLoop* Loop = Face.GetExternalLoop().Get();
+	if (!Loop)
+	{
+		return true;
+	}
 
-	Topo::FThinFaceContext Context(*Face.GetExternalLoop());
+	Topo::FThinFaceContext Context(*Loop);
 
 	FTimePoint StartTime = FChrono::Now();
 	BuildLoopSegments(Context);
diff --git a/Engine/Source/Runtime/Datasmith/CADKernel/Private/CADKernel/Topo/Model.cpp b/Engine/Source/Runtime/Datasmith/CADKernel/Private/CADKernel/Topo/Model.cpp
--- a/Engine/Source/Runtime/Datasmith/CADKernel/Private/CADKernel/Topo/Model.cpp
+++ b/Engine/Source/Runtime/Datasmith/CADKernel/Private/CADKernel/Topo/Model.cpp
@@ -38,18 +38,6 @@ bool FModel::Contains(TSharedPtr<FTopologicalEntity> Entity)
 	return false;
 }
 
-void FModel::RemoveEntity(TSharedPtr<FTopologicalEntity> Entity)
-{
-	switch (Entity->GetEntityType())
-	{
-	case EEntity::Body:
-		RemoveBody(StaticCastSharedPtr<FBody>(Entity));
-		break;
-	default:
-		break;
-	}
-}
-
 void FModel::PrintBodyAndShellCount()
 {
 	int32 NbBody = 0;
@@ -74,7 +62,28 @@ int32 FModel::FaceCount() const
 	return FaceCount;
 }
 
-void FModel::GetFaces(TArray<FTopologicalFace*>& OutFaces) 
+void FModel::RemoveEmptyBodies()
+{
+	TArray<TSharedPtr<FBody>> NewBodies;
+	NewBodies.Reserve(Bodies.Num());
+
+	for (const TSharedPtr<FBody>& Body : Bodies)
+	{
+		if (!Body->IsDeleted() && Body->ShellCount())
+		{
+			NewBodies.Add(Body);
+		}
+		else
+		{
+			Body->Delete();
+			Body->ResetHost();
+		}
+	}
+
+	Bodies = MoveTemp(NewBodies);
+}
+
+void FModel::GetFaces(TArray<FTopologicalFace*>& OutFaces)
 {
 	for (const TSharedPtr<FBody>& Body : Bodies)
 	{
diff --git a/Engine/Source/Runtime/Datasmith/CADKernel/Private/CADKernel/Topo/Shell.cpp b/Engine/Source/Runtime/Datasmith/CADKernel/Private/CADKernel/Topo/Shell.cpp
--- a/Engine/Source/Runtime/Datasmith/CADKernel/Private/CADKernel/Topo/Shell.cpp
+++ b/Engine/Source/Runtime/Datasmith/CADKernel/Private/CADKernel/Topo/Shell.cpp
@@ -58,6 +58,25 @@ void FShell::RemoveFaces()
 	TopologicalFaces.Empty();
 }
 
+void FShell::RemoveDeletedOrDegeneratedFaces()
+{
+	TArray<FOrientedFace> NewTopologicalFaces;
+	NewTopologicalFaces.Reserve(TopologicalFaces.Num());
+
+	for (FOrientedFace& Face : TopologicalFaces)
+	{
+		if (!Face.Entity->IsDeletedOrDegenerated())
+		{
+			NewTopologicalFaces.Emplace(Face);
+		}
+		else
+		{
+			Face.Entity->ResetHost();
+		}
+	}
+	TopologicalFaces = MoveTemp(NewTopologicalFaces);
+}
+
 void FShell::Empty()
 {
 	RemoveFaces();
@@ -92,7 +111,10 @@ void FShell::Remove(const FTopologicalShapeEntity* FaceToRemove)
 	}
 
 	int32 Index = TopologicalFaces.IndexOfByPredicate([&](const FOrientedFace& Face) { return (Face.Entity.Get() == FaceToRemove); });
-	TopologicalFaces.RemoveAt(Index);
+	if(Index >= 0)
+	{
+		TopologicalFaces.RemoveAt(Index);
+	}
 }
 
 #ifdef CADKERNEL_DEV
diff --git a/Engine/Source/Runtime/Datasmith/CADKernel/Private/CADKernel/Topo/TopologicalEntity.cpp b/Engine/Source/Runtime/Datasmith/CADKernel/Private/CADKernel/Topo/TopologicalEntity.cpp
deleted file mode 100644
diff --git a/Engine/Source/Runtime/Datasmith/CADKernel/Private/CADKernel/Topo/TopologicalFace.cpp b/Engine/Source/Runtime/Datasmith/CADKernel/Private/CADKernel/Topo/TopologicalFace.cpp
--- a/Engine/Source/Runtime/Datasmith/CADKernel/Private/CADKernel/Topo/TopologicalFace.cpp
+++ b/Engine/Source/Runtime/Datasmith/CADKernel/Private/CADKernel/Topo/TopologicalFace.cpp
@@ -357,6 +357,20 @@ bool FTopologicalFace::IsANonManifoldFace() const
 	return false;
 }
 
+void FTopologicalFace::DeleteNonmanifoldLink()
+{
+	for (const TSharedPtr<FTopologicalLoop>& Loop : GetLoops())
+	{
+		for (const FOrientedEdge& Edge : Loop->GetEdges())
+		{
+			if (Edge.Entity->GetTwinEntityCount() > 2)
+			{
+				Edge.Entity->UnlinkTwinEntities();
+			}
+		}
+	}
+}
+
 bool FTopologicalFace::IsABorderFace() const
 {
 	for (const TSharedPtr<FTopologicalLoop>& Loop : GetLoops())
diff --git a/Engine/Source/Runtime/Datasmith/CADKernel/Private/CADKernel/Topo/TopologicalLoop.cpp b/Engine/Source/Runtime/Datasmith/CADKernel/Private/CADKernel/Topo/TopologicalLoop.cpp
--- a/Engine/Source/Runtime/Datasmith/CADKernel/Private/CADKernel/Topo/TopologicalLoop.cpp
+++ b/Engine/Source/Runtime/Datasmith/CADKernel/Private/CADKernel/Topo/TopologicalLoop.cpp
@@ -50,6 +50,11 @@ TSharedPtr<FTopologicalLoop> FTopologicalLoop::Make(const TArray<TSharedPtr<FTop
 	Loop.EnsureLogicalClosing(GeometricTolerance);
 	Loop.RemoveDegeneratedEdges();
 
+	if (Loop.GetEdges().IsEmpty())
+	{
+		return  TSharedPtr<FTopologicalLoop>();
+	}
+
 	for (FOrientedEdge& OrientedEdge : Loop.GetEdges())
 	{
 		OrientedEdge.Entity->SetLoop(Loop);
@@ -143,6 +148,11 @@ EOrientation FTopologicalLoop::GetDirection(TSharedPtr<FTopologicalEdge>& InEdge
 
 void FTopologicalLoop::Get2DSampling(TArray<FPoint2D>& LoopSampling) const
 {
+	if (Edges.IsEmpty())
+	{
+		return;
+	}
+
 	int32 PointCount = 0;
 	for (const FOrientedEdge& Edge : Edges)
 	{
@@ -161,6 +171,11 @@ void FTopologicalLoop::Get2DSampling(TArray<FPoint2D>& LoopSampling) const
 
 bool FTopologicalLoop::Get2DSamplingWithoutDegeneratedEdges(TArray<FPoint2D>& LoopSampling) const
 {
+	if (Edges.IsEmpty())
+	{
+		return false;
+	}
+
 	double LoopLength = 0;
 	int32 EdgeCount = 0;
 	int32 PointCount = 0;
diff --git a/Engine/Source/Runtime/Datasmith/CADKernel/Private/CADKernel/Topo/Topomaker.cpp b/Engine/Source/Runtime/Datasmith/CADKernel/Private/CADKernel/Topo/Topomaker.cpp
--- a/Engine/Source/Runtime/Datasmith/CADKernel/Private/CADKernel/Topo/Topomaker.cpp
+++ b/Engine/Source/Runtime/Datasmith/CADKernel/Private/CADKernel/Topo/Topomaker.cpp
@@ -542,13 +542,6 @@ FTopomaker::FTopomaker(FSession& InSession, const FTopomakerOptions& InOptions)
 	InitFaces();
 }
 
-FTopomaker::FTopomaker(FSession& InSession, const TArray<TSharedPtr<FTopologicalFace>>& InFaces, const FTopomakerOptions& InOptions)
-	: Session(InSession)
-	, Faces(InFaces)
-{
-	SetTolerance(InOptions);
-}
-
 FTopomaker::FTopomaker(FSession& InSession, const TArray<TSharedPtr<FShell>>& InShells, const FTopomakerOptions& InOptions)
 	: Session(InSession)
 {
@@ -591,43 +584,6 @@ void FTopomaker::InitFaces()
 	}
 }
 
-void FTopomaker::RemoveFacesFromShell()
-{
-	// remove faces from their shells
-	TSet<FShell*> ShellSet;
-	for (const TSharedPtr<FTopologicalFace>& Face : Faces)
-	{
-		if (Face->GetHost() != nullptr)
-		{
-			ShellSet.Add((FShell*)Face->GetHost());
-			Face->ResetHost();
-		}
-	}
-
-	for (FShell* Shell : ShellSet)
-	{
-		const bool bIsOuter = Shell->IsOuter();
-
-		TArray<FOrientedFace> ShellFace;
-		ShellFace.Reserve(Shell->FaceCount());
-		for (const FOrientedFace& Face : Shell->GetFaces())
-		{
-			if (!Face.Entity->GetHost())
-			{
-				if (bIsOuter != (Face.Direction == EOrientation::Front))
-				{
-					Face.Entity->SetBackOriented();
-				}
-			}
-			else
-			{
-				ShellFace.Emplace(Face);
-			}
-		}
-		Shell->ReplaceFaces(ShellFace);
-	}
-}
-
 void FTopomaker::EmptyShells()
 {
 	for (FShell* Shell : Shells)
@@ -1014,6 +970,12 @@ void FTopomaker::RemoveThinFaces()
 	// Find thin faces
 	for (TSharedPtr<FTopologicalFace> Face : Faces)
 	{
+		if (Face->GetLoops().Num() == 0)
+		{
+			Face->SetDeletedMarker();
+			continue;
+		}
+
 		FFaceAnalyzer Analyer(*Face, ThinFaceWidth);
 		double GapSize = 0;
 		if (Analyer.IsThinFace(GapSize))
@@ -1026,6 +988,7 @@ void FTopomaker::RemoveThinFaces()
 			}
 #endif
 			Face->Remove(&BorderEdges);
+			Face->SetDeletedMarker();
 
 			TopomakerTools::GetVertices(BorderEdges, BorderVertices);
 			TopomakerTools::MergeCoincidentVertices(BorderVertices, GapSize * 1.2);
@@ -1114,6 +1077,8 @@ void FTopomaker::SplitIntoConnectedShells()
 {
 	FTimePoint StartTime = FChrono::Now();
 
+	DeleteNonmanifoldLink();
+
 	// Processed1 : Surfaces added in CandidateSurfacesForMesh
 
 	int32 TopologicalFaceCount = Faces.Num();
@@ -1227,14 +1192,7 @@ void FTopomaker::SplitIntoConnectedShells()
 		FaceSubset.SetMainColor(ColorToFaceCount);
 	}
 
-	if (Shells.Num())
-	{
-		EmptyShells();
-	}
-	else
-	{
-		RemoveFacesFromShell();
-	}
+	EmptyShells();
 
 	// for each FaceSubset, process the Shell
 	for (FFaceSubset FaceSubset : SubShells)
@@ -1279,6 +1237,7 @@ void FTopomaker::SplitIntoConnectedShells()
 	UnlinkNonManifoldVertex();
 
 	RemoveEmptyShells();
+	Session.GetModel().RemoveEmptyBodies();
 
 	FDuration StepDuration = FChrono::Elapse(StartTime);
 	FChrono::PrintClockElapse(EVerboseLevel::Log, TEXT(""), TEXT("Split"), StepDuration);
@@ -1311,7 +1270,7 @@ void FTopomaker::RemoveEmptyShells()
 				Body->SetMarker1();
 				Bodies.Add(Body);
 			}
-			Shell->Delete();
+			Shell->SetDeletedMarker();
 		}
 		else
 		{
@@ -1320,10 +1279,9 @@ void FTopomaker::RemoveEmptyShells()
 	}
 	Swap(Shells, NewShells);
 
-	FModel& Model = Session.GetModel();
 	for (FBody* Body : Bodies)
 	{
-		Body->RemoveEmptyShell(Model);
+		Body->RemoveEmptyShell();
 		Body->ResetMarkers();
 	}
 }
@@ -1473,6 +1431,13 @@ void FTopomaker::ResetMarkersOfFaces()
 	}
 }
 
+void FTopomaker::DeleteNonmanifoldLink()
+{
+	for (const TSharedPtr<FTopologicalFace>& Face : Faces)
+	{
+		Face->DeleteNonmanifoldLink();
+	}
+}
 
 void FTopomaker::UnlinkFromOther()
 {
@@ -1504,6 +1469,7 @@ void FTopomaker::UnlinkFromOther()
 		MergeCoincidents(VerticesToLink);
 	}
 
+	if(Shells.IsEmpty())
 	{
 		VerticesToLink.Reset();
 		TArray<FTopologicalFace*> FacePtrs;
diff --git a/Engine/Source/Runtime/Datasmith/CADKernel/Public/CADKernel/Topo/Body.h b/Engine/Source/Runtime/Datasmith/CADKernel/Public/CADKernel/Topo/Body.h
--- a/Engine/Source/Runtime/Datasmith/CADKernel/Public/CADKernel/Topo/Body.h
+++ b/Engine/Source/Runtime/Datasmith/CADKernel/Public/CADKernel/Topo/Body.h
@@ -79,7 +79,7 @@ public:
 
 	void AddShell(TSharedRef<FShell> Shell);
 
-	void RemoveEmptyShell(FModel& Model);
+	void RemoveEmptyShell();
 
 	virtual void Remove(const FTopologicalShapeEntity* ShellToRemove) override;
 
@@ -97,6 +97,16 @@ public:
 		return Shells;
 	}
 
+	int32 ShellCount() const
+	{
+		return Shells.Num();
+	}
+
+	bool IsEmpty() const
+	{
+		return Shells.IsEmpty();
+	}
+
 	virtual int32 FaceCount() const override
 	{
 		int32 FaceCount = 0;
diff --git a/Engine/Source/Runtime/Datasmith/CADKernel/Public/CADKernel/Topo/Model.h b/Engine/Source/Runtime/Datasmith/CADKernel/Public/CADKernel/Topo/Model.h
--- a/Engine/Source/Runtime/Datasmith/CADKernel/Public/CADKernel/Topo/Model.h
+++ b/Engine/Source/Runtime/Datasmith/CADKernel/Public/CADKernel/Topo/Model.h
@@ -96,13 +96,13 @@ public:
 		FTopologicalShapeEntity::Empty();
 	}
 
-	void RemoveBody(TSharedPtr<FBody> InBody)
-	{
-		Bodies.Remove(InBody);
-	}
-
 	virtual void Remove(const FTopologicalShapeEntity* InBody) override
 	{
+		if (!InBody)
+		{
+			return;
+		}
+
 		int32 Index = Bodies.IndexOfByPredicate([&](const TSharedPtr<FBody>& Body) { return (InBody == (FTopologicalShapeEntity*) Body.Get()); });
 		if (Index != INDEX_NONE)
 		{
@@ -110,9 +110,9 @@ public:
 		}
 	}
 
-	void PrintBodyAndShellCount();
+	void RemoveEmptyBodies();
 
-	void RemoveEntity(TSharedPtr<FTopologicalEntity> InEntity);
+	void PrintBodyAndShellCount();
 
 	bool Contains(TSharedPtr<FTopologicalEntity> InEntity);
 
@@ -147,6 +147,16 @@ public:
 		return Bodies.Num();
 	}
 
+	int32 BodyCount() const
+	{
+		return Bodies.Num();
+	}
+
+	bool IsEmpty() const
+	{
+		return Bodies.IsEmpty();
+	}
+
 	virtual void GetFaces(TArray<FTopologicalFace*>& OutFaces) override;
 
 	virtual int32 FaceCount() const override;
diff --git a/Engine/Source/Runtime/Datasmith/CADKernel/Public/CADKernel/Topo/Shell.h b/Engine/Source/Runtime/Datasmith/CADKernel/Public/CADKernel/Topo/Shell.h
--- a/Engine/Source/Runtime/Datasmith/CADKernel/Public/CADKernel/Topo/Shell.h
+++ b/Engine/Source/Runtime/Datasmith/CADKernel/Public/CADKernel/Topo/Shell.h
@@ -42,8 +42,7 @@ public:
 
 class CADKERNEL_API FShell : public FTopologicalShapeEntity
 {
-	friend FEntity;
-	friend class FBody;
+	friend class FEntity;
 
 private:
 	TArray<FOrientedFace> TopologicalFaces;
@@ -95,6 +94,9 @@ public:
 	}
 
 	void RemoveFaces();
+
+	void RemoveDeletedOrDegeneratedFaces();
+
 	virtual void Empty() override;
 
 	void Add(TSharedRef<FTopologicalFace> InTopologicalFace, EOrientation InOrientation);
@@ -117,7 +119,7 @@ public:
 	void ReplaceFaces(TArray<FOrientedFace>& NewFaces)
 	{
 		Swap(TopologicalFaces, NewFaces);
-		NewFaces.Empty();
+		NewFaces.Reset();
 	}
 
 	const TArray<FOrientedFace>& GetFaces() const
diff --git a/Engine/Source/Runtime/Datasmith/CADKernel/Public/CADKernel/Topo/TopologicalEdge.h b/Engine/Source/Runtime/Datasmith/CADKernel/Public/CADKernel/Topo/TopologicalEdge.h
--- a/Engine/Source/Runtime/Datasmith/CADKernel/Public/CADKernel/Topo/TopologicalEdge.h
+++ b/Engine/Source/Runtime/Datasmith/CADKernel/Public/CADKernel/Topo/TopologicalEdge.h
@@ -248,8 +248,16 @@ public:
 	 */
 	void Link(FTopologicalEdge& OtherEdge);
 
+	/**
+	 * Remove the edge and the extremity vertices of the linked entity
+	 * vs UnlinkTwinEntities delete only the edge link
+	 */
 	void Disjoin();
 
+	/**
+	 * Remove the edge and the extremity vertices of the linked entity
+	 * vs UnlinkTwinEntities delete only the edge link
+	 */
 	void Unlink()
 	{
 		Disjoin();
diff --git a/Engine/Source/Runtime/Datasmith/CADKernel/Public/CADKernel/Topo/TopologicalFace.h b/Engine/Source/Runtime/Datasmith/CADKernel/Public/CADKernel/Topo/TopologicalFace.h
--- a/Engine/Source/Runtime/Datasmith/CADKernel/Public/CADKernel/Topo/TopologicalFace.h
+++ b/Engine/Source/Runtime/Datasmith/CADKernel/Public/CADKernel/Topo/TopologicalFace.h
@@ -298,7 +298,6 @@ public:
 	void Remove(TArray<FTopologicalEdge*>* NewBorderEdges = nullptr)
 	{
 		Disjoin(NewBorderEdges);
-		RemoveOfHost();
 		Delete();
 	}
 
@@ -317,7 +316,6 @@ public:
 		}
 		Loops.Empty();
 #endif
-		RemoveOfHost();
 
 		Mesh.Reset();
 		MeshCuttingCoordinates.Empty();
@@ -355,6 +353,13 @@ public:
 	 */
 	void Disjoin(TArray<FTopologicalEdge*>* NewBorderEdges = nullptr);
 
+	/**
+	 * Delete only nonmanifold edge link
+	 * 
+	 */
+	void DeleteNonmanifoldLink();
+
+
 #ifdef CADKERNEL_DEV
 	virtual void FillTopologyReport(FTopologyReport& Report) const override;
 #endif
diff --git a/Engine/Source/Runtime/Datasmith/CADKernel/Public/CADKernel/Topo/TopologicalShapeEntity.h b/Engine/Source/Runtime/Datasmith/CADKernel/Public/CADKernel/Topo/TopologicalShapeEntity.h
--- a/Engine/Source/Runtime/Datasmith/CADKernel/Public/CADKernel/Topo/TopologicalShapeEntity.h
+++ b/Engine/Source/Runtime/Datasmith/CADKernel/Public/CADKernel/Topo/TopologicalShapeEntity.h
@@ -12,6 +12,10 @@ class FTopologyReport;
 
 class CADKERNEL_API FTopologicalShapeEntity : public FTopologicalEntity
 {
+	friend class FBody;
+	friend class FModel;
+	friend class FShell;
+
 private:
 	FTopologicalShapeEntity* HostedBy = nullptr;
 	FMetadataDictionary Dictionary;
@@ -73,13 +77,14 @@ public:
 
 	void SetHost(FTopologicalShapeEntity* Body)
 	{
+		if (HostedBy)
+		{
+			HostedBy->Remove(this);
+		}
 		HostedBy = Body;
 	}
 
-	void ResetHost()
-	{
-		HostedBy = nullptr;
-	}
+	virtual void Remove(const FTopologicalShapeEntity*) = 0;
 
 	void SetHostId(const int32 InHostId)
 	{
@@ -141,21 +146,19 @@ public:
 		return Dictionary.GetPatchId();
 	}
 
-	void RemoveOfHost()
-	{
-		if (FTopologicalShapeEntity* Host = GetHost())
-		{
-			Host->Remove(this);
-			ResetHost();
-		}
-	}
-
-	virtual void Remove(const FTopologicalShapeEntity*) = 0;
-
 #ifdef CADKERNEL_DEV
 	virtual FInfoEntity& GetInfo(FInfoEntity&) const override;
 #endif
 
+protected:
+	/**
+	 * Mandatory: The host must have remove this from his array
+	 */
+	void ResetHost()
+	{
+		HostedBy = nullptr;
+	}
+
 };
 
 } // namespace UE::CADKernel
diff --git a/Engine/Source/Runtime/Datasmith/CADKernel/Public/CADKernel/Topo/Topomaker.h b/Engine/Source/Runtime/Datasmith/CADKernel/Public/CADKernel/Topo/Topomaker.h
--- a/Engine/Source/Runtime/Datasmith/CADKernel/Public/CADKernel/Topo/Topomaker.h
+++ b/Engine/Source/Runtime/Datasmith/CADKernel/Public/CADKernel/Topo/Topomaker.h
@@ -100,7 +100,6 @@ public:
 
 	FTopomaker(FSession& InSession, const FTopomakerOptions& InOptions);
 	FTopomaker(FSession& InSession, const TArray<TSharedPtr<FShell>>& Shells, const FTopomakerOptions& InOptions);
-	FTopomaker(FSession& InSession, const TArray<TSharedPtr<FTopologicalFace>>& Surfaces, const FTopomakerOptions& InOptions);
 
 	void SetTolerance(const FTopomakerOptions Options)
 	{
@@ -146,6 +145,12 @@ public:
 
 	void RemoveThinFaces();
 
+	/**
+	 * Mandatory: 
+	 * UnlinkNonManifoldVertex has to be call after
+	 */
+	void DeleteNonmanifoldLink();
+
 #ifdef CADKERNEL_DEV
 	void PrintSewReport()
 	{
@@ -168,11 +173,6 @@ private:
 
 	void EmptyShells();
 
-	/**
-	 * Remove a subset of faces of their shell and keep the other shell's faces in the shell
-	 */
-	void RemoveFacesFromShell();
-
 	void RemoveEmptyShells();
 
 	void RemoveDuplicatedFaces();
