Remove multiprocessing in FParametricMesher::PreMeshingTasks()

diff --git a/Engine/Source/Runtime/Datasmith/CADKernel/Private/CADKernel/Mesh/Meshers/ParametricMesher.cpp b/Engine/Source/Runtime/Datasmith/CADKernel/Private/CADKernel/Mesh/Meshers/ParametricMesher.cpp
--- a/Engine/Source/Runtime/Datasmith/CADKernel/Private/CADKernel/Mesh/Meshers/ParametricMesher.cpp
+++ b/Engine/Source/Runtime/Datasmith/CADKernel/Private/CADKernel/Mesh/Meshers/ParametricMesher.cpp
@@ -117,13 +117,18 @@ void FParametricMesher::PreMeshingTasks()
 	//      Apply Surface Criteria
 	// ============================================================================================================
 	const double MeshingTolerance = Tolerances.MeshingTolerance;
-
+#ifdef CADKernelMultiThread
 	TArray<FTopologicalFace*>& LocalFaces = Faces;
 	ParallelFor(Faces.Num(), [&LocalFaces, &Criteria, &MeshingTolerance, &bThinZone = bThinZoneMeshing](int32 Index) {
 
 		FProgress _(1, TEXT("Meshing Entities : Apply Surface Criteria"));
 
-		FTopologicalFace* Face = LocalFaces[Index];
+	FTopologicalFace* Face = LocalFaces[Index];
+#else
+	bool bThinZone = bThinZoneMeshing;
+	for (FTopologicalFace* Face : Faces)
+	{ 
+#endif
 		if (Face == nullptr || Face->IsNotMeshable())
 		{
 			return;
@@ -133,7 +138,10 @@ void FParametricMesher::PreMeshingTasks()
 		{
 			Face->ComputeSurfaceSideProperties();
 		}
-	});
+	}
+#ifdef CADKernelMultiThread
+	, EParallelForFlags::Unbalanced);
+#endif
 
 #ifdef CADKERNEL_DEV
 	FMesherReport::GetChronos().ApplyCriteriaDuration = FChrono::Elapse(ApplyCriteriaStartTime);
