Fixing remaining IWYU errors

diff --git a/Engine/Plugins/Animation/ControlRig/Source/ControlRigEditor/Private/Editor/SRigCurveContainer.h b/Engine/Plugins/Animation/ControlRig/Source/ControlRigEditor/Private/Editor/SRigCurveContainer.h
--- a/Engine/Plugins/Animation/ControlRig/Source/ControlRigEditor/Private/Editor/SRigCurveContainer.h
+++ b/Engine/Plugins/Animation/ControlRig/Source/ControlRigEditor/Private/Editor/SRigCurveContainer.h
@@ -17,6 +17,7 @@ class SInlineEditableTextBlock;
 class FControlRigEditor;
 struct FAssetData;
 class URigVMBlueprint;
+class FRigVMEditor;
 
 class FDisplayedRigCurveInfo
 {
diff --git a/Engine/Plugins/Experimental/Mutable/Source/CustomizableObject/Private/MuCO/CustomizableObjectPrivate.h b/Engine/Plugins/Experimental/Mutable/Source/CustomizableObject/Private/MuCO/CustomizableObjectPrivate.h
--- a/Engine/Plugins/Experimental/Mutable/Source/CustomizableObject/Private/MuCO/CustomizableObjectPrivate.h
+++ b/Engine/Plugins/Experimental/Mutable/Source/CustomizableObject/Private/MuCO/CustomizableObjectPrivate.h
@@ -4,6 +4,10 @@
 
 #include "Templates/SharedPointer.h"
 
+#if WITH_EDITOR
+#include "Misc/Guid.h"
+#endif
+
 namespace mu { class Model; }
 
 
diff --git a/Engine/Plugins/Experimental/NNERuntimeRDG/Source/NNERuntimeRDG/Private/Dml/NNEDmlOperatorUtils.h b/Engine/Plugins/Experimental/NNERuntimeRDG/Source/NNERuntimeRDG/Private/Dml/NNEDmlOperatorUtils.h
--- a/Engine/Plugins/Experimental/NNERuntimeRDG/Source/NNERuntimeRDG/Private/Dml/NNEDmlOperatorUtils.h
+++ b/Engine/Plugins/Experimental/NNERuntimeRDG/Source/NNERuntimeRDG/Private/Dml/NNEDmlOperatorUtils.h
@@ -2,8 +2,13 @@
 
 #pragma once
 
+#ifdef NNE_USE_DIRECTML
+
 #include "Containers/StringView.h"
-#include "NNEDmlOperator.h"
+#include "NNE.h"
+#include "NNETypes.h"
+#include "Dml/NNEDmlCommon.h"
+#include "Dml/NNEDmlOperator.h"
 
 namespace UE::NNERuntimeRDG::Private::Dml
 {
@@ -233,4 +238,6 @@ inline void SetDmlAxesFromOnnx(Util::FSmallUIntArray& DmlAxes, int32 Rank, TCons
 	}
 }
 
-}
\ No newline at end of file
+}
+
+#endif // NNE_USE_DIRECTML
diff --git a/Engine/Plugins/Experimental/PhysicsControl/Source/PhysicsControl/Private/PhysicsControlActor.cpp b/Engine/Plugins/Experimental/PhysicsControl/Source/PhysicsControl/Private/PhysicsControlActor.cpp
--- a/Engine/Plugins/Experimental/PhysicsControl/Source/PhysicsControl/Private/PhysicsControlActor.cpp
+++ b/Engine/Plugins/Experimental/PhysicsControl/Source/PhysicsControl/Private/PhysicsControlActor.cpp
@@ -9,6 +9,10 @@
 
 #include UE_INLINE_GENERATED_CPP_BY_NAME(PhysicsControlActor)
 
+FInitialCharacterControls::FInitialCharacterControls() = default;
+FInitialCharacterControls::~FInitialCharacterControls() = default;
+
+
 //======================================================================================================================
 void UPhysicsControlInitializerComponent::CreateInitialCharacterControls(UPhysicsControlComponent* ControlComponent)
 {
diff --git a/Engine/Plugins/Experimental/PhysicsControl/Source/PhysicsControl/Public/PhysicsControlActor.h b/Engine/Plugins/Experimental/PhysicsControl/Source/PhysicsControl/Public/PhysicsControlActor.h
--- a/Engine/Plugins/Experimental/PhysicsControl/Source/PhysicsControl/Public/PhysicsControlActor.h
+++ b/Engine/Plugins/Experimental/PhysicsControl/Source/PhysicsControl/Public/PhysicsControlActor.h
@@ -153,7 +153,8 @@ struct PHYSICSCONTROL_API FInitialCharacterControls
 {
 	GENERATED_BODY()
 
-	FInitialCharacterControls() {}
+	FInitialCharacterControls();
+	~FInitialCharacterControls();
 
 	/** The owner of the character skeletal mesh */
 	UPROPERTY(EditAnywhere, BlueprintReadWrite, Category = PhysicsControl)
diff --git a/Engine/Plugins/Experimental/QuicMessaging/Source/QuicMessagingTransport/Private/QuicCertificate.h b/Engine/Plugins/Experimental/QuicMessaging/Source/QuicMessagingTransport/Private/QuicCertificate.h
--- a/Engine/Plugins/Experimental/QuicMessaging/Source/QuicMessagingTransport/Private/QuicCertificate.h
+++ b/Engine/Plugins/Experimental/QuicMessaging/Source/QuicMessagingTransport/Private/QuicCertificate.h
@@ -4,10 +4,11 @@
 
 #include "QuicIncludes.h"
 
+#include "Containers/StringConv.h"
 #include "HAL/FileManagerGeneric.h"
 #include "Misc/FileHelper.h"
 #include "Misc/Paths.h"
-
+#include "Templates/Tuple.h"
 
 namespace QuicCertificateUtils
 {
@@ -156,7 +157,7 @@ namespace QuicCertificateUtils
 	/**
 	 * Get the paths for self-signed certificate and key.
 	 */
-	TPair<FString, FString> GetSelfSignedPaths()
+	TTuple<FString, FString> GetSelfSignedPaths()
 	{
 		IPlatformFile& FileManager = FPlatformFileManager::Get().GetPlatformFile();
 
@@ -168,7 +169,7 @@ namespace QuicCertificateUtils
 		FString CertificateFile = CertificatesDir + TEXT("/transportcert.pem");
 		FString PrivateKeyFile = CertificatesDir + TEXT("/transportkey.key");
 
-		return TPair<FString, FString>(CertificateFile, PrivateKeyFile);
+		return TTuple<FString, FString>(CertificateFile, PrivateKeyFile);
 	}
 
 
@@ -193,7 +194,7 @@ namespace QuicCertificateUtils
 			return false;
 		}
 
-		TPair<FString, FString> Paths = GetSelfSignedPaths();
+		TTuple<FString, FString> Paths = GetSelfSignedPaths();
 
 		FString CertificateFile = Paths.Key;
 		FString PrivateKeyFile = Paths.Value;
diff --git a/Engine/Plugins/Experimental/QuicMessaging/Source/QuicMessagingTransport/Private/QuicEndpointManager.cpp b/Engine/Plugins/Experimental/QuicMessaging/Source/QuicMessagingTransport/Private/QuicEndpointManager.cpp
--- a/Engine/Plugins/Experimental/QuicMessaging/Source/QuicMessagingTransport/Private/QuicEndpointManager.cpp
+++ b/Engine/Plugins/Experimental/QuicMessaging/Source/QuicMessagingTransport/Private/QuicEndpointManager.cpp
@@ -352,7 +352,7 @@ void FQuicEndpointManager::InitializeServer()
 			return;
 		}
 
-		TPair<FString, FString> SelfSigned = QuicCertificateUtils::GetSelfSignedPaths();
+		TTuple<FString, FString> SelfSigned = QuicCertificateUtils::GetSelfSignedPaths();
 		Certificate = SelfSigned.Key;
 		PrivateKey = SelfSigned.Value;
 
diff --git a/Engine/Plugins/FX/Niagara/Source/Niagara/Private/NiagaraSimpleObjectPool.h b/Engine/Plugins/FX/Niagara/Source/Niagara/Private/NiagaraSimpleObjectPool.h
--- a/Engine/Plugins/FX/Niagara/Source/Niagara/Private/NiagaraSimpleObjectPool.h
+++ b/Engine/Plugins/FX/Niagara/Source/Niagara/Private/NiagaraSimpleObjectPool.h
@@ -3,6 +3,7 @@
 #pragma once
 
 #include "CoreMinimal.h"
+#include "UObject/Class.h"
 
 /**
  * Very simple object pool that will hold weak references to objects that are returned, these can then be recycled rather than creating new objects.
diff --git a/Engine/Plugins/Media/BlackmagicMedia/Source/BlackmagicCore/BlackmagicCore.build.cs b/Engine/Plugins/Media/BlackmagicMedia/Source/BlackmagicCore/BlackmagicCore.build.cs
--- a/Engine/Plugins/Media/BlackmagicMedia/Source/BlackmagicCore/BlackmagicCore.build.cs
+++ b/Engine/Plugins/Media/BlackmagicMedia/Source/BlackmagicCore/BlackmagicCore.build.cs
@@ -1,7 +1,5 @@
 // Copyright Epic Games, Inc. All Rights Reserved.
 
-using Microsoft.CodeAnalysis;
-using System.IO;
 using UnrealBuildTool;
 
 public class BlackmagicCore : ModuleRules
diff --git a/Engine/Plugins/Media/BlackmagicMedia/Source/BlackmagicCore/Private/Platform/Linux/LinuxPlatform.h b/Engine/Plugins/Media/BlackmagicMedia/Source/BlackmagicCore/Private/Platform/Linux/LinuxPlatform.h
--- a/Engine/Plugins/Media/BlackmagicMedia/Source/BlackmagicCore/Private/Platform/Linux/LinuxPlatform.h
+++ b/Engine/Plugins/Media/BlackmagicMedia/Source/BlackmagicCore/Private/Platform/Linux/LinuxPlatform.h
@@ -2,6 +2,8 @@
 
 #pragma once
 
+#include "HAL/Platform.h"
+
 #if PLATFORM_LINUX
 
 #include <codecvt>
@@ -10,9 +12,6 @@
 
 #include "LinuxCOM.h"
 
-// provide string the way Unreal Engine expects
-#define TCHAR char16_t
-
 // convert Windows style enums in Linux, i.e. BMDPixelFormat to _BMDPixelFormat
 #define ENUM(x) _ ## x
 
diff --git a/Engine/Plugins/Media/BlackmagicMedia/Source/BlackmagicCore/Private/targetver.h b/Engine/Plugins/Media/BlackmagicMedia/Source/BlackmagicCore/Private/targetver.h
deleted file mode 100644
diff --git a/Engine/Plugins/Online/OnlineServices/Source/OnlineServicesCommon/Public/Online/OnlineComponentRegistry.h b/Engine/Plugins/Online/OnlineServices/Source/OnlineServicesCommon/Public/Online/OnlineComponentRegistry.h
--- a/Engine/Plugins/Online/OnlineServices/Source/OnlineServicesCommon/Public/Online/OnlineComponentRegistry.h
+++ b/Engine/Plugins/Online/OnlineServices/Source/OnlineServicesCommon/Public/Online/OnlineComponentRegistry.h
@@ -5,6 +5,7 @@
 
 #include "Online/OnlineTypeInfo.h"
 #include "Online/OnlineMeta.h"
+#include "Templates/SharedPointer.h"
 
 namespace UE::Online { class IOnlineComponent; }
 
diff --git a/Engine/Plugins/Runtime/ModelViewViewModel/Source/ModelViewViewModelEditor/Private/MVVMFieldVariantListTraits.h b/Engine/Plugins/Runtime/ModelViewViewModel/Source/ModelViewViewModelEditor/Private/MVVMFieldVariantListTraits.h
--- a/Engine/Plugins/Runtime/ModelViewViewModel/Source/ModelViewViewModelEditor/Private/MVVMFieldVariantListTraits.h
+++ b/Engine/Plugins/Runtime/ModelViewViewModel/Source/ModelViewViewModelEditor/Private/MVVMFieldVariantListTraits.h
@@ -161,7 +161,7 @@ struct TListTypeTraits<FMVVMBlueprintPropertyPath>
 
 	static FString DebugDump(FMVVMBlueprintPropertyPath InPtr)
 	{
-		return InPtr.GetBasePropertyPath();
+		return FString();
 	}
 
 	class SerializerType {};
diff --git a/Engine/Plugins/Runtime/nDisplay/Source/DisplayCluster/Private/Game/EngineClasses/Scene/DisplayClusterRootActorPreviewRenderingManager.cpp b/Engine/Plugins/Runtime/nDisplay/Source/DisplayCluster/Private/Game/EngineClasses/Scene/DisplayClusterRootActorPreviewRenderingManager.cpp
--- a/Engine/Plugins/Runtime/nDisplay/Source/DisplayCluster/Private/Game/EngineClasses/Scene/DisplayClusterRootActorPreviewRenderingManager.cpp
+++ b/Engine/Plugins/Runtime/nDisplay/Source/DisplayCluster/Private/Game/EngineClasses/Scene/DisplayClusterRootActorPreviewRenderingManager.cpp
@@ -55,6 +55,13 @@ using namespace UE::DisplayCluster::DCRAPreview;
 //------------------------------------------------------------------
 // FDisplayClusterRootActorPreviewObject
 //------------------------------------------------------------------
+
+FDisplayClusterRootActorPreviewObject::FDisplayClusterRootActorPreviewObject() = default;
+FDisplayClusterRootActorPreviewObject::FDisplayClusterRootActorPreviewObject(ADisplayClusterRootActor* InRootActor)
+	: RootActorWeakPtr(InRootActor)
+{ }
+
+
 void FDisplayClusterRootActorPreviewObject::HandleRenderRequest()
 {
 	if (!GDisplayClusterRootActorPreviewRenderingEnable)
diff --git a/Engine/Plugins/Runtime/nDisplay/Source/DisplayCluster/Private/Game/EngineClasses/Scene/DisplayClusterRootActorPreviewRenderingManager.h b/Engine/Plugins/Runtime/nDisplay/Source/DisplayCluster/Private/Game/EngineClasses/Scene/DisplayClusterRootActorPreviewRenderingManager.h
--- a/Engine/Plugins/Runtime/nDisplay/Source/DisplayCluster/Private/Game/EngineClasses/Scene/DisplayClusterRootActorPreviewRenderingManager.h
+++ b/Engine/Plugins/Runtime/nDisplay/Source/DisplayCluster/Private/Game/EngineClasses/Scene/DisplayClusterRootActorPreviewRenderingManager.h
@@ -29,10 +29,8 @@ enum class EDisplayClusterRootActorPreviewEvent : uint8
 class FDisplayClusterRootActorPreviewObject
 {
 public:
-	FDisplayClusterRootActorPreviewObject() = default;
-	FDisplayClusterRootActorPreviewObject(ADisplayClusterRootActor* InRootActor)
-		: RootActorWeakPtr(InRootActor)
-	{ }
+	FDisplayClusterRootActorPreviewObject();
+	FDisplayClusterRootActorPreviewObject(ADisplayClusterRootActor* InRootActor);
 
 public:
 	/** Returns true if the DCRA object still exists. */
diff --git a/Engine/Plugins/Runtime/nDisplay/Source/DisplayClusterWarp/Private/WarpBlend/Math/DisplayClusterWarpBlendMath_WarpMap.h b/Engine/Plugins/Runtime/nDisplay/Source/DisplayClusterWarp/Private/WarpBlend/Math/DisplayClusterWarpBlendMath_WarpMap.h
--- a/Engine/Plugins/Runtime/nDisplay/Source/DisplayClusterWarp/Private/WarpBlend/Math/DisplayClusterWarpBlendMath_WarpMap.h
+++ b/Engine/Plugins/Runtime/nDisplay/Source/DisplayClusterWarp/Private/WarpBlend/Math/DisplayClusterWarpBlendMath_WarpMap.h
@@ -4,6 +4,7 @@
 
 #include "CoreMinimal.h"
 #include "Containers/DisplayClusterWarpAABB.h"
+#include "Render/Containers/IDisplayClusterRender_Texture.h"
 
 /**
  * Implement math for WarpMap texture as a source of geometry
diff --git a/Engine/Source/Runtime/AVEncoder/Public/AudioTypes.h b/Engine/Source/Runtime/AVEncoder/Public/AudioTypes.h
--- a/Engine/Source/Runtime/AVEncoder/Public/AudioTypes.h
+++ b/Engine/Source/Runtime/AVEncoder/Public/AudioTypes.h
@@ -4,6 +4,7 @@
 // ue
 #include "Containers/ArrayView.h"
 #include "Misc/Optional.h"
+#include "Templates/Function.h"
 
 // eossdk
 //#include "eos_platform_prereqs.h"
diff --git a/Engine/Source/Runtime/Core/Private/ProfilingDebugging/Unix/UnixMemoryTrace.cpp b/Engine/Source/Runtime/Core/Private/ProfilingDebugging/Unix/UnixMemoryTrace.cpp
--- a/Engine/Source/Runtime/Core/Private/ProfilingDebugging/Unix/UnixMemoryTrace.cpp
+++ b/Engine/Source/Runtime/Core/Private/ProfilingDebugging/Unix/UnixMemoryTrace.cpp
@@ -6,6 +6,7 @@
 
 #include "CoreTypes.h"
 #include "HAL/Platform.h"
+#include <stdio.h>
 
 ////////////////////////////////////////////////////////////////////////////////
 FMalloc* MemoryTrace_CreateInternal(FMalloc*, int, const ANSICHAR* const*);
diff --git a/Engine/Source/Runtime/Datasmith/CADKernel/Public/CADKernel/Mesh/Meshers/Mesher.h b/Engine/Source/Runtime/Datasmith/CADKernel/Public/CADKernel/Mesh/Meshers/Mesher.h
--- a/Engine/Source/Runtime/Datasmith/CADKernel/Public/CADKernel/Mesh/Meshers/Mesher.h
+++ b/Engine/Source/Runtime/Datasmith/CADKernel/Public/CADKernel/Mesh/Meshers/Mesher.h
@@ -7,6 +7,7 @@
 namespace UE::CADKernel
 {
 class FModelMesh;
+class FTopologicalShapeEntity;
 
 class CADKERNEL_API FMesher
 {
diff --git a/Engine/Source/Runtime/Datasmith/CADKernel/Public/CADKernel/Topo/TopologicalLink.h b/Engine/Source/Runtime/Datasmith/CADKernel/Public/CADKernel/Topo/TopologicalLink.h
--- a/Engine/Source/Runtime/Datasmith/CADKernel/Public/CADKernel/Topo/TopologicalLink.h
+++ b/Engine/Source/Runtime/Datasmith/CADKernel/Public/CADKernel/Topo/TopologicalLink.h
@@ -201,67 +201,4 @@ public:
 	}
 };
 
-/**
- * TTopologicalLink overload dedicated to FVertex to manage the barycenter of twin vertices
- */
-class CADKERNEL_API FVertexLink : public TTopologicalLink<FTopologicalVertex>
-{
-	friend class FTopologicalVertex;
-
-protected:
-	FPoint Barycenter;
-
-	void SetBarycenter(const FPoint& Point)
-	{
-		Barycenter = Point;
-	}
-
-public:
-	FVertexLink()
-		: Barycenter(FPoint::ZeroPoint)
-	{
-	}
-
-	FVertexLink(FTopologicalVertex& Entity)
-		: TTopologicalLink<FTopologicalVertex>(Entity)
-		, Barycenter(FPoint::ZeroPoint)
-	{
-	}
-
-	virtual void Serialize(FCADKernelArchive& Ar) override
-	{
-		TTopologicalLink<FTopologicalVertex>::Serialize(Ar);
-		Ar << Barycenter;
-	}
-
-#ifdef CADKERNEL_DEV
-	virtual FInfoEntity& GetInfo(FInfoEntity& Info) const override;
-#endif
-
-	const FPoint& GetBarycenter() const
-	{
-		return Barycenter;
-	}
-
-	virtual bool CleanLink() override
-	{
-		if (TTopologicalLink::CleanLink())
-		{
-			ComputeBarycenter();
-			DefineActiveEntity();
-			return true;
-		}
-		return false;
-	}
-
-	virtual EEntity GetEntityType() const override
-	{
-		return EEntity::VertexLink;
-	}
-
-	void ComputeBarycenter();
-	void DefineActiveEntity();
-};
-
-
 } // namespace UE::CADKernel
diff --git a/Engine/Source/Runtime/Datasmith/CADKernel/Public/CADKernel/Topo/TopologicalVertex.h b/Engine/Source/Runtime/Datasmith/CADKernel/Public/CADKernel/Topo/TopologicalVertex.h
--- a/Engine/Source/Runtime/Datasmith/CADKernel/Public/CADKernel/Topo/TopologicalVertex.h
+++ b/Engine/Source/Runtime/Datasmith/CADKernel/Public/CADKernel/Topo/TopologicalVertex.h
@@ -15,8 +15,71 @@ class FDatabase;
 class FModelMesh;
 class FTopologicalEdge;
 class FVertexMesh;
+class FTopologicalVertex;
 
-class CADKERNEL_API FTopologicalVertex : public TLinkable<FTopologicalVertex, FVertexLink>
+/**
+ * TTopologicalLink overload dedicated to FVertex to manage the barycenter of twin vertices
+ */
+class FVertexLink : public TTopologicalLink<FTopologicalVertex>
+{
+	friend class FTopologicalVertex;
+
+protected:
+	FPoint Barycenter;
+
+	void SetBarycenter(const FPoint& Point)
+	{
+		Barycenter = Point;
+	}
+
+public:
+	FVertexLink()
+		: Barycenter(FPoint::ZeroPoint)
+	{
+	}
+
+	FVertexLink(FTopologicalVertex& Entity)
+		: TTopologicalLink<FTopologicalVertex>(Entity)
+		, Barycenter(FPoint::ZeroPoint)
+	{
+	}
+
+	virtual void Serialize(FCADKernelArchive& Ar) override
+	{
+		TTopologicalLink<FTopologicalVertex>::Serialize(Ar);
+		Ar << Barycenter;
+	}
+
+#ifdef CADKERNEL_DEV
+	CADKERNEL_API virtual FInfoEntity& GetInfo(FInfoEntity& Info) const override;
+#endif
+
+	const FPoint& GetBarycenter() const
+	{
+		return Barycenter;
+	}
+
+	virtual bool CleanLink() override
+	{
+		if (TTopologicalLink::CleanLink())
+		{
+			ComputeBarycenter();
+			DefineActiveEntity();
+			return true;
+		}
+		return false;
+	}
+
+	virtual EEntity GetEntityType() const override
+	{
+		return EEntity::VertexLink;
+	}
+
+	CADKERNEL_API void ComputeBarycenter();
+	CADKERNEL_API void DefineActiveEntity();
+};
+
+class FTopologicalVertex : public TLinkable<FTopologicalVertex, FVertexLink>
 {
 	friend class FEntity;
 	friend class FVertexLink;
@@ -60,10 +123,10 @@ public:
 		SerializeIdents(Ar, ConnectedEdges);
 	}
 
-	virtual void SpawnIdent(FDatabase& Database) override;
+	CADKERNEL_API virtual void SpawnIdent(FDatabase& Database) override;
 
 #ifdef CADKERNEL_DEV
-	virtual FInfoEntity& GetInfo(FInfoEntity&) const override;
+	CADKERNEL_API virtual FInfoEntity& GetInfo(FInfoEntity&) const override;
 #endif
 
 	virtual EEntity GetEntityType() const override
@@ -129,7 +192,7 @@ public:
 		return Coordinates.SquareDistance(Point);
 	}
 
-	FVertexMesh& GetOrCreateMesh(FModelMesh& MeshModel);
+	CADKERNEL_API FVertexMesh& GetOrCreateMesh(FModelMesh& MeshModel);
 
 	const FVertexMesh* GetMesh() const
 	{
@@ -144,9 +207,9 @@ public:
 		return nullptr;
 	}
 
-	void Link(FTopologicalVertex& InEntity);
+	CADKERNEL_API void Link(FTopologicalVertex& InEntity);
 
-	void UnlinkTo(FTopologicalVertex& Entity);
+	CADKERNEL_API void UnlinkTo(FTopologicalVertex& Entity);
 
 	virtual void RemoveFromLink() override
 	{
@@ -182,10 +245,10 @@ public:
 		TLinkable<FTopologicalVertex, FVertexLink>::Empty();
 	}
 
-	bool IsBorderVertex() const;
+	CADKERNEL_API bool IsBorderVertex() const;
 
-	void AddConnectedEdge(FTopologicalEdge& Edge);
-	void RemoveConnectedEdge(FTopologicalEdge& Edge);
+	CADKERNEL_API void AddConnectedEdge(FTopologicalEdge& Edge);
+	CADKERNEL_API void RemoveConnectedEdge(FTopologicalEdge& Edge);
 
 	/**
 	 * Mandatory: to browse all the connected edges, you have to browse the connected edges of all the twin vertices
@@ -202,7 +265,7 @@ public:
 		return ConnectedEdges;
 	}
 
-	const FTopologicalFace* GetFace() const;
+	CADKERNEL_API const FTopologicalFace* GetFace() const;
 
 	void GetConnectedEdges(TArray<FTopologicalEdge*>& OutConnectedEdges) const
 	{
@@ -240,7 +303,7 @@ public:
 	/**
 	 *
 	 */
-	void GetConnectedEdges(const FTopologicalVertex& OtherVertex, TArray<FTopologicalEdge*>& Edges) const;
+	CADKERNEL_API void GetConnectedEdges(const FTopologicalVertex& OtherVertex, TArray<FTopologicalEdge*>& Edges) const;
 };
 
 } // namespace UE::CADKernel
