CADKernel Update:

diff --git a/Engine/Plugins/Enterprise/DatasmithCADImporter/Source/DatasmithDispatcher/Public/DatasmithDispatcher.h b/Engine/Plugins/Enterprise/DatasmithCADImporter/Source/DatasmithDispatcher/Public/DatasmithDispatcher.h
--- a/Engine/Plugins/Enterprise/DatasmithCADImporter/Source/DatasmithDispatcher/Public/DatasmithDispatcher.h
+++ b/Engine/Plugins/Enterprise/DatasmithCADImporter/Source/DatasmithDispatcher/Public/DatasmithDispatcher.h
@@ -1,6 +1,7 @@
 // Copyright Epic Games, Inc. All Rights Reserved.
 #pragma once
 
+#include "CADOptions.h"
 #include "DatasmithDispatcherTask.h"
 #include "DatasmithWorkerHandler.h"
 
@@ -29,6 +30,11 @@ public:
 
 	void LogMessages(const TArray<TPair<uint8, FString>>& Warnings) const;
 
+	void SetWorkerCount(const int32 InNumberOfWorkers)
+	{
+		CADLibrary::GMaxImportThreads = InNumberOfWorkers;
+	}
+
 private:
 	void SpawnHandlers();
 	int32 GetNextWorkerId();
diff --git a/Engine/Source/Runtime/Datasmith/CADKernel/Private/CADKernel/Mesh/Meshers/CycleTriangulator.cpp b/Engine/Source/Runtime/Datasmith/CADKernel/Private/CADKernel/Mesh/Meshers/CycleTriangulator.cpp
--- a/Engine/Source/Runtime/Datasmith/CADKernel/Private/CADKernel/Mesh/Meshers/CycleTriangulator.cpp
+++ b/Engine/Source/Runtime/Datasmith/CADKernel/Private/CADKernel/Mesh/Meshers/CycleTriangulator.cpp
@@ -25,7 +25,7 @@ FCycleTriangulator::FCycleTriangulator(FIsoTriangulator& IsoTriangulator, const
 	, InnerToOuterIsoSegmentsIntersectionTool(IsoTriangulator.InnerToOuterIsoSegmentsIntersectionTool)
 	, Mesh(IsoTriangulator.Mesh)
 	, IsoSegmentFactory(IsoTriangulator.IsoSegmentFactory)
-	, CycleIntersectionTool(Grid)
+	, CycleIntersectionTool(Grid, IsoTriangulator.Tolerances.GeometricTolerance)
 {
 }
 
@@ -33,11 +33,11 @@ bool FCycleTriangulator::MeshCycle()
 {
 #ifdef ADD_TRIANGLE_2D
 	static int32 CycleId = 0;
-	F3DDebugSession _(Grid.bDisplay, FString::Printf(TEXT("Mesh cycle %d"), CycleId++));
+	F3DDebugSession _(Grid.bDisplay, FString::Printf(TEXT("Mesh cycle %d"), ++CycleId));
 	if (Grid.bDisplay)
 	{
 		Grid.DisplayIsoSegments(TEXT("Cycle"), EGridSpace::UniformScaled, Cycle);
-		Wait(CycleId == 2);
+		Wait(CycleId == 60);
 	}
 #endif
 
@@ -123,10 +123,6 @@ bool FCycleTriangulator::FindTheCycleToMesh(FIsoSegment* Segment, bool bOrientat
 	{
 		F3DDebugSession W(FString::Printf(TEXT("Start Segment")));
 		Grid.DisplayIsoSegment(EGridSpace::UniformScaled, *Segment, 1, EVisuProperty::BlueCurve);
-		if (CycleToMeshIndex == CycleToMeshIndexToTest)
-		{
-			Wait();
-		}
 	}
 #endif
 
@@ -185,6 +181,13 @@ bool FCycleTriangulator::FindTheCycleToMesh(FIsoSegment* Segment, bool bOrientat
 
 	MeanSquareLength = SquareLengthSum / Index;
 
+#ifdef DEBUG_FIND_THE_CYCLE_TO_MESH
+	if (Grid.bDisplay && CycleToMeshIndex == CycleToMeshIndexToTest)
+	{
+		Wait();
+	}
+#endif
+
 	SubCycleNodeCount = SubCycleNodes.Num();
 	if (SubCycleNodeCount < 3)
 	{
@@ -206,6 +209,25 @@ bool IsCycleInsideCandidate(const FPoint2D& Candidate, const FPoint2D& Start, co
 	return (NewSlop > NextSlop);
 };
 
+bool FCycleTriangulator::ConfirmIntersection(const FIsoNode* Start, const FIsoNode* End, const FIsoNode* Candidate, const FIsoSegment* IntersectedSegment) const
+{
+	const FIsoSegment* StartCandidateSegment = Start->GetSegmentConnectedTo(Candidate);
+	const FIsoSegment* EndCandidateSegment = End->GetSegmentConnectedTo(Candidate);
+	if (StartCandidateSegment == IntersectedSegment || EndCandidateSegment == IntersectedSegment)
+	{
+		const int32 IntersectionCount = CycleIntersectionTool.CountIntersections(Start, End);
+		if (IntersectionCount > 1)
+		{
+			return true;
+		}
+	}
+	else
+	{
+		return true;
+	}
+	return false;
+}
+
 bool FCycleTriangulator::FindTheBestAcuteTriangle()
 {
 	const double MaxEdgeSquareLength = MeanSquareLength * 4.;
@@ -312,13 +334,16 @@ bool FCycleTriangulator::FindTheBestAcuteTriangle()
 			continue;
 		}
 
-		FIsoSegment* Segment;
-		if (CycleIntersectionTool.DoesIntersect(NewSegmentStart, NewSegmentEnd, &Segment))
+		if (const FIsoSegment* Intersection = CycleIntersectionTool.FindIntersectingSegment(NewSegmentStart, NewSegmentEnd))
 		{
-			continue;
+			if (ConfirmIntersection(NewSegmentStart, NewSegmentEnd, CandidateNode, Intersection))
+			{
+				continue;
+			}
 		}
 
-		if (InnerToOuterIsoSegmentsIntersectionTool.DoesIntersect(*NewSegmentStart, *NewSegmentEnd))
+		const int32 NewSegmentIntersectionCount = InnerToOuterIsoSegmentsIntersectionTool.CountIntersections(*NewSegmentStart, *NewSegmentEnd);
+		if (NewSegmentIntersectionCount > IntersectionCountAllowed)
 		{
 			continue;
 		}
@@ -334,67 +359,67 @@ bool FCycleTriangulator::FindTheBestAcuteTriangle()
 	int32 AcuteNodeIndex = -1;
 	constexpr double MaxSlopeToBeSelected = Slope::ThreeQuaterPiSlope;
 
-	for (; false && NIndex < SubCycleNodeCount; ++NIndex)
-	{
-		const TPair<int, double>& Candidate = VertexIndexToSlopes[NIndex];
-
-		if (Candidate.Value > MaxSlopeToBeSelected)
-		{
-			break;
-		}
-
-		GetCandidateTriangleNodes(Candidate.Key);
-
-		// Check that the cycle is not inside the selected triangle
-		if (IsCycleInsideCandidate(*CandidatePoint, *StartPoint, *EndPoint, *NextPoint, TransformIntoCounterClockwise)
-			|| IsCycleInsideCandidate(*CandidatePoint, *EndPoint, *StartPoint, *PrevPoint, TransformIntoClockwise))
-		{
-			continue;
-		}
-
-		FIsoSegment* Segment;
-		if (!CycleIntersectionTool.DoesIntersect(NewSegmentStart, NewSegmentEnd, &Segment))
-		{
-			double IsoscelesCriteria = IsoTriangulatorImpl::IsoscelesCriteriaMin(*StartPoint, *EndPoint, *CandidatePoint);
-
-			if (BestIsoscelesCriteria < IsoscelesCriteria)
-			{
-				double Length1 = StartPoint->SquareDistance(*CandidatePoint);
-
-				if (Length1 > MaxEdgeSquareLength)
-				{
-#ifdef DEBUG_FIND_THE_CYCLE_TO_MESH
-					if (Grid.bDisplay)
-					{
-						F3DDebugSession W(FString::Printf(TEXT("cancel due to Seg Length %f vs %f"), sqrt(Length1), sqrt(MaxEdgeSquareLength)));
-						Grid.DisplayIsoSegment(EGridSpace::UniformScaled, *NewSegmentStart, *CandidateNode, 0, EVisuProperty::YellowCurve);
-						Wait();
-					}
-#endif
-					continue;
-				}
-
-				double Length2 = EndPoint->SquareDistance(*CandidatePoint);
-
-				if (Length2 > MaxEdgeSquareLength)
-				{
-#ifdef DEBUG_FIND_THE_CYCLE_TO_MESH
-					if (Grid.bDisplay)
-					{
-						F3DDebugSession W(FString::Printf(TEXT("cancel due to Seg Length %f vs %f"), sqrt(Length2), sqrt(MaxEdgeSquareLength)));
-						Grid.DisplayIsoSegment(EGridSpace::UniformScaled, *NewSegmentEnd, *CandidateNode, 0, EVisuProperty::YellowCurve);
-						Wait();
-					}
-#endif
-					continue;
-				}
-
-				BestIsoscelesCriteria = IsoscelesCriteria;
-				AcuteNodeIndex = Index;
-				bAcuteTriangle = true;
-			}
-		}
-	}
+//	for (; false && NIndex < SubCycleNodeCount; ++NIndex)
+//	{
+//		const TPair<int, double>& Candidate = VertexIndexToSlopes[NIndex];
+//
+//		if (Candidate.Value > MaxSlopeToBeSelected)
+//		{
+//			break;
+//		}
+//
+//		GetCandidateTriangleNodes(Candidate.Key);
+//
+//		// Check that the cycle is not inside the selected triangle
+//		if (IsCycleInsideCandidate(*CandidatePoint, *StartPoint, *EndPoint, *NextPoint, TransformIntoCounterClockwise)
+//			|| IsCycleInsideCandidate(*CandidatePoint, *EndPoint, *StartPoint, *PrevPoint, TransformIntoClockwise))
+//		{
+//			continue;
+//		}
+//
+//		//FIsoSegment* Segment;
+//		if (!CycleIntersectionTool.DoesIntersect(NewSegmentStart, NewSegmentEnd/*, &Segment*/))
+//		{
+//			double IsoscelesCriteria = IsoTriangulatorImpl::IsoscelesCriteriaMin(*StartPoint, *EndPoint, *CandidatePoint);
+//
+//			if (BestIsoscelesCriteria < IsoscelesCriteria)
+//			{
+//				double Length1 = StartPoint->SquareDistance(*CandidatePoint);
+//
+//				if (Length1 > MaxEdgeSquareLength)
+//				{
+//#ifdef DEBUG_FIND_THE_CYCLE_TO_MESH
+//					if (Grid.bDisplay)
+//					{
+//						F3DDebugSession W(FString::Printf(TEXT("cancel due to Seg Length %f vs %f"), sqrt(Length1), sqrt(MaxEdgeSquareLength)));
+//						Grid.DisplayIsoSegment(EGridSpace::UniformScaled, *NewSegmentStart, *CandidateNode, 0, EVisuProperty::YellowCurve);
+//						Wait();
+//					}
+//#endif
+//					continue;
+//				}
+//
+//				double Length2 = EndPoint->SquareDistance(*CandidatePoint);
+//
+//				if (Length2 > MaxEdgeSquareLength)
+//				{
+//#ifdef DEBUG_FIND_THE_CYCLE_TO_MESH
+//					if (Grid.bDisplay)
+//					{
+//						F3DDebugSession W(FString::Printf(TEXT("cancel due to Seg Length %f vs %f"), sqrt(Length2), sqrt(MaxEdgeSquareLength)));
+//						Grid.DisplayIsoSegment(EGridSpace::UniformScaled, *NewSegmentEnd, *CandidateNode, 0, EVisuProperty::YellowCurve);
+//						Wait();
+//					}
+//#endif
+//					continue;
+//				}
+//
+//				BestIsoscelesCriteria = IsoscelesCriteria;
+//				AcuteNodeIndex = Index;
+//				bAcuteTriangle = true;
+//			}
+//		}
+//	}
 
 	if (bAcuteTriangle)
 	{
@@ -454,15 +479,15 @@ bool FCycleTriangulator::FindCandidateNodes(int32 StartIndex)
 #endif
 
 	constexpr double MaxSlope = Slope::PiSlope - Slope::OneDegree;
-	const double StartMaxSlope = FMath::Min(CandidateNodes.Last().SlopeAtStartNode + Slope::OneDegree, MaxSlope);
-	const double EndMaxSlope = FMath::Min(CandidateNodes[0].SlopeAtEndNode + Slope::OneDegree, MaxSlope);
+	const double StartMaxSlope = FMath::Min(CandidateNodes.Last().SlopeAtStartNode, MaxSlope);
+	const double EndMaxSlope = FMath::Min(CandidateNodes[0].SlopeAtEndNode, MaxSlope);
 
-	if (CandidateNodes.Last().SlopeAtStartNode < StartMaxSlope)
+	if (CandidateNodes.Last().SlopeAtStartNode - Slope::Epsilon < StartMaxSlope && CandidateNodes.Last().SlopeAtEndNode < EndMaxSlope)
 	{
 		CandidateNodes.Last().bIsValid = true;
 		ExtremityCandidateNodes[0] = &CandidateNodes.Last();
 	}
-	if (CandidateNodes[0].SlopeAtEndNode < EndMaxSlope)
+	if (CandidateNodes[0].SlopeAtEndNode - Slope::Epsilon < EndMaxSlope && CandidateNodes[0].SlopeAtStartNode < StartMaxSlope)
 	{
 		CandidateNodes[0].bIsValid = true;
 		ExtremityCandidateNodes[1] = &CandidateNodes[0];
@@ -521,17 +546,19 @@ bool FCycleTriangulator::FindTheBestCandidateNode()
 {
 	constexpr double MinSlopeToNotBeAligned = Slope::Epsilon;
 
-	IntersectionCountAllowed = 0;
-
 	double MinCriteria = HUGE_VALUE;
 	double CandidateSlopeAtStartNode = Slope::PiSlope - MinSlopeToNotBeAligned;
 	double CandidateSlopeAtEndNode = Slope::PiSlope - MinSlopeToNotBeAligned;
 	FCandidateNode* CandidateNode = nullptr;
 	double CandidateEndSquareDistance = HUGE_VALUE;
 	double CandidateStartSquareDistance = HUGE_VALUE;
+	bool bIsDegeneratedDueSlope = false;
+	bool bIsDegeneratedDueToLength = false;
+	int32 CandidateIntersectionCount = 0;
 
 	const FPoint2D& StartPoint2D = FirstSideStartNode->Get2DPoint(EGridSpace::UniformScaled, Grid);
 	const FPoint2D& EndPoint2D = FirstSideEndNode->Get2DPoint(EGridSpace::UniformScaled, Grid);
+	double NewSegmentSquareLength2500 = StartPoint2D.SquareDistance(EndPoint2D) * 2500; // need to check if the length of the sides are not disproportionate
 
 #ifdef DEBUG_FIND_THE_BEST_CANDIDATE_NODE
 	static int32 BestCandidateNode = 0;
@@ -545,6 +572,26 @@ bool FCycleTriangulator::FindTheBestCandidateNode()
 	}
 #endif
 
+	TFunction<bool(const double, const double)> IsDegeneratedDueToSlope = [](const double SlopeAtStartNode, const double SlopeAtEndNode) -> bool
+	{
+		// check if the angle of the candidate triangle are not degenerated
+		double MinSlope;
+		double MaxSlope;
+		GetMinMax(SlopeAtStartNode, SlopeAtEndNode, MinSlope, MaxSlope);
+		constexpr double PiMinus5 = Slope::PiSlope - Slope::FiveDegree;
+		return MinSlope < Slope::FiveDegree&& MaxSlope > PiMinus5;
+	};
+
+	TFunction<bool(const double, const double)> IsDegeneratedDueToLength = [&NewSegmentSquareLength2500](const double CandidateEndSquareDistance, const double CandidateStartSquareDistance)
+	{
+		// check if the length of the sides are not disproportionate min side minus than 50 times max length 
+		const double MinSquareLenth = FMath::Min(CandidateEndSquareDistance, CandidateStartSquareDistance);
+		return NewSegmentSquareLength2500 < MinSquareLenth;
+	};
+
+	const bool bCheckIntersectionWithIso = true;
+	IntersectionCountAllowed = InnerToOuterIsoSegmentsIntersectionTool.CountIntersections(*FirstSideStartNode, *FirstSideEndNode);
+
 	for (FNodeIntersectionCount& NodeIntersection : NodeToIntersection)
 	{
 		if (NodeIntersection.Value > IntersectionCountAllowed)
@@ -552,19 +599,8 @@ bool FCycleTriangulator::FindTheBestCandidateNode()
 			if (CandidateNode)
 			{
 				// if the candidate triangle is degenerated, allow intersection with iso to find a better candidate
-
-				// check if the angle of the candidate triangle are not degenerated
-				double MinSlope;
-				double MaxSlope;
-				GetMinMax(CandidateSlopeAtStartNode, CandidateSlopeAtEndNode, MinSlope, MaxSlope);
-				constexpr double PiMinus5 = Slope::PiSlope - Slope::FiveDegree;
-				const bool bIsDegeneratedDueSlope = MinSlope < Slope::FiveDegree && MaxSlope > PiMinus5;
-
-				// check if the length of the sides are not disproportionate min side minus than 50 times max length 
-				const double MinSquareLenth = FMath::Min(CandidateEndSquareDistance, CandidateStartSquareDistance);
-				const double SegmentSquareLength2500 = StartPoint2D.SquareDistance(EndPoint2D) * 2500;
-				const bool bIsDegeneratedDueToLength = SegmentSquareLength2500 < MinSquareLenth;
-
+				bIsDegeneratedDueSlope = IsDegeneratedDueToSlope(CandidateSlopeAtStartNode, CandidateSlopeAtEndNode);
+				bIsDegeneratedDueToLength = IsDegeneratedDueToLength(CandidateEndSquareDistance, CandidateStartSquareDistance);
 				const bool bIsDegenerated = bIsDegeneratedDueSlope || bIsDegeneratedDueToLength;
 
 				if (!bIsDegenerated)
@@ -588,13 +624,25 @@ bool FCycleTriangulator::FindTheBestCandidateNode()
 			F3DDebugSession _(TEXT("Candidate triangle"));
 			Grid.DisplayIsoSegment(EGridSpace::UniformScaled, *FirstSideStartNode, *Node, 1, EVisuProperty::RedCurve);
 			Grid.DisplayIsoSegment(EGridSpace::UniformScaled, *FirstSideEndNode, *Node, 2, EVisuProperty::RedCurve);
-			//Wait();
+			Wait();
 		}
 #endif
 
 		double CEndSquareDistance = StartPoint2D.SquareDistance(NodePoint2D);
 		double CStartSquareDistance = EndPoint2D.SquareDistance(NodePoint2D);
 
+		// if the previous candidate and the current candidate is Degenerated Due To Length, the previous is preferred
+		if (bIsDegeneratedDueToLength && IsDegeneratedDueToLength(CEndSquareDistance, CStartSquareDistance))
+		{
+			continue;
+		}
+
+		// if the previous candidate and the current candidate is degenerated Due to slope, the previous is preferred
+		if (bIsDegeneratedDueSlope && IsDegeneratedDueToSlope(CNode->SlopeAtStartNode, CNode->SlopeAtEndNode))
+		{
+			continue;
+		}
+
 		// the candidate triangle is inside the current candidate triangle
 		// Angle is not accurate enough to define if a triangle is inside or not an other triangle (especially with acute triangle).
 		// The length of the side is a better criteria but must be combine with the angles
@@ -610,32 +658,26 @@ bool FCycleTriangulator::FindTheBestCandidateNode()
 		{
 			// Check if the candidate segment is not in intersection with existing segments
 			// if the segment exist, it has already been tested
-			FIsoSegment* StartSegment = nullptr;
-			FIsoSegment* EndSegment = nullptr;
-			if (CycleIntersectionTool.DoesIntersect(FirstSideStartNode, Node, &StartSegment))
 			{
-#ifdef DEBUG_FIND_THE_BEST_CANDIDATE_NODE_
-				if (Grid.bDisplay)
+				const FIsoSegment* StartSegment = FirstSideStartNode->GetSegmentConnectedTo(Node);
+				if (!StartSegment && CycleIntersectionTool.DoesIntersect(FirstSideStartNode, Node))
 				{
-					F3DDebugSession _(TEXT("Candidate triangle"));
-					Grid.DisplayIsoSegment(EGridSpace::UniformScaled, *FirstSideStartNode, *Node, 1, EVisuProperty::RedCurve);
-					Wait();
+					if (ConfirmIntersection(FirstSideStartNode, Node, FirstSideEndNode, StartSegment))
+					{
+						continue;
+					}
 				}
-#endif
-				continue;
 			}
 
-			if (CycleIntersectionTool.DoesIntersect(FirstSideEndNode, Node, &EndSegment))
 			{
-#ifdef DEBUG_FIND_THE_BEST_CANDIDATE_NODE_
-				if (Grid.bDisplay)
+				const FIsoSegment* EndSegment = FirstSideEndNode->GetSegmentConnectedTo(Node);
+				if (CycleIntersectionTool.DoesIntersect(FirstSideEndNode, Node))
 				{
-					F3DDebugSession _(TEXT("Candidate triangle"));
-					Grid.DisplayIsoSegment(EGridSpace::UniformScaled, *FirstSideEndNode, *Node, 2, EVisuProperty::RedCurve);
-					Wait();
+					if (ConfirmIntersection(FirstSideEndNode, Node, FirstSideStartNode, EndSegment))
+					{
+						continue;
+					}
 				}
-#endif
-				continue;
 			}
 
 			MinCriteria = PointCriteria;
@@ -644,6 +686,7 @@ bool FCycleTriangulator::FindTheBestCandidateNode()
 			CandidateSlopeAtEndNode = CNode->SlopeAtEndNode;
 			CandidateEndSquareDistance = CEndSquareDistance;
 			CandidateStartSquareDistance = CStartSquareDistance;
+			CandidateIntersectionCount = NodeIntersection.Value;
 		}
 	}
 
@@ -694,20 +737,21 @@ bool FCycleTriangulator::FindTheBestCandidateNode()
 	};
 #endif
 
-	TFunction<void(FCandidateNode*, const double, const FPoint2D&, double, double) > CompareCandidateAtSegmentNode = [&CandidateNode, &Grid = Grid](FCandidateNode* CandidatePoint, const double CloseLimit, const FPoint2D& SegmentPoint2D, double CandidateToSegmentSquareDistance, double SlopeDiffAtSegmentExtremity)
+	TFunction<void(FCandidateNode*, const double, const FPoint2D&, double, double) > CompareCandidateAtSegmentNode = [&CandidateNode, &Grid = Grid](FCandidateNode* NewCandidate, const double CloseLimit, const FPoint2D& SegmentPoint2D, double CandidateToSegmentSquareDistance, double SlopeDiffAtSegmentExtremity)
 	{
-		const FPoint2D& NodePoint2D = CandidatePoint->Node->Get2DPoint(EGridSpace::UniformScaled, Grid);
+		const FPoint2D& NodePoint2D = NewCandidate->Node->Get2DPoint(EGridSpace::UniformScaled, Grid);
 		const double CFirstSquareDistance = SegmentPoint2D.SquareDistance(NodePoint2D) * 2.;
 
 		if (SlopeDiffAtSegmentExtremity < CloseLimit && CFirstSquareDistance < CandidateToSegmentSquareDistance)
 		{
-			CandidateNode = CandidatePoint;
+			CandidateNode = NewCandidate;
 		}
 	};
 
 	// Optimize selection step: if the selected triangle is close to one of the next segment, check if a triangle from it is not better 
 	if (ExtremityCandidateNodes[0] != CandidateNode && ExtremityCandidateNodes[1] != CandidateNode)
 	{
+		const int32 NewCandidateIntersectionCountAllowed = FMath::Max(CandidateIntersectionCount - 1, 0);
 		if (ExtremityCandidateNodes[0] && ExtremityCandidateNodes[1])
 		{
 			const double SlopeDiffAtStart = ExtremityCandidateNodes[0]->SlopeAtStartNode - CandidateNode->SlopeAtStartNode;
@@ -716,32 +760,29 @@ bool FCycleTriangulator::FindTheBestCandidateNode()
 			if (SlopeDiffAtStart < SlopeDiffAtEnd)
 			{
 				const FNodeIntersectionCount* ExtremityIntersectionCount = NodeToIntersection.FindByPredicate([Extremity = ExtremityCandidateNodes[0]](const FNodeIntersectionCount& Node) { return Node.Key == Extremity; });
-				if (IntersectionCountAllowed >= ExtremityIntersectionCount->Value)
+				if (NewCandidateIntersectionCountAllowed >= ExtremityIntersectionCount->Value)
 				{
 					CompareCandidateAtSegmentNode(ExtremityCandidateNodes[0], ExtremityCandidateNodes[0]->SlopeAtStartNode * AThird, StartPoint2D, CandidateStartSquareDistance, SlopeDiffAtStart);
-					printf("");
 				}
 			}
 			else
 			{
 				const FNodeIntersectionCount* ExtremityIntersectionCount = NodeToIntersection.FindByPredicate([Extremity = ExtremityCandidateNodes[1]](const FNodeIntersectionCount& Node) { return Node.Key == Extremity; });
-				if (IntersectionCountAllowed >= ExtremityIntersectionCount->Value)
+				if (NewCandidateIntersectionCountAllowed >= ExtremityIntersectionCount->Value)
 				{
 					CompareCandidateAtSegmentNode(ExtremityCandidateNodes[1], ExtremityCandidateNodes[1]->SlopeAtEndNode * AThird, EndPoint2D, CandidateEndSquareDistance, SlopeDiffAtEnd);
 				}
-				printf("");
 			}
 		}
 		else if (ExtremityCandidateNodes[0])
 		{
 			const FNodeIntersectionCount* ExtremityIntersectionCount = NodeToIntersection.FindByPredicate([Extremity = ExtremityCandidateNodes[0]](const FNodeIntersectionCount& Node) { return Node.Key == Extremity; });
-			if (IntersectionCountAllowed >= ExtremityIntersectionCount->Value)
+			if (NewCandidateIntersectionCountAllowed >= ExtremityIntersectionCount->Value)
 			{
 				const double SlopeDiffAtStart = ExtremityCandidateNodes[0]->SlopeAtStartNode - CandidateNode->SlopeAtStartNode;
 				if (SlopeDiffAtStart < Slope::FifteenDegree)
 				{
 					CompareCandidateAtSegmentNode(ExtremityCandidateNodes[0], ExtremityCandidateNodes[0]->SlopeAtStartNode * AThird, StartPoint2D, CandidateStartSquareDistance, SlopeDiffAtStart);
-					printf("");
 				}
 			}
 		}
@@ -749,13 +790,12 @@ bool FCycleTriangulator::FindTheBestCandidateNode()
 		{
 			//FLoopNode const* const* StartNode = NodesOfLoop.FindByPredicate([](const FLoopNode* Node) { return !Node->IsDelete(); });
 			const FNodeIntersectionCount* ExtremityIntersectionCount = NodeToIntersection.FindByPredicate([Extremity = ExtremityCandidateNodes[1]](const FNodeIntersectionCount& Node) { return Node.Key == Extremity; });
-			if (IntersectionCountAllowed >= ExtremityIntersectionCount->Value)
+			if (NewCandidateIntersectionCountAllowed >= ExtremityIntersectionCount->Value)
 			{
 				const double SlopeDiffAtEnd = ExtremityCandidateNodes[1]->SlopeAtEndNode - CandidateNode->SlopeAtEndNode;
 				if (SlopeDiffAtEnd < Slope::FifteenDegree)
 				{
 					CompareCandidateAtSegmentNode(ExtremityCandidateNodes[1], ExtremityCandidateNodes[1]->SlopeAtEndNode * AThird, EndPoint2D, CandidateEndSquareDistance, SlopeDiffAtEnd);
-					printf("");
 				}
 			}
 		}
@@ -1033,8 +1073,8 @@ bool FCycleTriangulator::ValidComplementaryNodeOrDeleteIt(FAdditionalIso& Side,
 	double Slope = ComputeOrientedSlope(OppositePoint, CandidatePoint, TestPoint);
 	if ((Slope > 0) == (CandidateIndex == 0))
 	{
-		FIsoSegment* Segment;
-		if (!CycleIntersectionTool.DoesIntersect(Candidate, OppositeNode, &Segment))
+		const FIsoSegment* Segment = Candidate->GetSegmentConnectedTo(OppositeNode);
+		if (Segment || !CycleIntersectionTool.DoesIntersect(Candidate, OppositeNode))
 		{
 			// The new segment is fully inside the cycle 
 			return true;
@@ -1159,8 +1199,8 @@ void FCycleTriangulator::ValidateComplementaryNodesWithInsideAndIntersectionsCri
 		}
 #endif
 		// Check if [Side.Nodes[0] - Side.Nodes[1]] intersect cycle
-		FIsoSegment* Segment = nullptr;
-		if (CycleIntersectionTool.DoesIntersect(Side.Nodes[0], Side.Nodes[1], &Segment) || !IsInnerSideSegmentInsideCycle(Side))
+		const FIsoSegment* Segment = Side.Nodes[0]->GetSegmentConnectedTo(Side.Nodes[1]);
+		if (!Segment && (CycleIntersectionTool.DoesIntersect(Side.Nodes[0], Side.Nodes[1]) || !IsInnerSideSegmentInsideCycle(Side)))
 		{
 			// intersect => choose between NextFinalNode and LastPotentialNode
 #ifdef DEBUG_VALID_COMPLEMENTARY_NODE
diff --git a/Engine/Source/Runtime/Datasmith/CADKernel/Private/CADKernel/Mesh/Meshers/CycleTriangulator.h b/Engine/Source/Runtime/Datasmith/CADKernel/Private/CADKernel/Mesh/Meshers/CycleTriangulator.h
--- a/Engine/Source/Runtime/Datasmith/CADKernel/Private/CADKernel/Mesh/Meshers/CycleTriangulator.h
+++ b/Engine/Source/Runtime/Datasmith/CADKernel/Private/CADKernel/Mesh/Meshers/CycleTriangulator.h
@@ -294,6 +294,10 @@ private:
 		return Max;
 	}
 
+	/**
+	 * In case of intersection with the loop, confirm that it was a real intersection and not a nearly parallel segment
+	 */
+	bool ConfirmIntersection(const FIsoNode* Start, const FIsoNode* End, const FIsoNode* Candidate, const FIsoSegment* IntersectedSegment) const;
 };
 
 namespace Polygon
diff --git a/Engine/Source/Runtime/Datasmith/CADKernel/Private/CADKernel/Mesh/Meshers/IsoTriangulator.cpp b/Engine/Source/Runtime/Datasmith/CADKernel/Private/CADKernel/Mesh/Meshers/IsoTriangulator.cpp
--- a/Engine/Source/Runtime/Datasmith/CADKernel/Private/CADKernel/Mesh/Meshers/IsoTriangulator.cpp
+++ b/Engine/Source/Runtime/Datasmith/CADKernel/Private/CADKernel/Mesh/Meshers/IsoTriangulator.cpp
@@ -46,7 +46,7 @@ struct FCandidateSegment
 	FLoopNode& StartNode;
 	FLoopNode& EndNode;
 	double Length;
-	
+
 	FCandidateSegment(const FGrid& Grid, FLoopNode& Node1, FLoopNode& Node2)
 		: StartNode(Node1)
 		, EndNode(Node2)
@@ -60,10 +60,10 @@ struct FCandidateSegment
 FIsoTriangulator::FIsoTriangulator(FGrid& InGrid, FFaceMesh& OutMesh, const FMeshingTolerances& InTolerances)
 	: Grid(InGrid)
 	, Mesh(OutMesh)
-	, LoopSegmentsIntersectionTool(InGrid)
-	, InnerSegmentsIntersectionTool(InGrid)
-	, InnerToOuterIsoSegmentsIntersectionTool(InGrid, /*AllowOverlaping*/ true)
-	, ThinZoneIntersectionTool(InGrid)
+	, LoopSegmentsIntersectionTool(InGrid, InTolerances.GeometricTolerance)
+	, InnerSegmentsIntersectionTool(InGrid, InTolerances.GeometricTolerance)
+	, InnerToOuterIsoSegmentsIntersectionTool(InGrid, InTolerances.GeometricTolerance, /*AllowOverlaping*/ true)
+	, ThinZoneIntersectionTool(InGrid, InTolerances.GeometricTolerance)
 	, Tolerances(InTolerances)
 {
 	FinalInnerSegments.Reserve(3 * Grid.InnerNodesCount());
@@ -179,7 +179,7 @@ bool FIsoTriangulator::Triangulate()
 	Grid.DisplayIsoSegments(TEXT("FIsoTrianguler::Final Iso ToLink Inner To Loops"), DisplaySpace, FinalToLoops, false, false, true, EVisuProperty::YellowCurve);
 #endif
 
-	AddCandidateSegmentsToLinkInnerAndLoop();
+	SelectSegmentsToLinkInnerToLoop();
 
 #ifdef DEBUG_ISOTRIANGULATOR
 	Grid.DisplayIsoSegments(TEXT("FIsoTrianguler::Final Segments"), DisplaySpace, FinalToLoops, false, false, true, EVisuProperty::YellowCurve);
@@ -340,48 +340,6 @@ void FIsoTriangulator::BuildLoopSegments()
 		LoopSegments.Add(&Segment);
 	}
 
-	for (FIsoSegment* Segment : LoopSegments)
-	{
-		double SegmentSlope = ComputeSlope(Segment->GetFirstNode().Get2DPoint(EGridSpace::UniformScaled, Grid), Segment->GetSecondNode().Get2DPoint(EGridSpace::UniformScaled, Grid));
-		if (SegmentSlope < IsoTriangulatorImpl::MaxSlopeToBeIso)
-		{
-			Segment->SetAsIsoU();
-		}
-		if (SegmentSlope < IsoTriangulatorImpl::LimitValueMax(2.) && SegmentSlope > IsoTriangulatorImpl::LimitValueMin(2.))
-		{
-			Segment->SetAsIsoV();
-		}
-		if (SegmentSlope < IsoTriangulatorImpl::LimitValueMax(4.) && SegmentSlope > IsoTriangulatorImpl::LimitValueMin(4.))
-		{
-			Segment->SetAsIsoU();
-		}
-		if (SegmentSlope < IsoTriangulatorImpl::LimitValueMax(6.) && SegmentSlope > IsoTriangulatorImpl::LimitValueMin(6.))
-		{
-			Segment->SetAsIsoV();
-		}
-		if (SegmentSlope > IsoTriangulatorImpl::LimitValueMin(8.))
-		{
-			Segment->SetAsIsoU();
-		}
-	}
-
-	for (FLoopNode& Node : LoopNodes)
-	{
-		if (Node.IsDelete())
-		{
-			continue;
-		}
-
-		if (Node.GetConnectedSegments()[0]->IsIsoU() && Node.GetConnectedSegments()[1]->IsIsoU())
-		{
-			Node.SetAsIsoU();
-		}
-		else if (Node.GetConnectedSegments()[0]->IsIsoV() && Node.GetConnectedSegments()[1]->IsIsoV())
-		{
-			Node.SetAsIsoV();
-		}
-	}
-
 #ifdef CADKERNEL_DEV
 	Chronos.BuildLoopSegmentsDuration += FChrono::Elapse(StartTime);
 #endif
@@ -534,6 +492,7 @@ void FIsoTriangulator::GetThinZoneMesh(const TMap<int32, FLoopNode*>& IndexToNod
 		CandidateSegment.EndNode.SetThinZoneNodeMarker();
 
 		Segment.ConnectToNode();
+		Segment.SetFinalMarker();
 		ThinZoneSegments.Add(&Segment);
 		ThinZoneIntersectionTool.AddSegment(Segment);
 	}
@@ -737,7 +696,7 @@ void FIsoTriangulator::BuildInnerSegments()
 			Temp.Add((UniformCutting[Index] - UniformCutting[Index - 1]));
 		}
 		TolerancesAlongU.Add(Temp[0] * 0.1);
-		for (int32 Index = 1; Index < Num-1; Index++)
+		for (int32 Index = 1; Index < Num - 1; Index++)
 		{
 			TolerancesAlongU.Add((Temp[Index] + Temp[Index - 1]) * 0.05);
 		}
@@ -942,6 +901,7 @@ void FIsoTriangulator::ConnectCellLoops()
 
 	for (FCell& Cell : Cells)
 	{
+		InitCellCorners(Cell);
 		Cell.InitLoopConnexions();
 	}
 
@@ -959,15 +919,32 @@ void FIsoTriangulator::ConnectCellLoops()
 		DisplayCell(Cell);
 #endif
 
-		if (Cell.LoopCells.Num())
+		if (Cell.CellLoops.Num())
 		{
-			Cell.ConnectLoopsByNeighborhood();
+			Cell.FindCandidateToConnectLoopsByNeighborhood();
+
+#ifdef DEBUG_CONNECTCELLCORNERTOINNERLOOP
+			if (bDisplay)
+			{
+				DisplayCellConnexions(FString::Printf(TEXT("Cell %d CandidateSegments"), Cell.Id), Cell.LoopConnexions, EVisuProperty::YellowCurve);
+			}
+#endif
+
+			FindCandidateToConnectCellCornerToLoops(Cell);
+
+#ifdef DEBUG_CONNECTCELLCORNERTOINNERLOOP
+			if (bDisplay)
+			{
+				DisplayCellConnexions(FString::Printf(TEXT("Cell %d CandidateSegments"), Cell.Id), Cell.LoopConnexions, EVisuProperty::BlueCurve);
+			}
+#endif
 
+			Cell.SelectSegmentToConnectLoops(IsoSegmentFactory);
+			Cell.SelectSegmentToConnectLoopToCorner(IsoSegmentFactory);
+			Cell.CheckAllLoopsConnectedTogetherAndConnect();
 #ifdef DEBUG_CONNECT_CELL_LOOPS
 			Grid.DisplayIsoSegments(TEXT("ConnectCellLoops::Step 0"), EGridSpace::UniformScaled, Cell.FinalSegments, false, false, false, EVisuProperty::YellowCurve);
 #endif
-
-			ConnectCellCornerToLoops(Cell);
 		}
 
 #ifdef DEBUG_CONNECT_CELL_LOOPS
@@ -1539,7 +1516,7 @@ FIsoSegment* FIsoTriangulator::FindNextSegment(EGridSpace Space, const FIsoSegme
 		const FPoint2D& OtherPoint = (StartNode == &Segment->GetFirstNode()) ? Segment->GetSecondNode().Get2DPoint(Space, Grid) : Segment->GetFirstNode().Get2DPoint(Space, Grid);
 
 		double Slope = GetSlope(StartPoint, OtherPoint, ReferenceSlope);
-		if (Slope < SMALL_NUMBER_SQUARE) 
+		if (Slope < SMALL_NUMBER_SQUARE)
 		{
 			Slope = 8;
 		}
@@ -1625,47 +1602,52 @@ void FIsoTriangulator::TriangulateInnerNodes()
 #endif
 }
 
-void FCell::ConnectLoopsByNeighborhood()
+void FCell::FindCandidateToConnectLoopsByNeighborhood()
 {
 #ifdef DEBUG_CONNECT_CELL_SUB_LOOPS_BY_NEIGHBORHOOD
 	F3DDebugSession _(Grid.bDisplay, FString::Printf(TEXT("Cell % d: Build Segments to Connect inner loops"), Id));
 #endif
 
-
-	for (FLoopConnexion& LoopConnexion : LoopConnexions)
+	for (FCellConnexion& LoopConnexion : LoopConnexions)
 	{
 		TryToConnectTwoSubLoopsWithShortestSegment(LoopConnexion);
 	}
+}
 
-	TArray<FLoopConnexion*> LoopConnexionPtrs;
+void FCell::SelectSegmentToConnectLoops(TFactory<FIsoSegment>& SegmentFactory)
+{
+	TArray<FCellConnexion*> LoopConnexionPtrs;
 	LoopConnexionPtrs.Reserve(LoopConnexions.Num());
-	for (FLoopConnexion& Connexion : LoopConnexions)
+	for (FCellConnexion& Connexion : LoopConnexions)
 	{
+		if (Connexion.Loop2.IsCellCorner())
+		{
+			continue;
+		}
 		LoopConnexionPtrs.Add(&Connexion);
 	}
 
-	Algo::Sort(LoopConnexionPtrs, [&](const FLoopConnexion* LoopConnexion1, const FLoopConnexion* LoopConnexion2)
+	Algo::Sort(LoopConnexionPtrs, [&](const FCellConnexion* LoopConnexion1, const FCellConnexion* LoopConnexion2)
 		{
 			return LoopConnexion1->MinDistance < LoopConnexion2->MinDistance;
 		});
 
-
 	{
 #ifdef DEBUG_CONNECT_CELL_SUB_LOOPS_BY_NEIGHBORHOOD
-		F3DDebugSession _(TEXT("Best Path"));
+		F3DDebugSession _(Grid.bDisplay, TEXT("Best Path"));
 #endif
-		const int32 LoopCount = LoopCells.Num();
-		for (FLoopConnexion* LoopConnexion : LoopConnexionPtrs)
+		const int32 LoopCount = CellLoops.Num();
+		for (FCellConnexion* LoopConnexion : LoopConnexionPtrs)
 		{
-			if (!LoopConnexion->bIsConnexionWithBorder && LoopConnexion->IsShortestPath(LoopCount))
+			if (!LoopConnexion->bIsConnexionWithOuter && LoopConnexion->IsShortestPath(LoopCount))
 			{
 				TryToCreateSegment(*LoopConnexion);
 			}
 		}
 
-		for (FLoopConnexion* LoopConnexion : LoopConnexionPtrs)
+		for (FCellConnexion* LoopConnexion : LoopConnexionPtrs)
 		{
-			if (LoopConnexion->bIsConnexionWithBorder)
+			if (LoopConnexion->bIsConnexionWithOuter && LoopConnexion->IsShortestPathToOuterLoop(LoopCount))
 			{
 				TryToCreateSegment(*LoopConnexion);
 			}
@@ -1673,62 +1655,156 @@ void FCell::ConnectLoopsByNeighborhood()
 	}
 
 #ifdef DEBUG_CONNECT_CELL_SUB_LOOPS_BY_NEIGHBORHOOD
-	Grid.DisplayIsoSegments(TEXT("Build Segments Connect inner close loops"), EGridSpace::UniformScaled, CandidateSegments, false, false, false, EVisuProperty::YellowCurve);
-	Wait(Grid.bDisplay);
+	if (Grid.bDisplay)
+	{
+		Grid.DisplayIsoSegments(TEXT("Build Segments Connect inner close loops: Candidates"), EGridSpace::UniformScaled, FinalSegments, false, false, false, EVisuProperty::YellowCurve);
+		Wait(Grid.bDisplay);
+	}
 #endif
+}
+
+void FCell::SelectSegmentToConnectLoopToCorner(TFactory<FIsoSegment>& SegmentFactory)
+{
+	TArray<FCellConnexion*> LoopConnexionPtrs;
+	LoopConnexionPtrs.Reserve(LoopConnexions.Num());
+	for (FCellConnexion& Connexion : LoopConnexions)
+	{
+		if (Connexion.Loop2.IsCellCorner() && !Connexion.Segment)
+		{
+			LoopConnexionPtrs.Add(&Connexion);
+		}
+	}
+
+	Algo::Sort(LoopConnexionPtrs, [&](const FCellConnexion* LoopConnexion1, const FCellConnexion* LoopConnexion2)
+		{
+			return LoopConnexion1->MinDistance < LoopConnexion2->MinDistance;
+		});
 
-	if(CandidateSegments.Num())
 	{
-		const bool bSelectWithCellSizeCriteria = true;
-		SelectSegmentInCandidateSegments(Triangulator.IsoSegmentFactory, bSelectWithCellSizeCriteria);
+#ifdef DEBUG_CONNECT_CELL_SUB_LOOPS_BY_NEIGHBORHOOD
+		F3DDebugSession _(Grid.bDisplay, TEXT("Best Path"));
+#endif
+		const int32 LoopCount = CellLoops.Num();
+		for (FCellConnexion* LoopConnexion : LoopConnexionPtrs)
+		{
+			if (!LoopConnexion->bIsConnexionWithOuter && LoopConnexion->IsShortestPathToCorner(LoopCount))
+			{
+				TryToCreateSegment(*LoopConnexion);
+			}
+		}
+
+		for (FCellConnexion* LoopConnexion : LoopConnexionPtrs)
+		{
+			if (LoopConnexion->bIsConnexionWithOuter)
+			{
+				// if the outerLoop is already connected to an innerLoop, the segment is not create.
+				// The aim is to avoid to create long segment that generate degenerated triangle
+				bool bCreateSegment = true;
+				for (const FCellConnexion* Connexion : LoopConnexion->Loop1.Connexions)
+				{
+					if (Connexion->Segment && Connexion->Segment->IsAFinalSegment())
+					{
+						const FCellLoop& InnerLoop = Connexion->Loop1.bIsOuterLoop ? Connexion->Loop2 : Connexion->Loop1;
+						if (!InnerLoop.IsCellCorner())
+						{
+							bCreateSegment = false;
+						}
+					}
+				}
+				if (bCreateSegment)
+				{
+					TryToCreateSegment(*LoopConnexion);
+				}
+			}
+		}
 	}
 
 #ifdef DEBUG_CONNECT_CELL_SUB_LOOPS_BY_NEIGHBORHOOD
-	Grid.DisplayIsoSegments(TEXT("Build Segments Connect inner close loops"), EGridSpace::UniformScaled, FinalSegments, false, false, false, EVisuProperty::YellowCurve);
-	Wait(Grid.bDisplay);
+	if (Grid.bDisplay)
+	{
+		Grid.DisplayIsoSegments(TEXT("Build Segments Connect inner close loops: Candidates"), EGridSpace::UniformScaled, CandidateSegments, false, false, false, EVisuProperty::YellowCurve);
+		Wait(Grid.bDisplay);
+	}
 #endif
+}
+
+void FCellLoop::SpreadAsConnected()
+{
+	bIsConnected = true;
+	for (FCellConnexion* Connexion : Connexions)
+	{
+		if (!Connexion->Segment)
+		{
+			continue;
+		}
+
+		FCellLoop* OtherCell = Connexion->GetOtherLoop(this);
+		if (!OtherCell->bIsConnected)
+		{
+			OtherCell->SpreadAsConnected();
+		}
+	}
+}
 
+void FCell::CheckAllLoopsConnectedTogetherAndConnect()
+{
+	for (FCellCorner& CellCorner : CellCorners)
+	{
+		CellCorner.SpreadAsConnected();
+	}
+
+	for (FCellLoop& CellLoop : CellLoops)
+	{
+		if (CellLoop.bIsOuterLoop)
+		{
+			CellLoop.SpreadAsConnected();
+		}
+	}
 }
 
-void FCell::TryToCreateSegment(FLoopConnexion& LoopConnexion)
+void FCell::TryToCreateSegment(FCellConnexion& LoopConnexion)
 {
 	const FPoint2D& ACoordinates = LoopConnexion.NodeA->Get2DPoint(EGridSpace::UniformScaled, Grid);
 	const FPoint2D& BCoordinates = LoopConnexion.NodeB->Get2DPoint(EGridSpace::UniformScaled, Grid);
 
-	Triangulator.TryToCreateSegment(*this, LoopConnexion.NodeA, ACoordinates, LoopConnexion.NodeB, BCoordinates, 0.1);
+	LoopConnexion.Segment = Triangulator.GetOrTryToCreateSegment(*this, LoopConnexion.NodeA, ACoordinates, LoopConnexion.NodeB, BCoordinates, Slope::OneDegree);
+	if (LoopConnexion.Segment && !LoopConnexion.Segment->IsAFinalSegment())
+	{
+		LoopConnexion.Segment->SetFinalMarker();
+		LoopConnexion.Segment->ConnectToNode();
+
+		IntersectionTool.AddSegment(*LoopConnexion.Segment);
+		FinalSegments.Add(LoopConnexion.Segment);
+	}
 }
 
 
 void FCell::InitLoopConnexions()
 {
-	FTimePoint StartTime = FChrono::Now();
+	const int32 MaxInnerToCornerConnexions = CellCorners.Num() * (InnerLoopCount + OuterLoopCount);
 
 	switch (InnerLoopCount)
 	{
-	case 0:
-	{
-		break;
-	}
 	case 1:
 	{
-		LoopConnexions.Reserve( OuterLoopCount );
+		LoopConnexions.Reserve(OuterLoopCount + MaxInnerToCornerConnexions);
 		LoopCellBorderIndices = { OuterLoopCount };
 		break;
 	}
 	case 2:
 	{
-		LoopConnexions.Reserve(1 + 2 * OuterLoopCount);
-		LoopConnexions.Emplace(LoopCells[OuterLoopCount], LoopCells[OuterLoopCount + 1]);
+		LoopConnexions.Reserve(1 + 2 * OuterLoopCount + MaxInnerToCornerConnexions);
+		LoopConnexions.Emplace(CellLoops[OuterLoopCount], CellLoops[OuterLoopCount + 1]);
 
 		LoopCellBorderIndices = { OuterLoopCount , OuterLoopCount + 1 };
 		break;
 	}
 	case 3:
 	{
-		LoopConnexions.Reserve(3 + 3 * OuterLoopCount);
-		LoopConnexions.Emplace(LoopCells[OuterLoopCount], LoopCells[OuterLoopCount + 1]);
-		LoopConnexions.Emplace(LoopCells[OuterLoopCount], LoopCells[OuterLoopCount + 2]);
-		LoopConnexions.Emplace(LoopCells[OuterLoopCount + 1], LoopCells[OuterLoopCount + 2]);
+		LoopConnexions.Reserve(3 + 3 * OuterLoopCount + MaxInnerToCornerConnexions);
+		LoopConnexions.Emplace(CellLoops[OuterLoopCount], CellLoops[OuterLoopCount + 1]);
+		LoopConnexions.Emplace(CellLoops[OuterLoopCount], CellLoops[OuterLoopCount + 2]);
+		LoopConnexions.Emplace(CellLoops[OuterLoopCount + 1], CellLoops[OuterLoopCount + 2]);
 
 		LoopCellBorderIndices = { OuterLoopCount , OuterLoopCount + 1 , OuterLoopCount + 2 };
 		break;
@@ -1744,13 +1820,14 @@ void FCell::InitLoopConnexions()
 		FBowyerWatsonTriangulator BWTriangulator(LoopBarycenters, EdgeVertexIndices);
 		BWTriangulator.Triangulate();
 		LoopCellBorderIndices = BWTriangulator.GetOuterVertices();
-		LoopConnexions.Reserve((EdgeVertexIndices.Num() >> 1) + LoopCellBorderIndices.Num() * OuterLoopCount);
+		LoopConnexions.Reserve((EdgeVertexIndices.Num() >> 1) + LoopCellBorderIndices.Num() * OuterLoopCount + MaxInnerToCornerConnexions);
 
 		for (int32 Index = 0; Index < EdgeVertexIndices.Num();)
 		{
 			const int32 StartIndex = Index++;
 			const int32 EndIndex = Index++;
-			LoopConnexions.Emplace(LoopCells[EdgeVertexIndices[StartIndex]], LoopCells[EdgeVertexIndices[EndIndex]]);
+			ensureCADKernel(LoopConnexions.Max() > LoopConnexions.Num());
+			LoopConnexions.Emplace(CellLoops[EdgeVertexIndices[StartIndex]], CellLoops[EdgeVertexIndices[EndIndex]]);
 		}
 	}
 	}
@@ -1759,29 +1836,25 @@ void FCell::InitLoopConnexions()
 	{
 		for (int32 BIndex : LoopCellBorderIndices)
 		{
-			LoopConnexions.Emplace(LoopCells[Index], LoopCells[BIndex], true);
+			ensureCADKernel(LoopConnexions.Max() > LoopConnexions.Num());
+			LoopConnexions.Emplace(CellLoops[Index], CellLoops[BIndex]);
 		}
 	}
 
 	const int32 ConnexionCount = LoopConnexions.Num();
-	for (FLoopCell& Loop : LoopCells)
+	for (FCellLoop& Loop : CellLoops)
 	{
 		Loop.Connexions.Reserve(ConnexionCount);
 	}
-
-	for (FLoopConnexion& Connexion : LoopConnexions)
-	{
-		Connexion.LinkToLoop();
-	}
 }
 
-void FCell::TryToConnectTwoSubLoopsWithShortestSegment(FLoopConnexion& LoopConnexion)
+void FCell::TryToConnectTwoSubLoopsWithShortestSegment(FCellConnexion& LoopConnexion)
 {
 	const TArray<FLoopNode*>& LoopA = LoopConnexion.Loop1.Nodes;
 	const TArray<FLoopNode*>& LoopB = LoopConnexion.Loop2.Nodes;
 
 #ifdef DEBUG_CONNECT_CELL_SUB_LOOPS_BY_NEIGHBORHOOD
-	F3DDebugSession A(Grid.bDisplay, TEXT("ConnectTwoSubLoopsWithShortestSegment"));
+	//F3DDebugSession A(Grid.bDisplay, TEXT("ConnectTwoSubLoopsWithShortestSegment"));
 #endif
 	double MinDistanceSquare = HUGE_VALUE_SQUARE;
 
@@ -1963,13 +2036,13 @@ void FIsoTriangulator::TryToConnectTwoLoopsWithIsocelesTriangle(FCell& Cell, con
 			}
 #endif 
 			const FPoint2D& NodeCoordinates = Node->Get2DPoint(EGridSpace::UniformScaled, Grid);
-			if(!NodeA1->IsDeleteOrThinNode())
+			if (!NodeA1->IsDeleteOrThinNode())
 			{
-				TryToCreateSegment(Cell, NodeA1, A1Coordinates, Node, NodeCoordinates, 0.1);
+				GetOrTryToCreateSegment(Cell, NodeA1, A1Coordinates, Node, NodeCoordinates, 0.1);
 			}
 			if (!NodeA2->IsDeleteOrThinNode())
 			{
-				TryToCreateSegment(Cell, NodeA2, A2Coordinates, Node, NodeCoordinates, 0.1);
+				GetOrTryToCreateSegment(Cell, NodeA2, A2Coordinates, Node, NodeCoordinates, 0.1);
 			}
 		}
 	}
@@ -2052,7 +2125,7 @@ void FIsoTriangulator::TryToConnectVertexSubLoopWithTheMostIsoSegment(FCell& Cel
 #ifdef DEBUG_FIND_ISO_SEGMENT_TO_LINK_OUTER_LOOP_NODES
 			Grid.DisplayIsoSegment(UniformScaled, *CandidateA, *CandidateB, 0, EVisuProperty::RedCurve);
 #endif
-			TryToCreateSegment(Cell, CandidateA, ACoordinates, CandidateB, BCoordinates, 0.1);
+			GetOrTryToCreateSegment(Cell, CandidateA, ACoordinates, CandidateB, BCoordinates, 0.1);
 			MinSlope = FlatSlope + DOUBLE_SMALL_NUMBER;
 		}
 	}
@@ -2118,33 +2191,37 @@ void FIsoTriangulator::TryToConnectTwoSubLoopsWithTheMostIsoSegment(FCell& Cell,
 #ifdef DEBUG_FIND_ISO_SEGMENT_TO_LINK_OUTER_LOOP_NODES
 			Grid.DisplayIsoSegment(EGridSpace::UniformScaled, *CandidateA, *CandidateB, 0, EVisuProperty::BlueCurve);
 #endif			
-			TryToCreateSegment(Cell, CandidateA, ACoordinates, CandidateB, BCoordinates, 0.1);
+			GetOrTryToCreateSegment(Cell, CandidateA, ACoordinates, CandidateB, BCoordinates, 0.1);
 			MinSlope = FlatSlope + DOUBLE_SMALL_NUMBER;
 		}
 	}
 }
 
-bool FIsoTriangulator::TryToCreateSegment(FCell& Cell, FLoopNode* NodeA, const FPoint2D& ACoordinates, FIsoNode* NodeB, const FPoint2D& BCoordinates, const double FlatAngle)
+FIsoSegment* FIsoTriangulator::GetOrTryToCreateSegment(FCell& Cell, FLoopNode* NodeA, const FPoint2D& ACoordinates, FIsoNode* NodeB, const FPoint2D& BCoordinates, const double FlatAngle)
 {
-
-	if (NodeA->GetSegmentConnectedTo(NodeB))
+	if (FIsoSegment* Segment = NodeA->GetSegmentConnectedTo(NodeB))
 	{
-		return false;
+		return Segment;
 	}
 
 	if (InnerSegmentsIntersectionTool.DoesIntersect(*NodeA, *NodeB))
 	{
-		return false;
+		return nullptr;
+	}
+
+	if (ThinZoneIntersectionTool.DoesIntersect(*NodeA, *NodeB))
+	{
+		return nullptr;
 	}
 
 	if (Cell.IntersectionTool.DoesIntersect(*NodeA, *NodeB))
 	{
-		return false;
+		return nullptr;
 	}
 
 	if (LoopSegmentsIntersectionTool.DoesIntersect(*NodeA, *NodeB))
 	{
-		return false;
+		return nullptr;
 	}
 
 #ifdef DEBUG_TRY_TO_CREATE_SEGMENT
@@ -2158,7 +2235,7 @@ bool FIsoTriangulator::TryToCreateSegment(FCell& Cell, FLoopNode* NodeA, const F
 	// Is Outside and not too flat at NodeA
 	if (NodeA->IsSegmentBeInsideFace(BCoordinates, Grid, FlatAngle))
 	{
-		return false;
+		return nullptr;
 	}
 
 	// Is Outside and not too flat at NodeB
@@ -2166,7 +2243,7 @@ bool FIsoTriangulator::TryToCreateSegment(FCell& Cell, FLoopNode* NodeA, const F
 	{
 		if (((FLoopNode*)NodeB)->IsSegmentBeInsideFace(ACoordinates, Grid, FlatAngle))
 		{
-			return false;
+			return nullptr;
 		}
 	}
 
@@ -2178,10 +2255,10 @@ bool FIsoTriangulator::TryToCreateSegment(FCell& Cell, FLoopNode* NodeA, const F
 #ifdef DEBUG_TRY_TO_CREATE_SEGMENT
 	DisplaySegment(ACoordinates, BCoordinates, 0, EVisuProperty::OrangePoint);
 #endif
-	return true;
+	return &Segment;
 };
 
-void FIsoTriangulator::ConnectCellCornerToLoops(FCell& Cell)
+void FIsoTriangulator::InitCellCorners(FCell& Cell)
 {
 	FIsoInnerNode* CellNodes[4];
 	int32 Index = Cell.Id;
@@ -2191,30 +2268,30 @@ void FIsoTriangulator::ConnectCellCornerToLoops(FCell& Cell)
 	CellNodes[2] = GlobalIndexToIsoInnerNodes[Index--];
 	CellNodes[3] = GlobalIndexToIsoInnerNodes[Index];
 
+	for (int32 ICell = 0; ICell < 4; ++ICell)
 	{
-		int32 ICell = 0;
-		for (; ICell < 4; ++ICell)
-		{
-			if (CellNodes[ICell])
-			{
-				break;
-			}
-		}
-		if (ICell == 4)
+		if (CellNodes[ICell])
 		{
-			// All Cell corners are null
-			return;
+			Cell.CellCorners.Emplace(ICell, *CellNodes[ICell], Grid);
 		}
 	}
+}
+
+void FIsoTriangulator::FindCandidateToConnectCellCornerToLoops(FCell& Cell)
+{
+	if (Cell.CellCorners.IsEmpty())
+	{
+		return;
+	}
 
 #ifdef DEBUG_CONNECTCELLCORNERTOINNERLOOP
 	F3DDebugSession _(bDisplay, TEXT("With cell corners"));
 #endif
 
-	TFunction<void(FLoopCell&, FIsoInnerNode*)> FindAndTryCreateCandidateSegmentToLinkLoopToCorner = [&](FLoopCell& LoopCell, FIsoInnerNode* InnerNode)
+	TFunction<void(FCellLoop&, FCellCorner&)> FindAndTryCreateCandidateSegmentToLinkLoopToCorner = [&](FCellLoop& LoopCell, FCellCorner& CellCorner)
 	{
-
-		const FPoint2D& InnerCoordinates = InnerNode->Get2DPoint(EGridSpace::UniformScaled, Grid);
+		FIsoInnerNode& CornerNode = CellCorner.CornerNode;
+		const FPoint2D& CornerPoint = CellCorner.Barycenter;
 
 		const TArray<FLoopNode*>& LoopA = LoopCell.Nodes;
 
@@ -2223,27 +2300,10 @@ void FIsoTriangulator::ConnectCellCornerToLoops(FCell& Cell)
 		for (int32 IndexA = 0; IndexA < LoopA.Num(); ++IndexA)
 		{
 			FLoopNode* NodeA = LoopA[IndexA];
-			if (NodeA->IsThinZoneNode())
-			{
-				continue;
-			}
 
 			const FPoint2D& ACoordinates = NodeA->Get2DPoint(EGridSpace::UniformScaled, Grid);
 
-			// if the candidate segment is to close to the cell border, we cannot check if in the other side of the cell there are not also a parallel candidate segment
-			// Add it in InnerToLoopCandidateSegments to be processed at the end
-			double SlopeVsCell = ComputeSlopeRelativeToNearestAxis(ACoordinates, InnerCoordinates);
-			if (SlopeVsCell < Slope::OneDegree)
-			{
-				FIsoSegment& Segment = IsoSegmentFactory.New();
-				Segment.Init(*NodeA, *InnerNode, ESegmentType::InnerToLoopU);
-				Segment.SetCandidate();
-
-				InnerToLoopCandidateSegments.Add(&Segment);
-				continue;
-			}
-
-			double SquareDistance = ACoordinates.SquareDistance(InnerCoordinates);
+			double SquareDistance = ACoordinates.SquareDistance(CornerPoint);
 			if (SquareDistance < MinDistanceSquare)
 			{
 				MinDistanceSquare = SquareDistance;
@@ -2256,7 +2316,32 @@ void FIsoTriangulator::ConnectCellCornerToLoops(FCell& Cell)
 			FLoopNode* NodeA = LoopA[MinIndexA];
 			const FPoint2D& ACoordinates = NodeA->Get2DPoint(EGridSpace::UniformScaled, Grid);
 
-			TryToCreateSegment(Cell, NodeA, ACoordinates, InnerNode, InnerCoordinates, 0.1);
+			double SlopeVsCell = ComputeSlopeRelativeToNearestAxis(ACoordinates, CornerPoint);
+			if (SlopeVsCell < Slope::OneDegree)
+			{
+				// if the candidate segment is to close to the cell border, we cannot check if in the other side of the cell there are not also a parallel candidate segment
+				// Add it in InnerToLoopCandidateSegments to be processed at the end
+
+				FIsoSegment& Segment = IsoSegmentFactory.New();
+				Segment.Init(*NodeA, CornerNode, ESegmentType::InnerToLoopU);
+				Segment.SetCandidate();
+
+				// The connection is nevertheless created to avoid failed in CheckAllLoopsConnectedTogetherAndConnect while the connection while be create in a second time based on InnerToLoopCandidateSegments
+				ensureCADKernel(Cell.LoopConnexions.Max() > Cell.LoopConnexions.Num());
+				FCellConnexion& Connexion = Cell.LoopConnexions.Emplace_GetRef(LoopCell, CellCorner);
+				Connexion.NodeA = NodeA;
+				Connexion.NodeB = &CornerNode;
+				Connexion.Segment = &Segment;
+
+				InnerToLoopCandidateSegments.Add(&Segment);
+			}
+			else
+			{
+				ensureCADKernel(Cell.LoopConnexions.Max() > Cell.LoopConnexions.Num());
+				FCellConnexion& Connexion = Cell.LoopConnexions.Emplace_GetRef(LoopCell, CellCorner);
+				Connexion.NodeA = NodeA;
+				Connexion.NodeB = &CornerNode;
+			}
 		}
 	};
 
@@ -2265,51 +2350,43 @@ void FIsoTriangulator::ConnectCellCornerToLoops(FCell& Cell)
 	Cell.IntersectionTool.AddSegments(Cell.CandidateSegments.GetData() + IntersectionToolCount, NewSegmentCount);
 	Cell.IntersectionTool.Sort();
 
-	for (int32 ICell = 0; ICell < 4; ++ICell)
+	for (FCellCorner& CellCorner : Cell.CellCorners)
 	{
-		if (CellNodes[ICell])
-		{
 #ifdef DEBUG_CONNECTCELLCORNERTOINNERLOOP
-			if (bDisplay)
-			{
-				F3DDebugSession _(bDisplay, TEXT("Cell corner"));
-				Grid.DisplayIsoNode(EGridSpace::UniformScaled, *CellNodes[ICell], 0, EVisuProperty::RedPoint);
-			}
-#endif
-			for (FLoopCell& LoopCell : Cell.LoopCells)
-			{
-				FindAndTryCreateCandidateSegmentToLinkLoopToCorner(LoopCell, CellNodes[ICell]);
-			}
-
-#ifdef DEBUG_CONNECTCELLCORNERTOINNERLOOP
-			if (bDisplay)
-			{
-				Grid.DisplayIsoSegments(TEXT("Cell CandidateSegments"), EGridSpace::UniformScaled, Cell.CandidateSegments, false, false, false, EVisuProperty::YellowCurve);
-				Wait();
-			}
+		if (bDisplay)
+		{
+			F3DDebugSession _(bDisplay, FString::Printf(TEXT("Cell corner %d"), CellCorner.Id));
+			Grid.DisplayIsoNode(EGridSpace::UniformScaled, CellCorner.CornerNode, 0, EVisuProperty::RedPoint);
+		}
 #endif
+		for (FCellLoop& LoopCell : Cell.CellLoops)
+		{
+			FindAndTryCreateCandidateSegmentToLinkLoopToCorner(LoopCell, CellCorner);
 		}
 	}
 
-	if(Cell.CandidateSegments.Num())
+#ifdef DEBUG_CONNECTCELLCORNERTOINNERLOOP
+	if (bDisplay)
 	{
-		const bool bSelectWithCellSizeCriteria = true;
-		Cell.SelectSegmentInCandidateSegments(IsoSegmentFactory, bSelectWithCellSizeCriteria);
+		DisplayCellConnexions(TEXT("Cell CandidateSegments"), Cell.LoopConnexions, EVisuProperty::YellowCurve);
+		Wait(false);
 	}
+#endif
+
 }
 
 
-void FIsoTriangulator::AddCandidateSegmentsToLinkInnerAndLoop()
+void FIsoTriangulator::SelectSegmentsToLinkInnerToLoop()
 {
 #ifdef DEBUG_ADD_CANDIDATE_SEGMENTS_TO_LINK_INNER_AND_LOOP
 	F3DDebugSession _(Grid.bDisplay, TEXT("AddCandidateSegmentsToLinkInnerAndLoop "));
 	Grid.DisplayIsoSegments(TEXT("CandidateSegments To Link Inner To Loops"), EGridSpace::UniformScaled, InnerToLoopCandidateSegments, false, false, false, EVisuProperty::RedCurve);
-	Wait();
 #endif
 
 	LoopSegmentsIntersectionTool.AddSegments(FinalToLoops);
 #ifdef DEBUG_ADD_CANDIDATE_SEGMENTS_TO_LINK_INNER_AND_LOOP
 	LoopSegmentsIntersectionTool.Display(Grid.bDisplay, TEXT("LoopSegmentsIntersectionTool with final to loop"));
+	Wait();
 #endif
 
 	TArray<TPair<double, FIsoSegment*>> LengthOfCandidateSegments;
@@ -2319,7 +2396,7 @@ void FIsoTriangulator::AddCandidateSegmentsToLinkInnerAndLoop()
 		LengthOfCandidateSegments.Emplace(Segment->Get2DLengthSquare(EGridSpace::UniformScaled, Grid), Segment);
 	}
 
-	FIntersectionSegmentTool IntersectionTool(Grid);
+	FIntersectionSegmentTool IntersectionTool(Grid, Tolerances.GeometricTolerance);
 	IntersectionTool.Reserve(InnerToLoopCandidateSegments.Num());
 
 	Algo::Sort(LengthOfCandidateSegments, [&](const TPair<double, FIsoSegment*>& P1, const TPair < double, FIsoSegment*>& P2) { return P1.Key < P2.Key; });
@@ -2335,7 +2412,6 @@ void FIsoTriangulator::AddCandidateSegmentsToLinkInnerAndLoop()
 			continue;
 		}
 
-
 		if (LoopSegmentsIntersectionTool.DoesIntersect(*Segment))
 		{
 			IsoSegmentFactory.DeleteEntity(Segment);
@@ -2351,6 +2427,7 @@ void FIsoTriangulator::AddCandidateSegmentsToLinkInnerAndLoop()
 		FinalToLoops.Add(Segment);
 		IntersectionTool.AddSegment(*Segment);
 		Segment->SetSelected();
+		Segment->SetFinalMarker();
 		Segment->ConnectToNode();
 	}
 	CandidateSegments.Empty();
diff --git a/Engine/Source/Runtime/Datasmith/CADKernel/Private/CADKernel/Mesh/Meshers/IsoTriangulator.h b/Engine/Source/Runtime/Datasmith/CADKernel/Private/CADKernel/Mesh/Meshers/IsoTriangulator.h
--- a/Engine/Source/Runtime/Datasmith/CADKernel/Private/CADKernel/Mesh/Meshers/IsoTriangulator.h
+++ b/Engine/Source/Runtime/Datasmith/CADKernel/Private/CADKernel/Mesh/Meshers/IsoTriangulator.h
@@ -26,7 +26,7 @@ struct FCell;
 
 using FMeshPolygonFunc = TFunction<void(const FGrid&, FIsoNode*[], FFaceMesh&)>;
 
-struct FLoopConnexion;
+struct FCellConnexion;
 
 struct CADKERNEL_API FIsoTriangulatorChronos
 {
@@ -109,6 +109,10 @@ protected:
 	TArray<FIsoSegment*> FinalInnerSegments;
 	TArray<FIsoSegment*> InnerToOuterSegments;
 
+	/** 
+	 * Waiting list for SelectSegmentsToLinkInnerToLoop
+	 * Segments are identified in each cell but not directly processed 
+	 */
 	TArray<FIsoSegment*> InnerToLoopCandidateSegments;
 
 	/**
@@ -117,11 +121,16 @@ protected:
 	 * Checks intersection with loop.
 	 */
 	FIntersectionSegmentTool LoopSegmentsIntersectionTool;
+
+	/** Boundary of the inner grid */
 	FIntersectionSegmentTool InnerSegmentsIntersectionTool;
+	
+	/** To check if the candidate segment intersect a iso line */
 	FIntersectionNodePairTool InnerToOuterIsoSegmentsIntersectionTool;
+	
+	/** To check if a candidate segment intersect a thin zone that is already meshed */
 	FIntersectionSegmentTool ThinZoneIntersectionTool;
 
-
 	/**
 	 * Define all the lower left index of grid node that the upper cell is surrounding a loop
 	 * This is set in FindSegmentIsoUVSurroundingSmallLoop
@@ -209,9 +218,11 @@ public:
 	void ConnectCellLoops();
 	void FindCellContainingBoundaryNodes(TArray<FCell>& Cells);
 
-	void ConnectCellCornerToLoops(FCell& Cell);
+	void FindCandidateToConnectCellCornerToLoops(FCell& Cell);
 
-	void AddCandidateSegmentsToLinkInnerAndLoop();
+	void SelectSegmentsToLinkInnerToLoop();
+
+	void InitCellCorners(FCell& Cell);
 
 	/**
 	 * Finalize the tessellation between inner grid boundary and loops.
@@ -306,7 +317,7 @@ private:
 	 *
 	 */
 	void TryToConnectVertexSubLoopWithTheMostIsoSegment(FCell& Cell, const TArray<FLoopNode*>& SubLoop);
-	bool TryToCreateSegment(FCell& Cell, FLoopNode* NodeA, const FPoint2D& ACoordinates, FIsoNode* NodeB, const FPoint2D& BCoordinates, const double FlatAngle);
+	FIsoSegment* GetOrTryToCreateSegment(FCell& Cell, FLoopNode* NodeA, const FPoint2D& ACoordinates, FIsoNode* NodeB, const FPoint2D& BCoordinates, const double FlatAngle);
 
 public:
 #ifdef CADKERNEL_DEV
@@ -327,6 +338,9 @@ public:
 	void DisplayCells(const TArray<FCell>& Cells) const;
 	void DisplayCell(const FCell& Cell) const;
 	void DrawCellBoundary(int32 Index, EVisuProperty Property) const;
+
+	void DisplayCellConnexion(const FCellConnexion& LoopConnexion, EVisuProperty Property) const;
+	void DisplayCellConnexions(const FString& Message, const TArray<FCellConnexion>& LoopConnexions, EVisuProperty Property) const;
 #endif
 
 };
diff --git a/Engine/Source/Runtime/Datasmith/CADKernel/Private/CADKernel/Mesh/Meshers/IsoTriangulator/IntersectionSegmentTool.cpp b/Engine/Source/Runtime/Datasmith/CADKernel/Private/CADKernel/Mesh/Meshers/IsoTriangulator/IntersectionSegmentTool.cpp
--- a/Engine/Source/Runtime/Datasmith/CADKernel/Private/CADKernel/Mesh/Meshers/IsoTriangulator/IntersectionSegmentTool.cpp
+++ b/Engine/Source/Runtime/Datasmith/CADKernel/Private/CADKernel/Mesh/Meshers/IsoTriangulator/IntersectionSegmentTool.cpp
@@ -23,49 +23,33 @@ FIsoSegment* FIntersectionSegmentTool::FindIntersectingSegment(const FIsoSegment
 	return const_cast<FIsoSegment*> (static_cast<const FIntersectionSegmentTool*>(this)->FindIntersectingSegment(Segment));
 }
 
-bool FIntersectionSegmentTool::DoesIntersect(const FIsoNode* StartNode, const FIsoNode* EndNode, FIsoSegment** Segment) const
-{
-	if (!StartNode || !EndNode)
-	{
-		return false;
-	}
-
-	*Segment = StartNode->GetSegmentConnectedTo(EndNode);
-	if (*Segment)
-	{
-		return false;
-	}
-	return DoesIntersect(StartNode, EndNode);
-}
-
-
-IntersectionSegmentTool::FSegment::FSegment(const FGrid& Grid, const FIsoNode& StartNode, const FIsoNode& EndNode)
-	: IntersectionToolBase::FSegment(StartNode.Get2DPoint(EGridSpace::UniformScaled, Grid), EndNode.Get2DPoint(EGridSpace::UniformScaled, Grid))
+IntersectionSegmentTool::FSegment::FSegment(const FGrid& Grid, const double Tolerance, const FIsoNode& StartNode, const FIsoNode& EndNode)
+	: IntersectionToolBase::FSegment(Tolerance, StartNode.Get2DPoint(EGridSpace::UniformScaled, Grid), EndNode.Get2DPoint(EGridSpace::UniformScaled, Grid))
 	, IsoSegment(nullptr)
 {
 }
 
-IntersectionSegmentTool::FSegment::FSegment(const FGrid& Grid, const FIsoNode& StartNode, const FPoint2D& EndPoint)
-	: IntersectionToolBase::FSegment(StartNode.Get2DPoint(EGridSpace::UniformScaled, Grid), EndPoint)
+IntersectionSegmentTool::FSegment::FSegment(const FGrid& Grid, const double Tolerance, const FIsoNode& StartNode, const FPoint2D& EndPoint)
+	: IntersectionToolBase::FSegment(Tolerance, StartNode.Get2DPoint(EGridSpace::UniformScaled, Grid), EndPoint)
 	, IsoSegment(nullptr)
 {
 }
 
-IntersectionSegmentTool::FSegment::FSegment(const FGrid& Grid, const FIsoSegment& InSegment)
-	: IntersectionToolBase::FSegment(InSegment.GetFirstNode().Get2DPoint(EGridSpace::UniformScaled, Grid), InSegment.GetSecondNode().Get2DPoint(EGridSpace::UniformScaled, Grid))
+IntersectionSegmentTool::FSegment::FSegment(const FGrid& Grid, const double Tolerance, const FIsoSegment& InSegment)
+	: IntersectionToolBase::FSegment(Tolerance, InSegment.GetFirstNode().Get2DPoint(EGridSpace::UniformScaled, Grid), InSegment.GetSecondNode().Get2DPoint(EGridSpace::UniformScaled, Grid))
 	, IsoSegment(&InSegment)
 {
 }
 
-IntersectionNodePairTool::FSegment::FSegment(const FIsoNode* InStartNode, const FIsoNode* InEndNode, const FPoint2D& InStartPoint, const FPoint2D& InEndPoint)
-	: IntersectionToolBase::FSegment(InStartPoint, InEndPoint)
+IntersectionNodePairTool::FSegment::FSegment(const double Tolerance, const FIsoNode* InStartNode, const FIsoNode* InEndNode, const FPoint2D& InStartPoint, const FPoint2D& InEndPoint)
+	: IntersectionToolBase::FSegment(Tolerance, InStartPoint, InEndPoint)
 	, StartNode(InStartNode)
 	, EndNode(InEndNode)
 {
 }
 
-IntersectionNodePairTool::FSegment::FSegment(const FGrid& Grid, const FIsoNode& InStartNode, const FIsoNode& InEndNode)
-	: IntersectionToolBase::FSegment(InStartNode.Get2DPoint(EGridSpace::UniformScaled, Grid), InEndNode.Get2DPoint(EGridSpace::UniformScaled, Grid))
+IntersectionNodePairTool::FSegment::FSegment(const FGrid& Grid, const double Tolerance, const FIsoNode& InStartNode, const FIsoNode& InEndNode)
+	: IntersectionToolBase::FSegment(Tolerance, InStartNode.Get2DPoint(EGridSpace::UniformScaled, Grid), InEndNode.Get2DPoint(EGridSpace::UniformScaled, Grid))
 	, StartNode(&InStartNode)
 	, EndNode(&InEndNode)
 {
diff --git a/Engine/Source/Runtime/Datasmith/CADKernel/Private/CADKernel/Mesh/Meshers/IsoTriangulator/IntersectionSegmentTool.h b/Engine/Source/Runtime/Datasmith/CADKernel/Private/CADKernel/Mesh/Meshers/IsoTriangulator/IntersectionSegmentTool.h
--- a/Engine/Source/Runtime/Datasmith/CADKernel/Private/CADKernel/Mesh/Meshers/IsoTriangulator/IntersectionSegmentTool.h
+++ b/Engine/Source/Runtime/Datasmith/CADKernel/Private/CADKernel/Mesh/Meshers/IsoTriangulator/IntersectionSegmentTool.h
@@ -44,9 +44,9 @@ struct FSegment
 	/**
 	 * WARNING StartPoint, EndPoint must be defined in EGridSpace::UniformScaled
 	 */
-	FSegment(const FPoint2D& StartPoint, const FPoint2D& EndPoint)
+	FSegment(const double Tolerance, const FPoint2D& StartPoint, const FPoint2D& EndPoint)
 		: Segment2D(StartPoint, EndPoint)
-		, Boundary(StartPoint, EndPoint)
+		, Boundary(StartPoint, EndPoint, Tolerance)
 	{
 		AxisMin = Boundary[EIso::IsoU].Min + Boundary[EIso::IsoV].Min;
 		AxisMax = Boundary[EIso::IsoU].Max + Boundary[EIso::IsoV].Max;
@@ -59,7 +59,7 @@ struct FSegment
 	virtual const FIsoNode* GetFirstNode() const = 0;
 	virtual const FIsoNode* GetSecondNode() const = 0;
 
-	virtual const FIsoSegment* GetIsoSegment() const 
+	virtual const FIsoSegment* GetIsoSegment() const
 	{
 		return nullptr;
 	}
@@ -117,7 +117,7 @@ struct FSegment
 			return IsSuperimposed(Segment2D, InSegment, true);
 		}
 
-		if (GetFirstNode() == EndNode )
+		if (GetFirstNode() == EndNode)
 		{
 			if (GetSecondNode() == StartNode)
 			{
@@ -182,11 +182,11 @@ struct FSegment : public IntersectionToolBase::FSegment
 {
 	const FIsoSegment* IsoSegment;
 
-	FSegment(const FGrid& Grid, const FIsoSegment& InSegment);
-	FSegment(const FGrid& Grid, const FIsoNode& StartNode, const FIsoNode& EndNode);
-	FSegment(const FGrid& Grid, const FIsoNode& StartNode, const FPoint2D& EndPoint);
-	FSegment(const FGrid& Grid, const FPoint2D& StartPoint, const FPoint2D& EndPoint)
-		: IntersectionToolBase::FSegment(StartPoint, EndPoint)
+	FSegment(const FGrid& Grid, const double Tolerance, const FIsoSegment& InSegment);
+	FSegment(const FGrid& Grid, const double Tolerance, const FIsoNode& StartNode, const FIsoNode& EndNode);
+	FSegment(const FGrid& Grid, const double Tolerance, const FIsoNode& StartNode, const FPoint2D& EndPoint);
+	FSegment(const FGrid& Grid, const double Tolerance, const FPoint2D& StartPoint, const FPoint2D& EndPoint)
+		: IntersectionToolBase::FSegment(Tolerance, StartPoint, EndPoint)
 		, IsoSegment(nullptr)
 	{
 	}
@@ -213,8 +213,8 @@ struct FSegment : public IntersectionToolBase::FSegment
 	const FIsoNode* StartNode;
 	const FIsoNode* EndNode;
 
-	FSegment(const FIsoNode* StartNode, const FIsoNode* EndNode, const FPoint2D& StartPoint, const FPoint2D& EndPoint);
-	FSegment(const FGrid& Grid, const FIsoNode& StartNode, const FIsoNode& EndNode);
+	FSegment(const double Tolerance, const FIsoNode* StartNode, const FIsoNode* EndNode, const FPoint2D& StartPoint, const FPoint2D& EndPoint);
+	FSegment(const FGrid& Grid, const double Tolerance, const FIsoNode& StartNode, const FIsoNode& EndNode);
 
 	virtual bool IsValid() const override
 	{
@@ -253,12 +253,19 @@ protected:
 
 	TArray<SegmentType> Segments;
 	bool bSegmentsAreSorted;
+
+	const double Tolerance;
+
+	// Allow a candidate segment to be parallel with and overlap an existing segment.
+	// This is used by InnerToOuterIsoSegmentsIntersectionTool, which existing segment are Isolines. 
+	// So this allow candidate segment to overlap iso segment. Otherwise, the candidate segment has one intersection with an isoline when we just want that is parallel with.
 	const bool bAllowOverlapping;
 
 public:
-	TIntersectionSegmentTool(const FGrid& InGrid, bool bInAllowOverlapping = false)
+	TIntersectionSegmentTool(const FGrid& InGrid, const double InTolerance, const bool bInAllowOverlapping = false)
 		: Grid(InGrid)
 		, bSegmentsAreSorted(false)
+		, Tolerance(InTolerance)
 		, bAllowOverlapping(bInAllowOverlapping)
 	{
 	}
@@ -307,7 +314,7 @@ public:
 	const SegmentType* FindIntersectingSegment(const ExtremityType1* StartExtremity, const ExtremityType2* EndExtremity) const
 	{
 		using namespace IntersectionSegmentTool;
-		const SegmentType InSegment(Grid, *StartExtremity, *EndExtremity);
+		const SegmentType InSegment(Grid, Tolerance, *StartExtremity, *EndExtremity);
 
 		for (const SegmentType& Segment : Segments)
 		{
@@ -329,6 +336,11 @@ public:
 				}
 			}
 
+			if (Segment.IsParallelWith(InSegment))
+			{
+				printf("");
+			}
+
 			if (bAllowOverlapping && Segment.IsParallelWith(InSegment))
 			{
 				continue;
@@ -358,8 +370,8 @@ public:
 	template<typename ExtremityType1, typename ExtremityType2>
 	int32 FindIntersectingSegments(const ExtremityType1* StartExtremity, const ExtremityType2* EndExtremity, TArray<const FIsoSegment*>* OutIntersectedSegments) const
 	{
-		SegmentType InSegment(Grid, *StartExtremity, *EndExtremity);
-		if(OutIntersectedSegments)
+		SegmentType InSegment(Grid, Tolerance, *StartExtremity, *EndExtremity);
+		if (OutIntersectedSegments)
 		{
 			OutIntersectedSegments->Empty(10);
 		}
@@ -445,8 +457,8 @@ public:
 class FIntersectionSegmentTool : public TIntersectionSegmentTool<IntersectionSegmentTool::FSegment>
 {
 public:
-	FIntersectionSegmentTool(const FGrid& InGrid)
-		: TIntersectionSegmentTool<IntersectionSegmentTool::FSegment>(InGrid)
+	FIntersectionSegmentTool(const FGrid& InGrid, const double Tolerance)
+		: TIntersectionSegmentTool<IntersectionSegmentTool::FSegment>(InGrid, Tolerance)
 	{
 	}
 
@@ -460,7 +472,7 @@ public:
 		if (SegmentIndex != INDEX_NONE)
 		{
 			Segments.RemoveAt(SegmentIndex);
-			Segments.EmplaceAt(SegmentIndex, Grid, *Segment);
+			Segments.EmplaceAt(SegmentIndex, Grid, Tolerance, *Segment);
 			bSegmentsAreSorted = false;
 			return true;
 		}
@@ -491,13 +503,13 @@ public:
 	void AddSegment(const FIsoSegment& Segment)
 	{
 		bSegmentsAreSorted = false;
-		Segments.Emplace(Grid, Segment);
+		Segments.Emplace(Grid, Tolerance, Segment);
 	}
 
 	void AddSegment(const FPoint2D& StartPoint, const FPoint2D& EndPoint)
 	{
 		bSegmentsAreSorted = false;
-		Segments.Emplace(Grid, StartPoint, EndPoint);
+		Segments.Emplace(Grid, Tolerance, StartPoint, EndPoint);
 	}
 
 	/**
@@ -534,7 +546,6 @@ public:
 		return TIntersectionSegmentTool<IntersectionSegmentTool::FSegment>::FindIntersectingSegment(&StartPoint, &EndPoint) != nullptr;
 	}
 
-	
 	/**
 	 * Allow StartNode to be connected to one segment
 	 * WARNING EndPoint must be defined in EGridSpace::UniformScaled
@@ -544,11 +555,6 @@ public:
 		return TIntersectionSegmentTool<IntersectionSegmentTool::FSegment>::FindIntersectingSegment(&StartNode, &EndPoint) != nullptr;
 	}
 
-	/**
-	 * Find if a Segment already exists, if return false with the segment
-	 */
-	bool DoesIntersect(const FIsoNode* StartNode, const FIsoNode* EndNode, FIsoSegment** Segment) const;
-
 	bool DoesIntersect(const FIsoNode* StartNode, const FIsoNode* EndNode) const
 	{
 		return TIntersectionSegmentTool<IntersectionSegmentTool::FSegment>::FindIntersectingSegment(StartNode, EndNode) != nullptr;
@@ -572,32 +578,51 @@ public:
 		return nullptr;
 	}
 
+	const FIsoSegment* FindIntersectingSegment(const FIsoNode* StartNode, const FIsoNode* EndNode) const
+	{
+		const IntersectionSegmentTool::FSegment* IntersectingSegment = TIntersectionSegmentTool<IntersectionSegmentTool::FSegment>::FindIntersectingSegment(StartNode, EndNode);
+		if (IntersectingSegment)
+		{
+			return IntersectingSegment->IsoSegment;
+		}
+		return nullptr;
+	}
+
+	bool FindIntersectingSegments(const FIsoNode* StartNode, const FIsoNode* EndNode, TArray<const FIsoSegment*>& OutIntersections) const
+	{
+		return TIntersectionSegmentTool<IntersectionSegmentTool::FSegment>::FindIntersectingSegments(StartNode, EndNode, &OutIntersections) > 0;
+	}
+
 	bool FindIntersectingSegments(const FIsoNode& StartNode, const FIsoNode& EndNode, TArray<const FIsoSegment*>& OutIntersections) const
 	{
 		return TIntersectionSegmentTool<IntersectionSegmentTool::FSegment>::FindIntersectingSegments(&StartNode, &EndNode, &OutIntersections) > 0;
 	}
 
+	int32 CountIntersections(const FIsoNode* StartNode, const FIsoNode* EndNode) const
+	{
+		return TIntersectionSegmentTool<IntersectionSegmentTool::FSegment>::FindIntersectingSegments(StartNode, EndNode, nullptr);
+	}
 };
 
 class FIntersectionNodePairTool : public TIntersectionSegmentTool<IntersectionNodePairTool::FSegment>
 {
 
 public:
-	FIntersectionNodePairTool(const FGrid& InGrid, bool bInAllowOverlaping)
-		: TIntersectionSegmentTool<IntersectionNodePairTool::FSegment>(InGrid, bInAllowOverlaping)
+	FIntersectionNodePairTool(const FGrid& InGrid, const double Tolerance, bool bInAllowOverlaping)
+		: TIntersectionSegmentTool<IntersectionNodePairTool::FSegment>(InGrid, Tolerance, bInAllowOverlaping)
 	{
 	}
 
 	void AddSegment(const FIsoNode& StartNode, const FIsoNode& EndNode)
 	{
 		bSegmentsAreSorted = false;
-		Segments.Emplace(Grid, StartNode, EndNode);
+		Segments.Emplace(Grid, Tolerance, StartNode, EndNode);
 	}
 
 	void AddSegment(const FIsoNode* StartNode, const FIsoNode* EndNode, const FPoint2D& StartPoint, const FPoint2D& EndPoint)
 	{
 		bSegmentsAreSorted = false;
-		Segments.Emplace(StartNode, EndNode, StartPoint, EndPoint);
+		Segments.Emplace(Tolerance, StartNode, EndNode, StartPoint, EndPoint);
 	}
 
 	int32 CountIntersections(const FIsoNode& StartNode, const FIsoNode& EndNode) const
diff --git a/Engine/Source/Runtime/Datasmith/CADKernel/Private/CADKernel/Mesh/Meshers/IsoTriangulator/IsoCell.h b/Engine/Source/Runtime/Datasmith/CADKernel/Private/CADKernel/Mesh/Meshers/IsoTriangulator/IsoCell.h
--- a/Engine/Source/Runtime/Datasmith/CADKernel/Private/CADKernel/Mesh/Meshers/IsoTriangulator/IsoCell.h
+++ b/Engine/Source/Runtime/Datasmith/CADKernel/Private/CADKernel/Mesh/Meshers/IsoTriangulator/IsoCell.h
@@ -17,17 +17,18 @@ class FGrid;
 class FIsoTriangulator;
 struct FLoopConnexion;
 
-struct FLoopCell
+struct FCellLoop
 {
 	int32 Id;
 	FPoint2D Barycenter;
 	TArray<FLoopNode*> Nodes;
-	TArray<FLoopConnexion*> Connexions;
+	TArray<FCellConnexion*> Connexions;
 	bool bIsOuterLoop = false;
+	bool bIsConnected = false;
 
-	FLoopCell(const int32 InIndex, TArray<FLoopNode*>& InNodes, const FGrid& Grid)
+	FCellLoop(const int32 InIndex, TArray<FLoopNode*>& InNodes, const FGrid& Grid)
 		: Id(InIndex)
-		, Barycenter(FPoint::ZeroPoint)
+		, Barycenter(FPoint2D::ZeroPoint)
 		, Nodes(InNodes)
 		, bIsOuterLoop(Nodes[0]->GetLoopIndex() == 0)
 	{
@@ -37,50 +38,94 @@ struct FLoopCell
 		}
 		Barycenter /= (double)Nodes.Num();
 	}
+
+	FCellLoop(const int32 InIndex)
+		: Id(InIndex)
+		, Barycenter(FPoint::ZeroPoint)
+	{
+	}
+
+	virtual ~FCellLoop() = default;
+
+	virtual bool IsCellCorner() const 
+	{
+		return false;
+	}
+
+	virtual bool IsCellCornerOrOuterLoop() const
+	{
+		return bIsOuterLoop;
+	}
+
+	void SpreadAsConnected();
 };
 
-struct FLoopConnexion
+struct FCellCorner : public FCellLoop
 {
-	bool bIsConnexionWithBorder = false;
-	FLoopCell& Loop1;
-	FLoopCell& Loop2;
+	FIsoInnerNode& CornerNode;
+
+	FCellCorner(const int32 InIndex, FIsoInnerNode& InNode, const FGrid& Grid)
+		: FCellLoop(InIndex)
+		, CornerNode(InNode)
+	{
+		Barycenter = CornerNode.Get2DPoint(EGridSpace::UniformScaled, Grid);
+	}
+
+	virtual bool IsCellCorner() const override
+	{
+		return true;
+	}
+
+	virtual bool IsCellCornerOrOuterLoop() const
+	{
+		return true;
+	}
+
+};
+
+struct FCellConnexion
+{
+	bool bIsConnexionWithOuter = false;
+	FCellLoop& Loop1;
+	FCellLoop& Loop2;
 
 	double MinDistance = HUGE_VALUE_SQUARE;
 
 	FLoopNode* NodeA = nullptr;
-	FLoopNode* NodeB = nullptr;
+	FIsoNode* NodeB = nullptr;
 
 	FIsoSegment* Segment = nullptr;
 
-	FLoopConnexion(FLoopCell& InLoop1, FLoopCell& InLoop2, bool bIsConnectingBoder = false)
-		: bIsConnexionWithBorder(bIsConnectingBoder)
-		, Loop1(InLoop1)
+	FCellConnexion(FCellLoop& InLoop1, FCellLoop& InLoop2)
+		: Loop1(InLoop1)
 		, Loop2(InLoop2)
 	{
+		bIsConnexionWithOuter = Loop1.bIsOuterLoop || (!Loop2.IsCellCorner() && Loop2.bIsOuterLoop);
+		Loop1.Connexions.Add(this);
+		Loop2.Connexions.Add(this);
 	}
 
-	void LinkToLoop()
+	const FCellLoop* GetOtherLoop(const FCellLoop* Loop) const
 	{
-		Loop1.Connexions.Add(this);
-		Loop2.Connexions.Add(this);
+		return (Loop == &Loop1) ? &Loop2 : &Loop1;
 	}
 
-	const FLoopCell* GetOtherLoop(const FLoopCell* Loop)
+	FCellLoop* GetOtherLoop(const FCellLoop* Loop)
 	{
 		return (Loop == &Loop1) ? &Loop2 : &Loop1;
 	}
 
 	bool IsShortestPath(const int32 MaxLoopCount)
 	{
-		TMap<const FLoopCell*, double> DistanceToLoops;
+		TMap<const FCellLoop*, double> DistanceToLoops;
 		DistanceToLoops.Reserve(MaxLoopCount);
-		DistanceToLoops.Add(&Loop1, 0.);
+		DistanceToLoops.Add(&Loop2, 0.);
 		double MinPathDistance = HUGE_VALUE;
 		while (true)
 		{
 			double DistanceToCurrent = HUGE_VALUE;
-			TPair<const FLoopCell*, double>* CurrentLoop = nullptr;
-			for (TPair<const FLoopCell*, double>& DistanceToLoop : DistanceToLoops)
+			TPair<const FCellLoop*, double>* CurrentLoop = nullptr;
+			for (TPair<const FCellLoop*, double>& DistanceToLoop : DistanceToLoops)
 			{
 				if (DistanceToLoop.Value >= 0. && DistanceToLoop.Value < DistanceToCurrent)
 				{
@@ -94,19 +139,29 @@ struct FLoopConnexion
 				return true;
 			}
 
-			if (CurrentLoop->Key == &Loop2)
+			if (CurrentLoop->Key == &Loop1)
 			{
 				return (CurrentLoop->Value > MinDistance);
 			}
 
-			for (FLoopConnexion* Connexion : CurrentLoop->Key->Connexions)
+			for (FCellConnexion* Connexion : CurrentLoop->Key->Connexions)
 			{
-				if (Connexion == this || Connexion->bIsConnexionWithBorder)
+				if (Connexion == this)
+				{
+					continue;
+				}
+
+				const FCellLoop* NextLoop = Connexion->GetOtherLoop(CurrentLoop->Key);
+				if (NextLoop->IsCellCorner())
+				{
+					continue;
+				}
+
+				if (NextLoop->bIsOuterLoop)
 				{
 					continue;
 				}
 
-				const FLoopCell* NextLoop = Connexion->GetOtherLoop(CurrentLoop->Key);
 				const double DistanceToNextByCurrent = CurrentLoop->Value + Connexion->MinDistance;
 
 				double* DistanceToNextLoop = DistanceToLoops.Find(NextLoop);
@@ -128,13 +183,155 @@ struct FLoopConnexion
 
 		return true;
 	}
+
+	bool IsShortestPathToOuterLoop(const int32 MaxLoopCount)
+	{
+		TMap<const FCellLoop*, double> DistanceToLoops;
+		DistanceToLoops.Reserve(MaxLoopCount);
+		DistanceToLoops.Add(&Loop2, 0.);
+		double MinPathDistance = HUGE_VALUE;
+
+		while (true)
+		{
+			double DistanceToCurrent = HUGE_VALUE;
+			TPair<const FCellLoop*, double>* CurrentLoopPair = nullptr;
+			for (TPair<const FCellLoop*, double>& DistanceToLoop : DistanceToLoops)
+			{
+				if (DistanceToLoop.Value >= 0. && DistanceToLoop.Value < DistanceToCurrent)
+				{
+					DistanceToCurrent = DistanceToLoop.Value;
+					CurrentLoopPair = &DistanceToLoop;
+				}
+			}
+
+			if (!CurrentLoopPair || CurrentLoopPair->Value > MinDistance)
+			{
+				return true;
+			}
+
+			const FCellLoop* CurrentLoop = CurrentLoopPair->Key;
+			double& CurrentLoopDistance = CurrentLoopPair->Value;
+
+			for (FCellConnexion* Connexion : CurrentLoop->Connexions)
+			{
+				if (Connexion == this)
+				{
+					continue;
+				}
+
+				const FCellLoop* NextLoop = Connexion->GetOtherLoop(CurrentLoop);
+				if (NextLoop->IsCellCorner())
+				{
+					continue;
+				}
+
+				const double DistanceToNextByCurrent = CurrentLoopDistance + Connexion->MinDistance;
+				if (NextLoop->bIsOuterLoop)
+				{
+					if (DistanceToNextByCurrent < MinDistance)
+					{
+						return false;
+					}
+					continue;
+				}
+
+				double* DistanceToNextLoop = DistanceToLoops.Find(NextLoop);
+				if (DistanceToNextLoop)
+				{
+					if (*DistanceToNextLoop > DistanceToNextByCurrent)
+					{
+						*DistanceToNextLoop = DistanceToNextByCurrent;
+					}
+				}
+				else
+				{
+					DistanceToLoops.Add(NextLoop, DistanceToNextByCurrent);
+				}
+			}
+
+			CurrentLoopDistance = -HUGE_VALUE;
+		}
+
+		return true;
+	}
+
+	bool IsShortestPathToCorner(const int32 MaxLoopCount)
+	{
+		TMap<const FCellLoop*, double> DistanceToLoops;
+		DistanceToLoops.Reserve(MaxLoopCount);
+		DistanceToLoops.Add(&Loop2, 0.);
+		double MinPathDistance = HUGE_VALUE;
+
+		while (true)
+		{
+			double DistanceToCurrent = HUGE_VALUE;
+			TPair<const FCellLoop*, double>* CurrentLoopPair = nullptr;
+			for (TPair<const FCellLoop*, double>& DistanceToLoop : DistanceToLoops)
+			{
+				if (DistanceToLoop.Value >= 0. && DistanceToLoop.Value < DistanceToCurrent)
+				{
+					DistanceToCurrent = DistanceToLoop.Value;
+					CurrentLoopPair = &DistanceToLoop;
+				}
+			}
+
+			if (!CurrentLoopPair || CurrentLoopPair->Value > MinDistance)
+			{
+				return true;
+			}
+
+			const FCellLoop* CurrentLoop = CurrentLoopPair->Key;
+			double& CurrentLoopDistance = CurrentLoopPair->Value;
+
+			for (FCellConnexion* Connexion : CurrentLoop->Connexions)
+			{
+				if (Connexion == this)
+				{
+					continue;
+				}
+
+				const FCellLoop* NextLoop = Connexion->GetOtherLoop(CurrentLoop);
+				if (NextLoop->bIsOuterLoop)
+				{
+					continue;
+				}
+
+				const double DistanceToNextByCurrent = CurrentLoopDistance + Connexion->MinDistance;
+				if (NextLoop->IsCellCorner())
+				{
+					if ((NextLoop == &Loop1) && (DistanceToNextByCurrent < MinDistance))
+					{
+						return false;
+					}
+					continue;
+				}
+
+				double* DistanceToNextLoop = DistanceToLoops.Find(NextLoop);
+				if (DistanceToNextLoop)
+				{
+					if (*DistanceToNextLoop > DistanceToNextByCurrent)
+					{
+						*DistanceToNextLoop = DistanceToNextByCurrent;
+					}
+				}
+				else
+				{
+					DistanceToLoops.Add(NextLoop, DistanceToNextByCurrent);
+				}
+			}
+
+			CurrentLoopDistance = -HUGE_VALUE;
+		}
+
+		return true;
+	}
 };
 
 struct FCellPath
 {
 	double Length = 0;
-	FLoopCell* CurrentLoop;
-	TArray<FLoopCell> Path;
+	FCellLoop* CurrentLoop;
+	TArray<FCellLoop> Path;
 };
 
 struct FCell
@@ -151,15 +348,16 @@ struct FCell
 
 	FIntersectionSegmentTool IntersectionTool;
 
-	TArray<FLoopCell> LoopCells;
-	TArray<FLoopConnexion> LoopConnexions;
+	TArray<FCellLoop> CellLoops;
+	TArray<FCellCorner> CellCorners;
+	TArray<FCellConnexion> LoopConnexions;
 	TArray<int32> LoopCellBorderIndices;
 
 	FCell(const int32 InLoopIndex, TArray<FLoopNode*>& InNodes, FIsoTriangulator& InTriangulator)
 		: Triangulator(InTriangulator)
 		, Grid(Triangulator.GetGrid())
 		, Id(InLoopIndex)
-		, IntersectionTool(Triangulator.GetGrid())
+		, IntersectionTool(Triangulator.GetGrid(), Triangulator.Tolerances.GeometricTolerance)
 	{
 		const int32 NodeCount = InNodes.Num();
 		ensureCADKernel(NodeCount > 0);
@@ -181,16 +379,16 @@ struct FCell
 			PreviousNode = Node;
 		}
 
-		LoopCells.Reserve(LoopCount);
+		CellLoops.Reserve(LoopCount);
 
 		LoopCount = 0;
 		TArray<FLoopNode*> LoopNodes;
 		LoopNodes.Reserve(NodeCount);
 
 		int32 LoopIndex = -1;
-		FLoopCell* FirstLoopCell = nullptr;
+		FCellLoop* FirstLoopCell = nullptr;
 
-		TFunction<void(TArray<FLoopNode*>&)> MakeLoopCell = [&FirstLoopCell, &LoopCells = LoopCells, &LoopIndex, &Grid = Grid](TArray<FLoopNode*>& LoopNodes)
+		TFunction<void(TArray<FLoopNode*>&)> MakeLoopCell = [&FirstLoopCell, &CellLoops = CellLoops, &LoopIndex, &Grid = Grid](TArray<FLoopNode*>& LoopNodes)
 		{
 			if (LoopNodes.Num())
 			{
@@ -201,13 +399,13 @@ struct FCell
 				}
 				else
 				{
-					LoopCells.Emplace(LoopCells.Num(), LoopNodes, Grid);
+					CellLoops.Emplace(CellLoops.Num(), LoopNodes, Grid);
 				}
 
 				if (LoopIndex != LoopNodes[0]->GetLoopIndex())
 				{
 					LoopIndex = LoopNodes[0]->GetLoopIndex();
-					FirstLoopCell = &LoopCells.Last();
+					FirstLoopCell = &CellLoops.Last();
 				}
 			}
 		};
@@ -230,7 +428,7 @@ struct FCell
 		}
 		MakeLoopCell(LoopNodes);
 
-		for(const FLoopCell& LoopCell : LoopCells)
+		for(const FCellLoop& LoopCell : CellLoops)
 		{
 			if (LoopCell.bIsOuterLoop)
 			{
@@ -248,8 +446,8 @@ struct FCell
 	TArray<TPair<int32, FPoint2D>> GetLoopBarycenters()
 	{
 		TArray<TPair<int32, FPoint2D>> LoopBarycenters;
-		LoopBarycenters.Reserve(LoopCells.Num());
-		for (const FLoopCell& LoopCell : LoopCells)
+		LoopBarycenters.Reserve(CellLoops.Num());
+		for (const FCellLoop& LoopCell : CellLoops)
 		{
 			if(!LoopCell.bIsOuterLoop)
 			{
@@ -259,7 +457,14 @@ struct FCell
 		return MoveTemp(LoopBarycenters);
 	}
 
-	void ConnectLoopsByNeighborhood();
+	void FindCandidateToConnectLoopsByNeighborhood();
+	void FindCandidateToConnectCellCornerToLoops();
+
+	void SelectSegmentToConnectLoops(TFactory<FIsoSegment>& SegmentFactory);
+	void SelectSegmentToConnectLoopToCorner(TFactory<FIsoSegment>& SegmentFactory);
+	void CheckAllLoopsConnectedTogetherAndConnect();
+
+	//void ConnectLoopsByNeighborhood2();
 
 	/**
 	 *  SubLoopA                  SubLoopB
@@ -273,85 +478,9 @@ struct FCell
 	 *
 	 *     ======= ShortestSegment
 	 */
-	void TryToConnectTwoSubLoopsWithShortestSegment(FLoopConnexion& LoopConnexion);
-
-	void TryToCreateSegment(FLoopConnexion& LoopConnexion);
-
-	void SelectSegmentInCandidateSegments(TFactory<FIsoSegment>& SegmentFactory, bool bSelectWithCellSizeCriteria)
-	{
-#ifdef DEBUG_SELECT_SEGMENT
-		F3DDebugSession _(Grid.bDisplay, TEXT("SelectSegmentInCandidateSegments "));
-		if(Grid.bDisplay)
-		{
-			IntersectionTool.Display(Grid.bDisplay, TEXT("Cell.IntersectionTool at SelectSegmentInCandidateSegments start"));
-			Wait();
-		}
-#endif
-
-		double CellSquareLength = HUGE_VALUE;
-		if(bSelectWithCellSizeCriteria)
-		{
-			int32 IndexU;
-			int32 IndexV;
-			Grid.UVIndexFromGlobalIndex(Id, IndexU, IndexV);
-
-			CellSquareLength = Grid.GetInner2DPoint(EGridSpace::UniformScaled, IndexU, IndexV).SquareDistance(Grid.GetInner2DPoint(EGridSpace::UniformScaled, IndexU, IndexV + 1));
-			const double CellLength2 = Grid.GetInner2DPoint(EGridSpace::UniformScaled, IndexU, IndexV).SquareDistance(Grid.GetInner2DPoint(EGridSpace::UniformScaled, IndexU + 1, IndexV));
-			CellSquareLength = FMath::Min(CellSquareLength, CellLength2);
-		}
-
-		TArray<TPair<double, FIsoSegment*>> LengthOfCandidateSegments;
-		LengthOfCandidateSegments.Reserve(CandidateSegments.Num());
-		for (FIsoSegment* Segment : CandidateSegments)
-		{
-			LengthOfCandidateSegments.Emplace(Segment->Get2DLengthSquare(EGridSpace::UniformScaled, Grid), Segment);
-		}
-
-		Algo::Sort(LengthOfCandidateSegments, [&](const TPair<double, FIsoSegment*>& P1, const TPair < double, FIsoSegment*>& P2) { return P1.Key < P2.Key; });
-
-		// Validate the first candidate segments
-		for (const TPair<double, FIsoSegment*>& Candidate : LengthOfCandidateSegments)
-		{
-			if (Candidate.Key > CellSquareLength)
-			{
-				break;
-			}
-
-			FIsoSegment* Segment = Candidate.Value;
-#ifdef DEBUG_SELECT_SEGMENT
-			F3DDebugSession _(Grid.bDisplay, TEXT("Segment"));
-#endif
-			if (IntersectionTool.DoesIntersect(*Segment))
-			{
-#ifdef DEBUG_SELECT_SEGMENT
-				F3DDebugSession _(Grid.bDisplay, TEXT("SelectSegmentInCandidateSegments "));
-				Grid.DisplayIsoSegment(EGridSpace::UniformScaled, *Segment, YellowCurve);
-#endif
-				SegmentFactory.DeleteEntity(Segment);
-				continue;
-			}
-
-			if (FIsoSegment::IsItAlreadyDefined(&Segment->GetFirstNode(), &Segment->GetSecondNode()))
-			{
-#ifdef DEBUG_SELECT_SEGMENT
-				Grid.DisplayIsoSegment(EGridSpace::UniformScaled, *Segment, GreenCurve);
-#endif
-				SegmentFactory.DeleteEntity(Segment);
-				continue;
-			}
-
-#ifdef DEBUG_SELECT_SEGMENT
-			Grid.DisplayIsoSegment(EGridSpace::UniformScaled, *Segment, BlueCurve);
-#endif
-
-			FinalSegments.Add(Segment);
-			IntersectionTool.AddSegment(*Segment);
-			Segment->SetSelected();
-			Segment->ConnectToNode();
-		}
-		CandidateSegments.Empty();
-	}
+	void TryToConnectTwoSubLoopsWithShortestSegment(FCellConnexion& LoopConnexion);
 
+	void TryToCreateSegment(FCellConnexion& LoopConnexion);
 };
 
 } // namespace UE::CADKernel
diff --git a/Engine/Source/Runtime/Datasmith/CADKernel/Private/CADKernel/Mesh/Meshers/IsoTriangulator/IsoNode.h b/Engine/Source/Runtime/Datasmith/CADKernel/Private/CADKernel/Mesh/Meshers/IsoTriangulator/IsoNode.h
--- a/Engine/Source/Runtime/Datasmith/CADKernel/Private/CADKernel/Mesh/Meshers/IsoTriangulator/IsoNode.h
+++ b/Engine/Source/Runtime/Datasmith/CADKernel/Private/CADKernel/Mesh/Meshers/IsoTriangulator/IsoNode.h
@@ -393,38 +393,6 @@ public:
 		return (!IsPointPInsideSectorABC(GetPreviousNode().Get2DPoint(EGridSpace::UniformScaled, Grid), Get2DPoint(EGridSpace::UniformScaled, Grid), GetNextNode().Get2DPoint(EGridSpace::UniformScaled, Grid), EndSegmentCoordinate, FlatAngle));
 	}
 
-	void SetAsIsoU()
-	{
-		States |= EIsoNodeStates::NearlyIsoU;
-	}
-
-	bool IsIsoU() const
-	{
-		return (States & EIsoNodeStates::NearlyIsoU) == EIsoNodeStates::NearlyIsoU;
-	}
-
-	void SetAsIsoV()
-	{
-		States |= EIsoNodeStates::NearlyIsoV;
-	}
-
-	bool IsIsoV() const
-	{
-		return (States & EIsoNodeStates::NearlyIsoV) == EIsoNodeStates::NearlyIsoV;
-	}
-
-	constexpr bool IsIso(EIso Axe) const
-	{
-		if (Axe == IsoV)
-		{
-			return IsIsoV();
-		}
-		else
-		{
-			return IsIsoU();
-		}
-	}
-
 	virtual uint32 GetTypeHash() const override
 	{
 		return HashCombine(GetIndex(), GetLoopIndex() + 1);
diff --git a/Engine/Source/Runtime/Datasmith/CADKernel/Private/CADKernel/Mesh/Meshers/IsoTriangulator/IsoSegment.h b/Engine/Source/Runtime/Datasmith/CADKernel/Private/CADKernel/Mesh/Meshers/IsoTriangulator/IsoSegment.h
--- a/Engine/Source/Runtime/Datasmith/CADKernel/Private/CADKernel/Mesh/Meshers/IsoTriangulator/IsoSegment.h
+++ b/Engine/Source/Runtime/Datasmith/CADKernel/Private/CADKernel/Mesh/Meshers/IsoTriangulator/IsoSegment.h
@@ -26,7 +26,7 @@ enum class EIsoSegmentStates : uint8
 {
 	None = 0x00u,	// No flags.
 
-	CandidateFlag = 0x01u,
+	Candidate = 0x01u,
 
 	LeftCycle = 0x02u,
 	RightCycle = 0x04u,
@@ -37,10 +37,9 @@ enum class EIsoSegmentStates : uint8
 	RightTriangle = 0x20u,
 	SegmentComplete = 0x30u,
 
-	StatusIsoU = 0x40u,
-	StatusIsoV = 0x80u,
+	Final = 0x40u,
+	Delete = 0x80u,
 
-	Delete = 0xFFu,
 	All = 0xFFu
 };
 
@@ -140,17 +139,27 @@ public:
 
 	void SetSelected()
 	{
-		States &= ~EIsoSegmentStates::CandidateFlag;
+		States &= ~EIsoSegmentStates::Candidate;
 	}
 
 	void SetCandidate()
 	{
-		States |= EIsoSegmentStates::CandidateFlag;
+		States |= EIsoSegmentStates::Candidate;
 	}
 
-	bool IsACandidate() const // [0 Pi/2]
+	bool IsACandidate() const 
 	{
-		return (States & EIsoSegmentStates::CandidateFlag) == EIsoSegmentStates::CandidateFlag;
+		return (States & EIsoSegmentStates::Candidate) == EIsoSegmentStates::Candidate;
+	}
+
+	void SetFinalMarker()
+	{
+		States |= EIsoSegmentStates::Final;
+	}
+
+	bool IsAFinalSegment() const
+	{
+		return (States & EIsoSegmentStates::Final) == EIsoSegmentStates::Final;
 	}
 
 	bool HasCycleOnLeft()
@@ -232,26 +241,6 @@ public:
 		States &= ~EIsoSegmentStates::SegmentComplete;
 	}
 
-	void SetAsIsoU()
-	{
-		States |= EIsoSegmentStates::StatusIsoU;
-	}
-
-	bool IsIsoU() const
-	{
-		return (States & EIsoSegmentStates::StatusIsoU) == EIsoSegmentStates::StatusIsoU;
-	}
-
-	void SetAsIsoV()
-	{
-		States |= EIsoSegmentStates::StatusIsoV;
-	}
-
-	bool IsIsoV() const
-	{
-		return (States & EIsoSegmentStates::StatusIsoV) == EIsoSegmentStates::StatusIsoV;
-	}
-
 	void SetAsDegenerated()
 	{
 		States |= EIsoSegmentStates::Degenerate;
diff --git a/Engine/Source/Runtime/Datasmith/CADKernel/Private/CADKernel/Mesh/Meshers/ParametricMesher.cpp b/Engine/Source/Runtime/Datasmith/CADKernel/Private/CADKernel/Mesh/Meshers/ParametricMesher.cpp
--- a/Engine/Source/Runtime/Datasmith/CADKernel/Private/CADKernel/Mesh/Meshers/ParametricMesher.cpp
+++ b/Engine/Source/Runtime/Datasmith/CADKernel/Private/CADKernel/Mesh/Meshers/ParametricMesher.cpp
@@ -66,6 +66,7 @@ void FParametricMesher::MeshEntities(TArray<FTopologicalShapeEntity*>& InEntitie
 			continue;
 		}
 		FaceCount += TopologicalEntity->FaceCount();
+		TopologicalEntity->ResetMarkersRecursively();
 	}
 	Faces.Reserve(Faces.Num() + FaceCount);
 
diff --git a/Engine/Source/Runtime/Datasmith/CADKernel/Private/CADKernel/Mesh/Structure/LoopCleaner.cpp b/Engine/Source/Runtime/Datasmith/CADKernel/Private/CADKernel/Mesh/Structure/LoopCleaner.cpp
--- a/Engine/Source/Runtime/Datasmith/CADKernel/Private/CADKernel/Mesh/Structure/LoopCleaner.cpp
+++ b/Engine/Source/Runtime/Datasmith/CADKernel/Private/CADKernel/Mesh/Structure/LoopCleaner.cpp
@@ -22,7 +22,7 @@ FLoopCleaner::FLoopCleaner(FIsoTriangulator& Triangulator)
 	, LoopSegments(Triangulator.LoopSegments)
 	, IsoSegmentFactory(Triangulator.IsoSegmentFactory)
 	, bDisplay(Triangulator.bDisplay)
-	, LoopSegmentsIntersectionTool(Grid)
+	, LoopSegmentsIntersectionTool(Grid, Triangulator.Tolerances.GeometricTolerance)
 	, GetNext(LoopCleanerImpl::GetNextNodeImpl)
 	, GetPrevious(LoopCleanerImpl::GetPreviousNodeImpl)
 	, GetFirst(LoopCleanerImpl::GetFirstNode)
diff --git a/Engine/Source/Runtime/Datasmith/CADKernel/Private/CADKernel/Topo/Shell.cpp b/Engine/Source/Runtime/Datasmith/CADKernel/Private/CADKernel/Topo/Shell.cpp
--- a/Engine/Source/Runtime/Datasmith/CADKernel/Private/CADKernel/Topo/Shell.cpp
+++ b/Engine/Source/Runtime/Datasmith/CADKernel/Private/CADKernel/Topo/Shell.cpp
@@ -167,7 +167,7 @@ void FShell::UpdateShellOrientation()
 	const bool bIsOuter = IsOuter();
 	for (FOrientedFace& Face : TopologicalFaces)
 	{
-		if (bIsOuter != Face.Entity->IsBackOriented())
+		if (bIsOuter == Face.Entity->IsBackOriented())
 		{
 			Face.Direction = EOrientation::Back;
 		}
@@ -500,8 +500,6 @@ int32 FShell::Orient()
 		}
 	};
 
-	SpreadBodyOrientation();
-
 	for (FOrientedFace& Face : GetFaces())
 	{
 		Face.Entity->ResetMarkers();
diff --git a/Engine/Source/Runtime/Datasmith/CADKernel/Private/CADKernel/Topo/Topomaker.cpp b/Engine/Source/Runtime/Datasmith/CADKernel/Private/CADKernel/Topo/Topomaker.cpp
--- a/Engine/Source/Runtime/Datasmith/CADKernel/Private/CADKernel/Topo/Topomaker.cpp
+++ b/Engine/Source/Runtime/Datasmith/CADKernel/Private/CADKernel/Topo/Topomaker.cpp
@@ -49,10 +49,7 @@ namespace TopomakerTools
  */
 void MergeCoincidentVertices(TArray<FTopologicalVertex*>& VerticesToMerge, double Tolerance)
 {
-	double SquareTolerance = FMath::Square(Tolerance);
-
-	FTimePoint StartTime = FChrono::Now();
-
+	const double SquareTolerance = FMath::Square(Tolerance);
 	const double WeigthTolerance = 3 * Tolerance;
 
 	int32 VertexNum = (int32)VerticesToMerge.Num();
@@ -147,10 +144,6 @@ void MergeCoincidentVertices(TArray<FTopologicalVertex*>& VerticesToMerge, doubl
 	}
 
 	Swap(ActiveVertices, VerticesToMerge);
-
-	FDuration StepDuration = FChrono::Elapse(StartTime);
-	FChrono::PrintClockElapse(EVerboseLevel::Log, TEXT("    "), TEXT("Merge Coincident vertices"), StepDuration);
-
 }
 
 FTopologicalVertex* SplitAndLink(FTopologicalVertex& StartVertex, FTopologicalEdge& EdgeToLink, FTopologicalEdge& EdgeToSplit, double SquareSewTolerance, double SquareMinEdgeLength)
@@ -386,6 +379,20 @@ void MergeCoincidentEdges(FTopologicalEdge* Edge, double MinEdgeLength)
 {
 	const FTopologicalVertex& StartVertex = *Edge->GetStartVertex()->GetLinkActiveEntity();
 	const FTopologicalVertex& EndVertex = *Edge->GetEndVertex()->GetLinkActiveEntity();
+	if (&StartVertex == &EndVertex && Edge->Length() < MinEdgeLength)
+	{
+		if (Edge->GetTwinEntityCount() > 1)
+		{
+			FMessage::Printf(Debug, TEXT("Face %d Edge %d was self connected\n"), Edge->GetFace()->GetId(), Edge->GetId());
+			Edge->GetStartVertex()->UnlinkTo(*Edge->GetEndVertex());
+		}
+		else
+		{
+			Edge->SetAsDegenerated();
+			FMessage::Printf(Debug, TEXT("Face %d Edge %d is set as degenerated\n"), Edge->GetFace()->GetId(), Edge->GetId());
+		}
+		return;
+	}
 
 	TArray<FTopologicalEdge*> ConnectedEdges;
 	StartVertex.GetConnectedEdges(EndVertex, ConnectedEdges);
@@ -421,8 +428,6 @@ void MergeCoincidentEdges(FTopologicalEdge* Edge, double MinEdgeLength)
 
 void MergeCoincidentEdges(TArray<FTopologicalEdge*>& EdgesToProcess, double MinEdgeLength)
 {
-	FTimePoint StartTime = FChrono::Now();
-
 	for (FTopologicalEdge* Edge : EdgesToProcess)
 	{
 		if (Edge->IsDegenerated() || !Edge->IsBorder())
@@ -431,9 +436,6 @@ void MergeCoincidentEdges(TArray<FTopologicalEdge*>& EdgesToProcess, double MinE
 		}
 		MergeCoincidentEdges(Edge, MinEdgeLength);
 	}
-
-	FDuration StepDuration = FChrono::Elapse(StartTime);
-	FChrono::PrintClockElapse(EVerboseLevel::Log, TEXT("    "), TEXT("Merge coincident edges"), StepDuration);
 }
 
 void MergeCoincidentEdges(TArray<FTopologicalVertex*>& VerticesToProcess, double MinEdgeLength)
@@ -644,6 +646,8 @@ void FTopomaker::Sew()
 		RemoveDuplicatedFaces();
 	}
 
+	ResetMarkersOfFaces();
+
 #ifdef CADKERNEL_DEV
 	Report.SewDuration = FChrono::Elapse(StartJoinTime);
 	FChrono::PrintClockElapse(EVerboseLevel::Log, TEXT(""), TEXT("Sew"), Report.SewDuration);
diff --git a/Engine/Source/Runtime/Datasmith/CADKernel/Private/CADKernel/UI/Display.cpp b/Engine/Source/Runtime/Datasmith/CADKernel/Private/CADKernel/UI/Display.cpp
--- a/Engine/Source/Runtime/Datasmith/CADKernel/Private/CADKernel/UI/Display.cpp
+++ b/Engine/Source/Runtime/Datasmith/CADKernel/Private/CADKernel/UI/Display.cpp
@@ -514,6 +514,10 @@ void Display(const FShell& Shell)
 	FProgress Progress(Shell.GetFaces().Num(), TEXT("Display Shell"));
 	for (const FOrientedFace& Face : Shell.GetFaces())
 	{
+		if (Face.Entity->IsDeleted())
+		{
+			continue;
+		}
 		Display(*Face.Entity);
 	}
 #endif
@@ -525,6 +529,10 @@ void Draw(const FShell& Shell)
 	FProgress Progress(Shell.GetFaces().Num(), TEXT("Display Shell"));
 	for (const FOrientedFace& Face : Shell.GetFaces())
 	{
+		if (Face.Entity->IsDeleted())
+		{
+			continue;
+		}
 		F3DDebugSegment _(Face.Entity->GetId());
 		Draw(*Face.Entity);
 	}
@@ -539,7 +547,7 @@ void Display(const FBody& Body)
 	F3DDebugSegment GraphicSegment(Body.GetId());
 	for (const TSharedPtr<FShell>& Shell : Body.GetShells())
 	{
-		if (!Shell.IsValid())
+		if (!Shell.IsValid() || Shell->IsDeleted())
 		{
 			continue;
 		}
diff --git a/Engine/Source/Runtime/Datasmith/CADKernel/Public/CADKernel/Math/Boundary.h b/Engine/Source/Runtime/Datasmith/CADKernel/Public/CADKernel/Math/Boundary.h
--- a/Engine/Source/Runtime/Datasmith/CADKernel/Public/CADKernel/Math/Boundary.h
+++ b/Engine/Source/Runtime/Datasmith/CADKernel/Public/CADKernel/Math/Boundary.h
@@ -35,11 +35,23 @@ struct CADKERNEL_API FLinearBoundary
 	{
 	}
 
-	FLinearBoundary(double UMin, double UMax)
+	FLinearBoundary(const double UMin, const double UMax)
 	{
 		Set(UMin, UMax);
 	}
 
+	FLinearBoundary(const FLinearBoundary& Boundary, const double OffsetTolerance)
+		: Min(Boundary.Min - OffsetTolerance)
+		, Max(Boundary.Max + OffsetTolerance)
+	{
+	}
+
+	FLinearBoundary(const double UMin, const double UMax, const double OffsetTolerance)
+	{
+		Set(UMin, UMax);
+		Offset(OffsetTolerance);
+	}
+
 	friend FArchive& operator<<(FArchive& Ar, FLinearBoundary& Boundary)
 	{
 		Ar.Serialize(&Boundary, sizeof(FLinearBoundary));
@@ -56,7 +68,7 @@ struct CADKERNEL_API FLinearBoundary
 		return Max;
 	}
 
-	constexpr double GetAt(double Coordinate) const
+	constexpr double GetAt(const double Coordinate) const
 	{
 		return Min + (Max - Min) * Coordinate;
 	}
@@ -68,17 +80,17 @@ struct CADKERNEL_API FLinearBoundary
 
 	double Size() const { return Max - Min; }
 
-	void SetMin(double Coordinates)
+	void SetMin(const double Coordinates)
 	{
 		GetMinMax(Coordinates, Max, Min, Max);
 	}
 
-	void SetMax(double Coordinates)
+	void SetMax(const double Coordinates)
 	{
 		GetMinMax(Min, Coordinates, Min, Max);
 	}
 
-	void Set(double InUMin = 0., double InUMax = 1.)
+	void Set(const double InUMin = 0., const double InUMax = 1.)
 	{
 		GetMinMax(InUMin, InUMax, Min, Max);
 	}
@@ -101,7 +113,7 @@ struct CADKERNEL_API FLinearBoundary
 		return Min <= Max;
 	}
 
-	bool Contains(double Coordinate) const
+	bool Contains(const double Coordinate) const
 	{
 		return RealCompare(Coordinate, Min) >= 0 && RealCompare(Coordinate, Max) <= 0;
 	}
@@ -134,7 +146,7 @@ struct CADKERNEL_API FLinearBoundary
 	/**
 	 * If a coordinate is outside the bounds, set the coordinate at the closed limit
 	 */
-	void MoveInsideIfNot(double& Coordinate, double Tolerance = DOUBLE_SMALL_NUMBER) const
+	void MoveInsideIfNot(double& Coordinate, const double Tolerance = DOUBLE_SMALL_NUMBER) const
 	{
 		if (Coordinate <= Min)
 		{
@@ -146,6 +158,12 @@ struct CADKERNEL_API FLinearBoundary
 		}
 	}
 
+	void Offset(const double Tolerance = DOUBLE_SMALL_NUMBER)
+	{
+		Min -= Tolerance;
+		Max += Tolerance;
+	}
+
 	/**
 	 * Uses to initiate a boundary computation with ExtendTo
 	 */
@@ -180,7 +198,7 @@ struct CADKERNEL_API FLinearBoundary
 		Max = FMath::Max(Max, MaxBound.Max);
 	}
 
-	void ExtendTo(double Coordinate)
+	void ExtendTo(const double Coordinate)
 	{
 		if (Coordinate < Min)
 		{
@@ -237,10 +255,17 @@ public:
 
 	FSurfacicBoundary() = default;
 
-	FSurfacicBoundary(double InUMin, double InUMax, double InVMin, double InVMax)
+	FSurfacicBoundary(const double InUMin, const double InUMax, const double InVMin, const double InVMax)
+	{
+		UVBoundaries[EIso::IsoU].Set(InUMin, InUMax);
+		UVBoundaries[EIso::IsoV].Set(InVMin, InVMax);
+	}
+
+	FSurfacicBoundary(const double InUMin, const double InUMax, const double InVMin, const double InVMax, const double OffsetTolerance)
 	{
 		UVBoundaries[EIso::IsoU].Set(InUMin, InUMax);
 		UVBoundaries[EIso::IsoV].Set(InVMin, InVMax);
+		Offset(OffsetTolerance);
 	}
 
 	FSurfacicBoundary(const FPoint2D& Point1, const FPoint2D& Point2)
@@ -248,6 +273,12 @@ public:
 		Set(Point1, Point2);
 	}
 
+	FSurfacicBoundary(const FPoint2D& Point1, const FPoint2D& Point2, const double OffsetTolerance)
+	{
+		Set(Point1, Point2);
+		Offset(OffsetTolerance);
+	}
+
 	void Set(const FPoint2D& Point1, const FPoint2D& Point2)
 	{
 		UVBoundaries[EIso::IsoU].Set(Point1.U, Point2.U);
@@ -267,7 +298,7 @@ public:
 		UVBoundaries[EIso::IsoV] = BoundV;
 	}
 
-	void Set(double InUMin, double InUMax, double InVMin, double InVMax)
+	void Set(const double InUMin, const double InUMax, const double InVMin, const double InVMax)
 	{
 		UVBoundaries[EIso::IsoU].Set(InUMin, InUMax);
 		UVBoundaries[EIso::IsoV].Set(InVMin, InVMax);
@@ -350,25 +381,25 @@ public:
 		UVBoundaries[EIso::IsoV].Init();
 	}
 
-	void TrimAt(const FSurfacicBoundary MaxLimit)
+	void TrimAt(const FSurfacicBoundary& MaxLimit)
 	{
 		UVBoundaries[EIso::IsoU].TrimAt(MaxLimit[EIso::IsoU]);
 		UVBoundaries[EIso::IsoV].TrimAt(MaxLimit[EIso::IsoV]);
 	}
 
-	void ExtendTo(const FSurfacicBoundary MaxLimit)
+	void ExtendTo(const FSurfacicBoundary& MaxLimit)
 	{
 		UVBoundaries[EIso::IsoU].ExtendTo(MaxLimit[EIso::IsoU]);
 		UVBoundaries[EIso::IsoV].ExtendTo(MaxLimit[EIso::IsoV]);
 	}
 
-	void ExtendTo(FPoint2D Point)
+	void ExtendTo(const FPoint2D& Point)
 	{
 		UVBoundaries[EIso::IsoU].ExtendTo(Point.U);
 		UVBoundaries[EIso::IsoV].ExtendTo(Point.V);
 	}
 
-	void ExtendTo(FPoint Point)
+	void ExtendTo(const FPoint& Point)
 	{
 		UVBoundaries[EIso::IsoU].ExtendTo(Point.X);
 		UVBoundaries[EIso::IsoV].ExtendTo(Point.Y);
@@ -392,7 +423,7 @@ public:
 	/**
 	 * If a point is outside the bounds, set the coordinate to insert the point inside the bounds
 	 */
-	void MoveInsideIfNot(FPoint& Point, double Tolerance = DOUBLE_SMALL_NUMBER) const
+	void MoveInsideIfNot(FPoint& Point, const double Tolerance = DOUBLE_SMALL_NUMBER) const
 	{
 		UVBoundaries[EIso::IsoU].MoveInsideIfNot(Point.X, Tolerance);
 		UVBoundaries[EIso::IsoV].MoveInsideIfNot(Point.Y, Tolerance);
@@ -401,7 +432,7 @@ public:
 	/**
 	 * If a point is outside the bounds, set the coordinate to insert the point inside the bounds
 	 */
-	void MoveInsideIfNot(FPoint2D& Point, double Tolerance = DOUBLE_SMALL_NUMBER) const
+	void MoveInsideIfNot(FPoint2D& Point, const double Tolerance = DOUBLE_SMALL_NUMBER) const
 	{
 		UVBoundaries[EIso::IsoU].MoveInsideIfNot(Point.U, Tolerance);
 		UVBoundaries[EIso::IsoV].MoveInsideIfNot(Point.V, Tolerance);
@@ -422,6 +453,12 @@ public:
 		return UVBoundaries[Iso];
 	}
 
+	void Offset(const double Tolerance = DOUBLE_SMALL_NUMBER)
+	{
+		UVBoundaries[EIso::IsoU].Offset(Tolerance);
+		UVBoundaries[EIso::IsoV].Offset(Tolerance);
+	}
+
 };
 }
 
