[NO NON-PII COMMIT MESSAGE]

diff --git a/Engine/Plugins/Enterprise/DatasmithCADImporter/Source/CADInterfaces/Private/TechSoftBridge.cpp b/Engine/Plugins/Enterprise/DatasmithCADImporter/Source/CADInterfaces/Private/TechSoftBridge.cpp
--- a/Engine/Plugins/Enterprise/DatasmithCADImporter/Source/CADInterfaces/Private/TechSoftBridge.cpp
+++ b/Engine/Plugins/Enterprise/DatasmithCADImporter/Source/CADInterfaces/Private/TechSoftBridge.cpp
@@ -239,7 +239,7 @@ UE::CADKernel::FBody* FTechSoftBridge::AddBody(A3DRiBrepModel* A3DBRepModel, FAr
 	}
 
 	TSharedRef<UE::CADKernel::FBody> Body = UE::CADKernel::FEntity::MakeShared<UE::CADKernel::FBody>();
-	AddMetadata(BRepMetaData, *Body);
+	AddMetaData(BRepMetaData, *Body);
 
 	Body->SetDisplayData(ArchiveBody.ColorUId, ArchiveBody.MaterialUId);
 
@@ -330,7 +330,7 @@ void FTechSoftBridge::TraverseShell(const A3DTopoShell* A3DShell, TSharedRef<UE:
 #endif
 
 	Shell->SetDisplayData(*Body);
-	AddMetadata(MetaData, *Shell);
+	AddMetaData(MetaData, *Shell);
 
 	TUniqueTSObj<A3DTopoShellData> ShellData(A3DShell);
 
@@ -936,8 +936,8 @@ void FTechSoftBridge::AddFace(const A3DTopoFace* A3DFace, UE::CADKernel::EOrient
 		return;
 	}
 
-	AddMetadata(MetaData, *Face);
-	Face->SetDisplayDataIfUndefined(*Shell);
+	AddMetaData(MetaData, *Face);
+	Face->CompleteMetaData();
 
 	Face->SetHostId(ShellIndex);
 	Shell->Add(Face, Orientation);
@@ -1415,7 +1415,7 @@ TSharedPtr<UE::CADKernel::FSurface> FTechSoftBridge::AddSurfaceAsNurbs(const A3D
 
 }
 
-void FTechSoftBridge::AddMetadata(FArchiveCADObject& MetaData, UE::CADKernel::FTopologicalShapeEntity& Entity)
+void FTechSoftBridge::AddMetaData(FArchiveCADObject& MetaData, UE::CADKernel::FTopologicalShapeEntity& Entity)
 {
 	FString* Name = MetaData.MetaData.Find(TEXT("Name"));
 	if (Name != nullptr)
diff --git a/Engine/Plugins/Enterprise/DatasmithCADImporter/Source/CADInterfaces/Private/TechSoftBridge.h b/Engine/Plugins/Enterprise/DatasmithCADImporter/Source/CADInterfaces/Private/TechSoftBridge.h
--- a/Engine/Plugins/Enterprise/DatasmithCADImporter/Source/CADInterfaces/Private/TechSoftBridge.h
+++ b/Engine/Plugins/Enterprise/DatasmithCADImporter/Source/CADInterfaces/Private/TechSoftBridge.h
@@ -250,7 +250,7 @@ private:
 
 	TSharedPtr<UE::CADKernel::FCurve> AddCurveAsNurbs(const A3DCrvBase* A3DCurve, const TechSoftUtils::FUVReparameterization& UVReparameterization);
 
-	void AddMetadata(FArchiveCADObject& EntityData, UE::CADKernel::FTopologicalShapeEntity& Entity);
+	void AddMetaData(FArchiveCADObject& EntityData, UE::CADKernel::FTopologicalShapeEntity& Entity);
 };
 }
 
diff --git a/Engine/Plugins/Enterprise/DatasmithCADImporter/Source/DatasmithWireTranslator/Private/AliasModelToCADKernelConverter.cpp b/Engine/Plugins/Enterprise/DatasmithCADImporter/Source/DatasmithWireTranslator/Private/AliasModelToCADKernelConverter.cpp
--- a/Engine/Plugins/Enterprise/DatasmithCADImporter/Source/DatasmithWireTranslator/Private/AliasModelToCADKernelConverter.cpp
+++ b/Engine/Plugins/Enterprise/DatasmithCADImporter/Source/DatasmithWireTranslator/Private/AliasModelToCADKernelConverter.cpp
@@ -290,14 +290,15 @@ bool FAliasModelToCADKernelConverter::AddBRep(AlDagNode& DagNode, const FColor&
 	AlEdge2CADKernelEdge.Empty();
 
 	TSharedRef<UE::CADKernel::FBody> CADKernelBody = UE::CADKernel::FEntity::MakeShared<UE::CADKernel::FBody>();
-	TSharedRef<UE::CADKernel::FShell> CADKernelShell = UE::CADKernel::FEntity::MakeShared<UE::CADKernel::FShell>();
 	uint32 ColorId = (uint32) CADLibrary::BuildColorUId(Color);
-	CADKernelShell->SetColorId(ColorId);
+	CADKernelBody->SetColorId(ColorId);
+
+	TSharedRef<UE::CADKernel::FShell> CADKernelShell = UE::CADKernel::FEntity::MakeShared<UE::CADKernel::FShell>();
 	CADKernelBody->AddShell(CADKernelShell);
 
 	boolean bAlOrientation;
 	DagNode.getSurfaceOrientation(bAlOrientation);
-	bool bOrientation = (bool)bAlOrientation;
+	bool bOrientation = !(bool)bAlOrientation;
 
 	AlMatrix4x4 AlMatrix;
 	if (InObjectReference == EAliasObjectReference::ParentReference)
@@ -340,8 +341,7 @@ bool FAliasModelToCADKernelConverter::AddBRep(AlDagNode& DagNode, const FColor&
 		break;
 	}
 
-
-	// Create body from faces
+	CADKernelBody->CompleteMetaData();
 	CADKernelSession.GetModel().Add(CADKernelBody);
 
 	return CADKernelShell->FaceCount() > 0;
diff --git a/Engine/Plugins/Enterprise/DatasmithCADImporter/Source/DatasmithWireTranslator/Private/AliasModelToTechSoftConverter.cpp b/Engine/Plugins/Enterprise/DatasmithCADImporter/Source/DatasmithWireTranslator/Private/AliasModelToTechSoftConverter.cpp
--- a/Engine/Plugins/Enterprise/DatasmithCADImporter/Source/DatasmithWireTranslator/Private/AliasModelToTechSoftConverter.cpp
+++ b/Engine/Plugins/Enterprise/DatasmithCADImporter/Source/DatasmithWireTranslator/Private/AliasModelToTechSoftConverter.cpp
@@ -355,6 +355,10 @@ bool FAliasModelToTechSoftConverter::AddBRep(AlDagNode& DagNode, const FColor& C
 		DagNode.localTransformationMatrix(AlMatrix);
 	}
 
+	boolean bAlOrientation;
+	DagNode.getSurfaceOrientation(bAlOrientation);
+	bool bOrientation = !(bool)bAlOrientation;
+
 	TArray<A3DTopoFace*> TSFaces;
 	TSFaces.Reserve(100);
 
@@ -430,7 +434,7 @@ bool FAliasModelToTechSoftConverter::AddBRep(AlDagNode& DagNode, const FColor& C
 	A3DTopoShell* TopoShellPtr = nullptr;
 	{
 		TArray<A3DUns8> FaceOrientations;
-		FaceOrientations.Init(true, TSFaces.Num());
+		FaceOrientations.Init(bOrientation, TSFaces.Num());
 
 		CADLibrary::TUniqueTSObj<A3DTopoShellData> TopoShellData;
 		TopoShellData->m_bClosed = false;
diff --git a/Engine/Source/Runtime/Datasmith/CADKernel/Private/CADKernel/Mesh/Meshers/ParametricMesher.cpp b/Engine/Source/Runtime/Datasmith/CADKernel/Private/CADKernel/Mesh/Meshers/ParametricMesher.cpp
--- a/Engine/Source/Runtime/Datasmith/CADKernel/Private/CADKernel/Mesh/Meshers/ParametricMesher.cpp
+++ b/Engine/Source/Runtime/Datasmith/CADKernel/Private/CADKernel/Mesh/Meshers/ParametricMesher.cpp
@@ -87,7 +87,7 @@ void FParametricMesher::MeshEntities(TArray<FTopologicalShapeEntity*>& InEntitie
 			continue;
 		}
 
-		TopologicalEntity->SpreadBodyOrientation();
+		TopologicalEntity->PropagateBodyOrientation();
 		TopologicalEntity->GetFaces(Faces);
 	}
 
diff --git a/Engine/Source/Runtime/Datasmith/CADKernel/Private/CADKernel/Topo/Model.cpp b/Engine/Source/Runtime/Datasmith/CADKernel/Private/CADKernel/Topo/Model.cpp
--- a/Engine/Source/Runtime/Datasmith/CADKernel/Private/CADKernel/Topo/Model.cpp
+++ b/Engine/Source/Runtime/Datasmith/CADKernel/Private/CADKernel/Topo/Model.cpp
@@ -91,11 +91,19 @@ void FModel::GetFaces(TArray<FTopologicalFace*>& OutFaces)
 	}
 }
 
-void FModel::SpreadBodyOrientation()
+void FModel::PropagateBodyOrientation()
 {
 	for (TSharedPtr<FBody>& Body : Bodies)
 	{
-		Body->SpreadBodyOrientation();
+		Body->PropagateBodyOrientation();
+	}
+}
+
+void FModel::CompleteMetaData()
+{
+	for (TSharedPtr<FBody>& Body : Bodies)
+	{
+		Body->CompleteMetaData();
 	}
 }
 
diff --git a/Engine/Source/Runtime/Datasmith/CADKernel/Private/CADKernel/Topo/Shell.cpp b/Engine/Source/Runtime/Datasmith/CADKernel/Private/CADKernel/Topo/Shell.cpp
--- a/Engine/Source/Runtime/Datasmith/CADKernel/Private/CADKernel/Topo/Shell.cpp
+++ b/Engine/Source/Runtime/Datasmith/CADKernel/Private/CADKernel/Topo/Shell.cpp
@@ -146,7 +146,7 @@ void FShell::Merge(TSharedPtr<FShell>& Shell)
 }
 
 
-void FShell::SpreadBodyOrientation()
+void FShell::PropagateBodyOrientation()
 {
 	const bool bIsOuter = IsOuter();
 	for (FOrientedFace& Face : TopologicalFaces)
@@ -162,6 +162,16 @@ void FShell::SpreadBodyOrientation()
 	}
 }
 
+void FShell::CompleteMetaData()
+{
+	CompleteMetaDataWithHostMetaData();
+	for (FOrientedFace& Face : TopologicalFaces)
+	{
+		Face.Entity->CompleteMetaData();
+	}
+}
+
+
 void FShell::UpdateShellOrientation()
 {
 	const bool bIsOuter = IsOuter();
diff --git a/Engine/Source/Runtime/Datasmith/CADKernel/Private/CADKernel/Topo/TopologicalShapeEntity.cpp b/Engine/Source/Runtime/Datasmith/CADKernel/Private/CADKernel/Topo/TopologicalShapeEntity.cpp
--- a/Engine/Source/Runtime/Datasmith/CADKernel/Private/CADKernel/Topo/TopologicalShapeEntity.cpp
+++ b/Engine/Source/Runtime/Datasmith/CADKernel/Private/CADKernel/Topo/TopologicalShapeEntity.cpp
@@ -4,11 +4,11 @@
 namespace UE::CADKernel
 {
 
-void FTopologicalShapeEntity::CompleteMetadata()
+void FTopologicalShapeEntity::CompleteMetaDataWithHostMetaData()
 {
 	if (HostedBy != nullptr)
 	{
-		Dictionary.CompleteDictionary(HostedBy->GetMetadataDictionary());
+		Dictionary.CompleteDictionary(HostedBy->GetMetaDataDictionary());
 	}
 }
 
diff --git a/Engine/Source/Runtime/Datasmith/CADKernel/Private/CADKernel/Topo/Topomaker.cpp b/Engine/Source/Runtime/Datasmith/CADKernel/Private/CADKernel/Topo/Topomaker.cpp
--- a/Engine/Source/Runtime/Datasmith/CADKernel/Private/CADKernel/Topo/Topomaker.cpp
+++ b/Engine/Source/Runtime/Datasmith/CADKernel/Private/CADKernel/Topo/Topomaker.cpp
@@ -535,8 +535,10 @@ FTopomaker::FTopomaker(FSession& InSession, const FTopomakerOptions& InOptions)
 
 	for (const TSharedPtr<FBody>& Body : Session.GetModel().GetBodies())
 	{
+		Body->CompleteMetaData();
 		for (const TSharedPtr<FShell>& Shell : Body->GetShells())
 		{
+			Shell->PropagateBodyOrientation();
 			Shells.Add(Shell.Get());
 		}
 	}
@@ -552,6 +554,7 @@ FTopomaker::FTopomaker(FSession& InSession, const TArray<TSharedPtr<FShell>>& In
 	Shells.Reserve(InShells.Num());
 	for (const TSharedPtr<FShell>& Shell : InShells)
 	{
+		Shell->PropagateBodyOrientation();
 		Shells.Add(Shell.Get());
 	}
 	InitFaces();
@@ -562,27 +565,21 @@ void FTopomaker::InitFaces()
 	int32 FaceCount = 0;
 	for (FShell* Shell : Shells)
 	{
-		Shell->CompleteMetadata();
 		FaceCount += Shell->FaceCount();
 	}
 	Faces.Reserve(FaceCount);
 
 	for (FShell* Shell : Shells)
 	{
-		Shell->SpreadBodyOrientation();
 		for (const FOrientedFace& Face : Shell->GetFaces())
 		{
 			Faces.Add(Face.Entity);
 		}
-		for (const FOrientedFace& Face : Shell->GetFaces())
-		{
-			Face.Entity->CompleteMetadata();
-		}
 	}
 
 	for (const TSharedPtr<FTopologicalFace>& Face : Faces)
 	{
-		Face->ResetMarker2();
+		Face->ResetMarkers();
 	}
 }
 
diff --git a/Engine/Source/Runtime/Datasmith/CADKernel/Public/CADKernel/Topo/Body.h b/Engine/Source/Runtime/Datasmith/CADKernel/Public/CADKernel/Topo/Body.h
--- a/Engine/Source/Runtime/Datasmith/CADKernel/Public/CADKernel/Topo/Body.h
+++ b/Engine/Source/Runtime/Datasmith/CADKernel/Public/CADKernel/Topo/Body.h
@@ -125,11 +125,20 @@ public:
 		}
 	}
 
-	virtual void SpreadBodyOrientation() override
+	virtual void CompleteMetaData() override
+	{
+		CompleteMetaDataWithHostMetaData();
+		for (TSharedPtr<FShell>& Shell : Shells)
+		{
+			Shell->CompleteMetaData();
+		}
+	}
+
+	virtual void PropagateBodyOrientation() override
 	{
 		for (const TSharedPtr<FShell>& Shell : Shells)
 		{
-			Shell->SpreadBodyOrientation();
+			Shell->PropagateBodyOrientation();
 		}
 	}
 
diff --git a/Engine/Source/Runtime/Datasmith/CADKernel/Public/CADKernel/Topo/Model.h b/Engine/Source/Runtime/Datasmith/CADKernel/Public/CADKernel/Topo/Model.h
--- a/Engine/Source/Runtime/Datasmith/CADKernel/Public/CADKernel/Topo/Model.h
+++ b/Engine/Source/Runtime/Datasmith/CADKernel/Public/CADKernel/Topo/Model.h
@@ -166,7 +166,9 @@ public:
 		return Bodies;
 	}
 
-	virtual void SpreadBodyOrientation() override;
+	virtual void PropagateBodyOrientation() override;
+
+	virtual void CompleteMetaData() override;
 
 	// Topo functions
 
diff --git a/Engine/Source/Runtime/Datasmith/CADKernel/Public/CADKernel/Topo/Shell.h b/Engine/Source/Runtime/Datasmith/CADKernel/Public/CADKernel/Topo/Shell.h
--- a/Engine/Source/Runtime/Datasmith/CADKernel/Public/CADKernel/Topo/Shell.h
+++ b/Engine/Source/Runtime/Datasmith/CADKernel/Public/CADKernel/Topo/Shell.h
@@ -136,7 +136,9 @@ public:
 
 	virtual void Merge(TSharedPtr<FShell>& Shell);
 
-	virtual void SpreadBodyOrientation() override;
+	virtual void PropagateBodyOrientation() override;
+
+	virtual void CompleteMetaData() override;
 
 	/**
 	 * Update each FOrientedFace::Direction according to FOrientedFace::Entity->IsBackOriented() flag
diff --git a/Engine/Source/Runtime/Datasmith/CADKernel/Public/CADKernel/Topo/TopologicalFace.h b/Engine/Source/Runtime/Datasmith/CADKernel/Public/CADKernel/Topo/TopologicalFace.h
--- a/Engine/Source/Runtime/Datasmith/CADKernel/Public/CADKernel/Topo/TopologicalFace.h
+++ b/Engine/Source/Runtime/Datasmith/CADKernel/Public/CADKernel/Topo/TopologicalFace.h
@@ -171,7 +171,12 @@ public:
 		}
 	}
 
-	virtual void SpreadBodyOrientation() override
+	virtual void CompleteMetaData() override
+	{
+		CompleteMetaDataWithHostMetaData();
+	}
+
+	virtual void PropagateBodyOrientation() override
 	{
 	}
 
diff --git a/Engine/Source/Runtime/Datasmith/CADKernel/Public/CADKernel/Topo/TopologicalShapeEntity.h b/Engine/Source/Runtime/Datasmith/CADKernel/Public/CADKernel/Topo/TopologicalShapeEntity.h
--- a/Engine/Source/Runtime/Datasmith/CADKernel/Public/CADKernel/Topo/TopologicalShapeEntity.h
+++ b/Engine/Source/Runtime/Datasmith/CADKernel/Public/CADKernel/Topo/TopologicalShapeEntity.h
@@ -40,7 +40,7 @@ public:
 		FTopologicalEntity::Empty();
 	}
 
-	const FMetadataDictionary& GetMetadataDictionary() const
+	const FMetadataDictionary& GetMetaDataDictionary() const
 	{
 		return Dictionary;
 	}
@@ -50,7 +50,8 @@ public:
 		Dictionary.ExtractMetaData(OutMetaData);
 	}
 
-	void CompleteMetadata();
+	virtual void CompleteMetaData() = 0;
+	void CompleteMetaDataWithHostMetaData();
 
 	virtual int32 FaceCount() const = 0;
 	virtual void GetFaces(TArray<FTopologicalFace*>& OutFaces) = 0;
@@ -59,7 +60,7 @@ public:
 	 * Each face of model is set by its orientation. This allow to make oriented mesh and to keep the face orientation in topological function.
 	 * Marker2 of spread face is set. It must be reset after the process
 	 */
-	virtual void SpreadBodyOrientation() = 0;
+	virtual void PropagateBodyOrientation() = 0;
 
 #ifdef CADKERNEL_DEV
 	virtual void FillTopologyReport(FTopologyReport& Report) const = 0;
@@ -153,18 +154,6 @@ public:
 		SetDisplayData(DisplayData.GetColorId(), DisplayData.GetMaterialId());
 	}
 
-	void SetDisplayDataIfUndefined(const FTopologicalShapeEntity& DefaultDisplayData)
-	{
-		if (Dictionary.GetColorId() == 0 && DefaultDisplayData.GetColorId())
-		{
-			Dictionary.SetColorId(DefaultDisplayData.GetColorId());
-		}
-		if (Dictionary.GetMaterialId() == 0 && DefaultDisplayData.GetMaterialId())
-		{
-			Dictionary.SetMaterialId(DefaultDisplayData.GetMaterialId());
-		}
-	}
-
 	void SetPatchId(int32 InPatchId)
 	{
 		Dictionary.SetPatchId(InPatchId);
