[NO NON-PII COMMIT MESSAGE]

diff --git a/Engine/Plugins/Enterprise/DatasmithCADImporter/Source/CADInterfaces/Private/TechSoftFileParser.cpp b/Engine/Plugins/Enterprise/DatasmithCADImporter/Source/CADInterfaces/Private/TechSoftFileParser.cpp
--- a/Engine/Plugins/Enterprise/DatasmithCADImporter/Source/CADInterfaces/Private/TechSoftFileParser.cpp
+++ b/Engine/Plugins/Enterprise/DatasmithCADImporter/Source/CADInterfaces/Private/TechSoftFileParser.cpp
@@ -290,7 +290,7 @@ void UpdateIOOptionAccordingToFormat(const CADLibrary::ECADFormat Format, A3DImp
 double ExtractUniformScale(FVector3d& Scale)
 {
 	double UniformScale = (Scale.X + Scale.Y + Scale.Z) / 3.;
-	double Tolerance = UniformScale * KINDA_SMALL_NUMBER;
+	double Tolerance = UniformScale * DOUBLE_KINDA_SMALL_NUMBER;
 
 	if (!FMath::IsNearlyEqual(UniformScale, Scale.X, Tolerance) && !FMath::IsNearlyEqual(UniformScale, Scale.Y, Tolerance))
 	{
@@ -1884,7 +1884,7 @@ void FTechSoftFileParser::ExtractGeneralTransformation(const A3DMiscTransformati
 
 		FTransform3d Transform(Matrix);
 		FVector3d Scale = Transform.GetScale3D();
-		if (Scale.Equals(FVector3d::OneVector, KINDA_SMALL_NUMBER))
+		if (Scale.Equals(FVector3d::OneVector, DOUBLE_KINDA_SMALL_NUMBER))
 		{
 			const double TranslationScale = Component.Unit * FImportParameters::GUnitScale;
 			for (Index = 0; Index < 3; ++Index)
diff --git a/Engine/Source/Runtime/Datasmith/CADKernel/Private/CADKernel/Core/Session.cpp b/Engine/Source/Runtime/Datasmith/CADKernel/Private/CADKernel/Core/Session.cpp
--- a/Engine/Source/Runtime/Datasmith/CADKernel/Private/CADKernel/Core/Session.cpp
+++ b/Engine/Source/Runtime/Datasmith/CADKernel/Private/CADKernel/Core/Session.cpp
@@ -96,5 +96,6 @@ void FSession::SetGeometricTolerance(double NewTolerance)
 {
 	ensure(Database.GetModel().EntityCount() == 0);
 	GeometricTolerance = NewTolerance;
+	IntersectionTool::SetTolerance(NewTolerance);
 }
 } // namespace UE::CADKernel
diff --git a/Engine/Source/Runtime/Datasmith/CADKernel/Private/CADKernel/Math/Geometry.cpp b/Engine/Source/Runtime/Datasmith/CADKernel/Private/CADKernel/Math/Geometry.cpp
--- a/Engine/Source/Runtime/Datasmith/CADKernel/Private/CADKernel/Math/Geometry.cpp
+++ b/Engine/Source/Runtime/Datasmith/CADKernel/Private/CADKernel/Math/Geometry.cpp
@@ -10,6 +10,85 @@
 
 namespace UE::CADKernel
 {
+namespace IntersectionTool
+{
+static double IntersectionToolTolerance = 0.01;
+
+void SetTolerance(const double Tolerance)
+{
+	IntersectionToolTolerance = Tolerance;
+}
+
+struct FIntersectionContext
+{
+	const FSegment2D& SegmentAB;
+	const FSegment2D& SegmentCD;
+	const FPoint2D AB;
+	const FPoint2D CD;
+	const FPoint2D CA;
+
+	double NormAB = 0;
+	double NormCD = 0;
+
+	FIntersectionContext(const FSegment2D& InSegmentAB, const FSegment2D& InSegmentCD)
+		: SegmentAB(InSegmentAB)
+		, SegmentCD(InSegmentCD)
+		, AB(SegmentAB.GetVector())
+		, CD(SegmentCD.GetVector())
+		, CA(SegmentAB[0] - SegmentCD[0])
+	{
+	}
+};
+
+bool DoCoincidentSegmentsIntersectInside(double A, double B, double C, double D)
+{
+	return !((D < A + DOUBLE_KINDA_SMALL_NUMBER) || (B < C + DOUBLE_KINDA_SMALL_NUMBER));
+}
+
+bool DoCoincidentSegmentsIntersect(double A, double B, double C, double D)
+{
+	return !((D < A) || (B < C));
+}
+
+constexpr double MinValue(bool OnlyInside)
+{
+	return OnlyInside ? DOUBLE_KINDA_SMALL_NUMBER : -DOUBLE_KINDA_SMALL_NUMBER;
+}
+
+constexpr double MaxValue(bool OnlyInside)
+{
+	return OnlyInside ? 1. - DOUBLE_KINDA_SMALL_NUMBER : 1. + DOUBLE_KINDA_SMALL_NUMBER;
+}
+
+bool ConfirmIntersectionWhenNearlyCoincident(const FPoint2D& AB, const FPoint2D& AC, const FPoint2D& AD, const double NormAB)
+{
+	double HeightC = (AB ^ AC) / NormAB;
+	double HeightD = (AB ^ AD) / NormAB;
+
+	if (fabs(HeightC) < IntersectionToolTolerance || fabs(HeightD) < IntersectionToolTolerance)
+	{
+		return true;
+	}
+	return HeightC * HeightD < 0;
+};
+
+bool ConfirmIntersectionWhenNearlyCoincident(const FIntersectionContext& Context)
+{
+	if (Context.NormAB > Context.NormCD)
+	{
+		const FPoint2D DA = Context.SegmentAB[0] - Context.SegmentCD[1];
+		return ConfirmIntersectionWhenNearlyCoincident(Context.AB, Context.CA, DA, Context.NormAB);
+	}
+	else
+	{
+		const FPoint2D CB = Context.SegmentAB[1] - Context.SegmentCD[0];
+		return ConfirmIntersectionWhenNearlyCoincident(Context.CD, Context.CA, CB, Context.NormCD);
+	}
+}
+
+}
+
+
 double ComputeCurvature(const FPoint& Gradient, const FPoint& Laplacian)
 {
 	const FPoint GradientCopy = Gradient.Normalize();
@@ -65,132 +144,99 @@ void FindLoopIntersectionsWithIso(const EIso Iso, const double IsoParameter, con
 	Algo::Sort(OutIntersections);
 }
 
-bool DoIntersectInside(const FSegment2D& SegmentAB, const FSegment2D& SegmentCD)
+bool DoIntersect(const FSegment2D& SegmentAB, const FSegment2D& SegmentCD, TFunction<bool(double, double, double, double)> DoCoincidentSegmentsIntersect, const double Min, const double Max)
 {
-	constexpr const double Min = DOUBLE_KINDA_SMALL_NUMBER;
-	constexpr const double Max = 1. - DOUBLE_KINDA_SMALL_NUMBER;
+	using namespace IntersectionTool;
 
-	TFunction<bool(double, double, double, double)> Intersect = [](double A, double B, double C, double D)
-	{
-		return !((D < A + DOUBLE_KINDA_SMALL_NUMBER) || (B < C + DOUBLE_KINDA_SMALL_NUMBER));
-	};
-
-	TFunction<bool(double, double, double, double)> TestWhenParallel = [&](double A, double B, double C, double D)
+	TFunction<bool(double, double, double, double)> FastIntersectionTestWhenCoincident = [&DoCoincidentSegmentsIntersect](double A, double B, double C, double D) -> bool
 	{
 		if (A < B)
 		{
 			if (C < D)
 			{
-				return Intersect(A, B, C, D);
+				return DoCoincidentSegmentsIntersect(A, B, C, D);
 			}
 			else
 			{
-				return Intersect(A, B, D, C);
+				return DoCoincidentSegmentsIntersect(A, B, D, C);
 			}
 		}
 		else
 		{
 			if (C < D)
 			{
-				return Intersect(B, A, C, D);
+				return DoCoincidentSegmentsIntersect(B, A, C, D);
 			}
 			else
 			{
-				return Intersect(B, A, D, C);
+				return DoCoincidentSegmentsIntersect(B, A, D, C);
 			}
 		}
 	};
 
-	const FPoint2D AB = SegmentAB[1] - SegmentAB[0];
-	const FPoint2D CD = SegmentCD[1] - SegmentCD[0];
-	const FPoint2D CA = SegmentAB[0] - SegmentCD[0];
+	FIntersectionContext Context(SegmentAB, SegmentCD);
 
-	const double ParallelCoef = CD ^ AB;
+	const double ParallelCoef = Context.CD ^ Context.AB;
 	if (FMath::IsNearlyZero(ParallelCoef, DOUBLE_KINDA_SMALL_NUMBER))
 	{
-		const double ParallelCoef2 = CA ^ AB;
-		if (!FMath::IsNearlyZero(ParallelCoef2, DOUBLE_KINDA_SMALL_NUMBER))
-		{
-			return false;
-		}
-
-		if (fabs(AB.U) > fabs(AB.V))
+		// double check with normalized vectors
 		{
-			return TestWhenParallel(SegmentAB[0].U, SegmentAB[1].U, SegmentCD[0].U, SegmentCD[1].U);
-		}
-		else
-		{
-			return TestWhenParallel(SegmentAB[0].V, SegmentAB[1].V, SegmentCD[0].V, SegmentCD[1].V);
-		}
-	}
-
-	const double ABIntersectionCoordinate = (CA ^ CD) / ParallelCoef;
-	const double CDIntersectionCoordinate = (CA ^ AB) / ParallelCoef;
-	return (ABIntersectionCoordinate < Max && ABIntersectionCoordinate > Min && CDIntersectionCoordinate < Max && CDIntersectionCoordinate > Min);
-}
+			FPoint2D NormalizedAB = Context.AB;
+			FPoint2D NormalizedCD = Context.CD;
+			FPoint2D NormalizedCA = Context.CA;
 
-bool DoIntersect(const TSegment<FPoint2D>& SegmentAB, const TSegment<FPoint2D>& SegmentCD)
-{
-	constexpr const double Min = -DOUBLE_KINDA_SMALL_NUMBER;
-	constexpr const double Max = 1. + DOUBLE_KINDA_SMALL_NUMBER;
+			NormalizedAB.Normalize(Context.NormAB);
+			NormalizedCD.Normalize(Context.NormCD);
+			NormalizedCA.Normalize();
 
-	TFunction<bool(double, double, double, double)> Intersect = [](double A, double B, double C, double D)
-	{
-		return !((D < A) || (B < C));
-	};
-
-	TFunction<bool(double, double, double, double)> TestWhenParallel = [&](double A, double B, double C, double D)
-	{
-		if (A < B)
-		{
-			if (C < D)
-			{
-				return Intersect(A, B, C, D);
-			}
-			else
+			const double NormalizedParallelCoef = NormalizedCD ^ NormalizedAB;
+			if (FMath::IsNearlyZero(NormalizedParallelCoef, DOUBLE_KINDA_SMALL_NUMBER))
 			{
-				return Intersect(A, B, D, C);
-			}
-		}
-		else
-		{
-			if (C < D)
-			{
-				return Intersect(B, A, C, D);
-			}
-			else
-			{
-				return Intersect(B, A, D, C);
-			}
-		}
-	};
-
-	const FPoint2D AB = SegmentAB[1] - SegmentAB[0];
-	const FPoint2D CD = SegmentCD[1] - SegmentCD[0];
-	const FPoint2D CA = SegmentAB[0] - SegmentCD[0];
-
-	const double ParallelCoef = CD ^ AB;
-	if (FMath::IsNearlyZero(ParallelCoef, DOUBLE_KINDA_SMALL_NUMBER))
-	{
-		const double ParallelCoef2 = CA ^ AB;
-		if (!FMath::IsNearlyZero(ParallelCoef2, DOUBLE_KINDA_SMALL_NUMBER))
-		{
-			return false;
-		}
+				const double NormalizedParallelCoef2 = NormalizedCA ^ NormalizedAB;
+				if (!FMath::IsNearlyZero(NormalizedParallelCoef2, DOUBLE_KINDA_SMALL_NUMBER))
+				{
+					return false;
+				}
 
-		if (fabs(AB.U) > fabs(AB.V))
-		{
-			return TestWhenParallel(SegmentAB[0].U, SegmentAB[1].U, SegmentCD[0].U, SegmentCD[1].U);
-		}
-		else
-		{
-			return TestWhenParallel(SegmentAB[0].V, SegmentAB[1].V, SegmentCD[0].V, SegmentCD[1].V);
+				if (fabs(Context.AB.U) > fabs(Context.AB.V))
+				{
+					if (FastIntersectionTestWhenCoincident(SegmentAB[0].U, SegmentAB[1].U, SegmentCD[0].U, SegmentCD[1].U))
+					{
+						return ConfirmIntersectionWhenNearlyCoincident(Context);
+					}
+					else
+					{
+						return false;
+					}
+				}
+				else
+				{
+					if (FastIntersectionTestWhenCoincident(SegmentAB[0].V, SegmentAB[1].V, SegmentCD[0].V, SegmentCD[1].V))
+					{
+						return ConfirmIntersectionWhenNearlyCoincident(Context);
+					}
+					else
+					{
+						return false;
+					}
+				}
+			}
 		}
 	}
 
-	const double ABIntersectionCoordinate = (CA ^ CD) / ParallelCoef;
-	const double CDIntersectionCoordinate = (CA ^ AB) / ParallelCoef;
+	const double ABIntersectionCoordinate = (Context.CA ^ Context.CD) / ParallelCoef;
+	const double CDIntersectionCoordinate = (Context.CA ^ Context.AB) / ParallelCoef;
 	return (ABIntersectionCoordinate <= Max && ABIntersectionCoordinate >= Min && CDIntersectionCoordinate <= Max && CDIntersectionCoordinate >= Min);
 }
 
+bool DoIntersectInside(const FSegment2D& SegmentAB, const FSegment2D& SegmentCD)
+{
+	return DoIntersect(SegmentAB, SegmentCD, IntersectionTool::DoCoincidentSegmentsIntersectInside, IntersectionTool::MinValue(true), IntersectionTool::MaxValue(true));
+}
+
+bool DoIntersect(const FSegment2D& SegmentAB, const FSegment2D& SegmentCD)
+{
+	return DoIntersect(SegmentAB, SegmentCD, IntersectionTool::DoCoincidentSegmentsIntersect, IntersectionTool::MinValue(false), IntersectionTool::MaxValue(false));
+}
+
 }
diff --git a/Engine/Source/Runtime/Datasmith/CADKernel/Private/CADKernel/Mesh/Meshers/CycleTriangulator.cpp b/Engine/Source/Runtime/Datasmith/CADKernel/Private/CADKernel/Mesh/Meshers/CycleTriangulator.cpp
--- a/Engine/Source/Runtime/Datasmith/CADKernel/Private/CADKernel/Mesh/Meshers/CycleTriangulator.cpp
+++ b/Engine/Source/Runtime/Datasmith/CADKernel/Private/CADKernel/Mesh/Meshers/CycleTriangulator.cpp
@@ -353,81 +353,7 @@ bool FCycleTriangulator::FindTheBestAcuteTriangle()
 		return true;
 	}
 
-	double BestIsoscelesCriteria = Slope::NullSlope;
-
-	// Check in the acute Triangles the most isosceles
-	int32 AcuteNodeIndex = -1;
-	constexpr double MaxSlopeToBeSelected = Slope::ThreeQuaterPiSlope;
-
-//	for (; false && NIndex < SubCycleNodeCount; ++NIndex)
-//	{
-//		const TPair<int, double>& Candidate = VertexIndexToSlopes[NIndex];
-//
-//		if (Candidate.Value > MaxSlopeToBeSelected)
-//		{
-//			break;
-//		}
-//
-//		GetCandidateTriangleNodes(Candidate.Key);
-//
-//		// Check that the cycle is not inside the selected triangle
-//		if (IsCycleInsideCandidate(*CandidatePoint, *StartPoint, *EndPoint, *NextPoint, TransformIntoCounterClockwise)
-//			|| IsCycleInsideCandidate(*CandidatePoint, *EndPoint, *StartPoint, *PrevPoint, TransformIntoClockwise))
-//		{
-//			continue;
-//		}
-//
-//		//FIsoSegment* Segment;
-//		if (!CycleIntersectionTool.DoesIntersect(NewSegmentStart, NewSegmentEnd/*, &Segment*/))
-//		{
-//			double IsoscelesCriteria = IsoTriangulatorImpl::IsoscelesCriteriaMin(*StartPoint, *EndPoint, *CandidatePoint);
-//
-//			if (BestIsoscelesCriteria < IsoscelesCriteria)
-//			{
-//				double Length1 = StartPoint->SquareDistance(*CandidatePoint);
-//
-//				if (Length1 > MaxEdgeSquareLength)
-//				{
-//#ifdef DEBUG_FIND_THE_CYCLE_TO_MESH
-//					if (Grid.bDisplay)
-//					{
-//						F3DDebugSession W(FString::Printf(TEXT("cancel due to Seg Length %f vs %f"), sqrt(Length1), sqrt(MaxEdgeSquareLength)));
-//						Grid.DisplayIsoSegment(EGridSpace::UniformScaled, *NewSegmentStart, *CandidateNode, 0, EVisuProperty::YellowCurve);
-//						Wait();
-//					}
-//#endif
-//					continue;
-//				}
-//
-//				double Length2 = EndPoint->SquareDistance(*CandidatePoint);
-//
-//				if (Length2 > MaxEdgeSquareLength)
-//				{
-//#ifdef DEBUG_FIND_THE_CYCLE_TO_MESH
-//					if (Grid.bDisplay)
-//					{
-//						F3DDebugSession W(FString::Printf(TEXT("cancel due to Seg Length %f vs %f"), sqrt(Length2), sqrt(MaxEdgeSquareLength)));
-//						Grid.DisplayIsoSegment(EGridSpace::UniformScaled, *NewSegmentEnd, *CandidateNode, 0, EVisuProperty::YellowCurve);
-//						Wait();
-//					}
-//#endif
-//					continue;
-//				}
-//
-//				BestIsoscelesCriteria = IsoscelesCriteria;
-//				AcuteNodeIndex = Index;
-//				bAcuteTriangle = true;
-//			}
-//		}
-//	}
-
-	if (bAcuteTriangle)
-	{
-		GetCandidateTriangleNodes(AcuteNodeIndex);
-		SetFinalCandidate();
-	}
-
-	return bAcuteTriangle;
+	return false;
 }
 
 // Find candidate nodes i.e. Nodes that can be linked with the Segment to build a valid triangle without intersection with other segments
@@ -568,7 +494,7 @@ bool FCycleTriangulator::FindTheBestCandidateNode()
 	{
 		F3DDebugSession _(FString::Printf(TEXT("Start Segment %d"), BestCandidateNode));
 		Grid.DisplayIsoSegment(EGridSpace::UniformScaled, *FirstSideStartNode, *FirstSideEndNode, 1, EVisuProperty::RedCurve);
-		//Wait();
+		Wait(false);
 	}
 #endif
 
@@ -1195,12 +1121,12 @@ void FCycleTriangulator::ValidateComplementaryNodesWithInsideAndIntersectionsCri
 			Grid.DisplayIsoSegment(EGridSpace::UniformScaled, *Side.Nodes[0], *Side.Nodes[1], 1, EVisuProperty::RedCurve);
 			Grid.DisplayIsoSegment(EGridSpace::UniformScaled, *Side.Nodes[1], *Side.EndNode, 1, EVisuProperty::BlueCurve);
 			Grid.DisplayIsoSegment(EGridSpace::UniformScaled, *Side.Nodes[0], *Side.StartNode, 1, EVisuProperty::BlueCurve);
-			Wait();
+			Wait(false);
 		}
 #endif
 		// Check if [Side.Nodes[0] - Side.Nodes[1]] intersect cycle
 		const FIsoSegment* Segment = Side.Nodes[0]->GetSegmentConnectedTo(Side.Nodes[1]);
-		if (!Segment && (CycleIntersectionTool.DoesIntersect(Side.Nodes[0], Side.Nodes[1]) || !IsInnerSideSegmentInsideCycle(Side)))
+		if (!IsInnerSideSegmentInsideCycle(Side) || (!Segment && CycleIntersectionTool.DoesIntersect(Side.Nodes[0], Side.Nodes[1])))
 		{
 			// intersect => choose between NextFinalNode and LastPotentialNode
 #ifdef DEBUG_VALID_COMPLEMENTARY_NODE
@@ -1209,7 +1135,7 @@ void FCycleTriangulator::ValidateComplementaryNodesWithInsideAndIntersectionsCri
 				F3DDebugSession _(TEXT("choose between NextFinalNode and LastPotentialNode"));
 				Grid.DisplayIsoSegment(EGridSpace::UniformScaled, *Side.Nodes[0], *Side.EndNode, 1, EVisuProperty::BlueCurve);
 				Grid.DisplayIsoSegment(EGridSpace::UniformScaled, *Side.Nodes[1], *Side.StartNode, 1, EVisuProperty::BlueCurve);
-				Wait();
+				Wait(false);
 			}
 #endif
 			ValidComplementaryNodeOrDeleteIt(Side, 0);
@@ -1982,7 +1908,16 @@ bool FCycleTriangulator::BuildSegmentIfNeeded(FIsoNode* NodeA, FIsoNode* NodeB,
 
 		FIsoSegment& NewSegment = IsoSegmentFactory.New();
 		NewSegment.Init(*NodeA, *NodeB, ESegmentType::Unknown);
-		NewSegment.ConnectToNode();
+		if (!NewSegment.ConnectToNode())
+		{
+#ifdef DEBUG_BUILD_SEGMENT_IF_NEEDED
+			F3DDebugSession _(FString::Printf(TEXT("ERROR Segment")));
+			Grid.DisplayIsoSegment(EGridSpace::UniformScaled, *NodeA, *NodeB, 0, EVisuProperty::RedCurve);
+			Wait();
+#endif
+			IsoSegmentFactory.DeleteEntity(&NewSegment);
+			return false;
+		}
 		NewSegment.SetHasTriangleOnLeft();
 		SegmentStack.Add(&NewSegment);
 
diff --git a/Engine/Source/Runtime/Datasmith/CADKernel/Private/CADKernel/Mesh/Meshers/CycleTriangulator.h b/Engine/Source/Runtime/Datasmith/CADKernel/Private/CADKernel/Mesh/Meshers/CycleTriangulator.h
--- a/Engine/Source/Runtime/Datasmith/CADKernel/Private/CADKernel/Mesh/Meshers/CycleTriangulator.h
+++ b/Engine/Source/Runtime/Datasmith/CADKernel/Private/CADKernel/Mesh/Meshers/CycleTriangulator.h
@@ -6,6 +6,7 @@
 #include "CADKernel/Math/SlopeUtils.h"
 #include "CADKernel/Mesh/MeshEnum.h"
 #include "CADKernel/Mesh/Meshers/IsoTriangulator/IntersectionSegmentTool.h"
+#include "CADKernel/Mesh/Meshers/IsoTriangulator/IntersectionIsoSegmentTool.h"
 
 #ifdef CADKERNEL_DEV
 #include "CADKernel/UI/DefineForDebug.h"
@@ -121,7 +122,7 @@ private:
 	const TArray<bool>& CycleOrientation;
 
 	const EGridSpace Space = EGridSpace::UniformScaled;
-	const FIntersectionNodePairTool& InnerToOuterIsoSegmentsIntersectionTool;
+	const FIntersectionIsoSegmentTool& InnerToOuterIsoSegmentsIntersectionTool;
 	FFaceMesh& Mesh;
 	TFactory<FIsoSegment>& IsoSegmentFactory;
 
@@ -251,20 +252,6 @@ private:
 		return false;
 	}
 
-	int32 IsIntersectingIso(const TArray<FCandidateNode>& Nodes)
-	{
-		const FIsoNode* StartNode = Nodes.Last().Node;
-		for (const FCandidateNode& CNode : Nodes)
-		{
-			const FIsoNode* Node = CNode.Node;
-			if (InnerToOuterIsoSegmentsIntersectionTool.DoesIntersect(*StartNode, *Node))
-			{
-				return true;
-			}
-		}
-		return false;
-	}
-
 	int32 CountIntersectionWithIso()
 	{
 		FirstSideStartNode = SubCycleNodes[FirstSideStartIndex];
diff --git a/Engine/Source/Runtime/Datasmith/CADKernel/Private/CADKernel/Mesh/Meshers/IsoTriangulator.cpp b/Engine/Source/Runtime/Datasmith/CADKernel/Private/CADKernel/Mesh/Meshers/IsoTriangulator.cpp
--- a/Engine/Source/Runtime/Datasmith/CADKernel/Private/CADKernel/Mesh/Meshers/IsoTriangulator.cpp
+++ b/Engine/Source/Runtime/Datasmith/CADKernel/Private/CADKernel/Mesh/Meshers/IsoTriangulator.cpp
@@ -62,7 +62,7 @@ FIsoTriangulator::FIsoTriangulator(FGrid& InGrid, FFaceMesh& OutMesh, const FMes
 	, Mesh(OutMesh)
 	, LoopSegmentsIntersectionTool(InGrid, InTolerances.GeometricTolerance)
 	, InnerSegmentsIntersectionTool(InGrid, InTolerances.GeometricTolerance)
-	, InnerToOuterIsoSegmentsIntersectionTool(InGrid, InTolerances.GeometricTolerance, /*AllowOverlaping*/ true)
+	, InnerToOuterIsoSegmentsIntersectionTool(InGrid, InTolerances.GeometricTolerance)
 	, ThinZoneIntersectionTool(InGrid, InTolerances.GeometricTolerance)
 	, Tolerances(InTolerances)
 {
@@ -336,8 +336,19 @@ void FIsoTriangulator::BuildLoopSegments()
 
 		FIsoSegment& Segment = IsoSegmentFactory.New();
 		Segment.Init(Node, Node.GetNextNode(), ESegmentType::Loop);
-		Segment.ConnectToNode();
-		LoopSegments.Add(&Segment);
+		if (Segment.ConnectToNode())
+		{
+			LoopSegments.Add(&Segment);
+		}
+		else
+		{
+#ifdef DEBUG_BUILD_SEGMENT_IF_NEEDED
+			F3DDebugSession _(FString::Printf(TEXT("ERROR Segment")));
+			Grid.DisplayIsoSegment(EGridSpace::UniformScaled, Node, Node.GetNextNode(), 0, EVisuProperty::RedCurve);
+			Wait();
+#endif
+			IsoSegmentFactory.DeleteEntity(&Segment);
+		}
 	}
 
 #ifdef CADKERNEL_DEV
@@ -488,13 +499,23 @@ void FIsoTriangulator::GetThinZoneMesh(const TMap<int32, FLoopNode*>& IndexToNod
 		FIsoSegment& Segment = IsoSegmentFactory.New();
 		Segment.Init(CandidateSegment.StartNode, CandidateSegment.EndNode, ESegmentType::ThinZone);
 
-		CandidateSegment.StartNode.SetThinZoneNodeMarker();
-		CandidateSegment.EndNode.SetThinZoneNodeMarker();
-
-		Segment.ConnectToNode();
-		Segment.SetFinalMarker();
-		ThinZoneSegments.Add(&Segment);
-		ThinZoneIntersectionTool.AddSegment(Segment);
+		if(Segment.ConnectToNode())
+		{
+			CandidateSegment.StartNode.SetThinZoneNodeMarker();
+			CandidateSegment.EndNode.SetThinZoneNodeMarker();
+			Segment.SetFinalMarker();
+			ThinZoneSegments.Add(&Segment);
+			ThinZoneIntersectionTool.AddSegment(Segment);
+		}
+		else
+		{
+#ifdef DEBUG_BUILD_SEGMENT_IF_NEEDED
+			F3DDebugSession _(FString::Printf(TEXT("ERROR Segment")));
+			Grid.DisplayIsoSegment(EGridSpace::UniformScaled, CandidateSegment.StartNode, CandidateSegment.EndNode, 0, EVisuProperty::RedCurve);
+			Wait();
+#endif
+			IsoSegmentFactory.DeleteEntity(&Segment);
+		}
 	}
 }
 
@@ -578,15 +599,12 @@ void FIsoTriangulator::BuildInnerSegments()
 		return false;
 	};
 
-	TFunction<void(const int32, const int32)> AddToInnerToOuterSegmentsIntersectionTool = [&](const int32 IndexNode1, const int32 IndexNode2)
+	TFunction<void(const int32, const int32, const ESegmentType)> AddToInnerToOuterSegmentsIntersectionTool = [&](const int32 IndexNode1, const int32 IndexNode2, const ESegmentType InType)
 	{
 		const FPoint2D& Point1 = Grid.GetInner2DPoint(EGridSpace::UniformScaled, IndexNode1);
 		const FPoint2D& Point2 = Grid.GetInner2DPoint(EGridSpace::UniformScaled, IndexNode2);
 
-		const FIsoInnerNode* Node1 = GlobalIndexToIsoInnerNodes[IndexNode1];
-		const FIsoInnerNode* Node2 = GlobalIndexToIsoInnerNodes[IndexNode2];
-
-		InnerToOuterIsoSegmentsIntersectionTool.AddSegment(Node1, Node2, Point1, Point2);
+		InnerToOuterIsoSegmentsIntersectionTool.AddIsoSegment(Point1, Point2, InType);
 	};
 
 	TFunction<void(const int32, const int32, const ESegmentType)> AddToInnerSegments = [&](const int32 IndexNode1, const int32 IndexNode2, const ESegmentType InType)
@@ -595,8 +613,19 @@ void FIsoTriangulator::BuildInnerSegments()
 		FIsoInnerNode& Node2 = *GlobalIndexToIsoInnerNodes[IndexNode2];
 		FIsoSegment& Segment = IsoSegmentFactory.New();
 		Segment.Init(Node1, Node2, InType);
-		Segment.ConnectToNode();
-		FinalInnerSegments.Add(&Segment);
+		if (Segment.ConnectToNode())
+		{
+			FinalInnerSegments.Add(&Segment);
+		}
+		else
+		{
+#ifdef DEBUG_BUILD_SEGMENT_IF_NEEDED
+			F3DDebugSession _(FString::Printf(TEXT("ERROR Segment")));
+			Grid.DisplayIsoSegment(EGridSpace::UniformScaled, Node1, Node2, 0, EVisuProperty::RedCurve);
+			Wait();
+#endif
+			IsoSegmentFactory.DeleteEntity(&Segment);
+		}
 	};
 
 	TFunction<void(const int32, const int32, const ESegmentType, const double)> BuildSegmentIfValid = [&](const int32 IndexNode1, const int32 IndexNode2, const ESegmentType InType, const double Tolerance)
@@ -609,7 +638,7 @@ void FIsoTriangulator::BuildInnerSegments()
 				DisplaySegmentWithScale(Grid.GetInner2DPoint(EGridSpace::UniformScaled, IndexNode1), Grid.GetInner2DPoint(EGridSpace::UniformScaled, IndexNode2), IndexNode1, EVisuProperty::PinkCurve);
 			}
 #endif
-			AddToInnerToOuterSegmentsIntersectionTool(IndexNode1, IndexNode2);
+			AddToInnerToOuterSegmentsIntersectionTool(IndexNode1, IndexNode2, InType);
 			return;
 		}
 
@@ -636,7 +665,7 @@ void FIsoTriangulator::BuildInnerSegments()
 					Wait();
 				}
 #endif
-				AddToInnerToOuterSegmentsIntersectionTool(IndexNode1, IndexNode2);
+				AddToInnerToOuterSegmentsIntersectionTool(IndexNode1, IndexNode2, InType);
 			}
 			else
 			{
@@ -681,7 +710,7 @@ void FIsoTriangulator::BuildInnerSegments()
 			DisplaySegmentWithScale(Grid.GetInner2DPoint(EGridSpace::UniformScaled, IndexNode1), Grid.GetInner2DPoint(EGridSpace::UniformScaled, IndexNode2), IndexNode1, EVisuProperty::YellowCurve);
 		}
 #endif
-		AddToInnerToOuterSegmentsIntersectionTool(IndexNode1, IndexNode2);
+		AddToInnerToOuterSegmentsIntersectionTool(IndexNode1, IndexNode2, InType);
 	};
 
 	TFunction<const TArray<double>(const TArray<double>&)> ComputeLocalTolerance = [](const TArray<double>& UniformCutting) -> const TArray<double>
@@ -899,6 +928,17 @@ void FIsoTriangulator::ConnectCellLoops()
 	TArray<FCell> Cells;
 	FindCellContainingBoundaryNodes(Cells);
 
+#ifdef DEBUG_CONNECT_CELL_LOOPS
+	if (bDisplay)
+	{
+		for (FCell& Cell : Cells)
+		{
+			F3DDebugSession _(FString::Printf(TEXT("Cell %d"), Cell.Id));
+			DisplayCell(Cell);
+		}
+	}
+#endif
+
 	for (FCell& Cell : Cells)
 	{
 		InitCellCorners(Cell);
@@ -923,7 +963,7 @@ void FIsoTriangulator::ConnectCellLoops()
 		{
 			Cell.FindCandidateToConnectLoopsByNeighborhood();
 
-#ifdef DEBUG_CONNECTCELLCORNERTOINNERLOOP
+#ifdef DEBUG_CONNECT_CELL_CORNER_TO_INNER_LOOP
 			if (bDisplay)
 			{
 				DisplayCellConnexions(FString::Printf(TEXT("Cell %d CandidateSegments"), Cell.Id), Cell.LoopConnexions, EVisuProperty::YellowCurve);
@@ -932,7 +972,7 @@ void FIsoTriangulator::ConnectCellLoops()
 
 			FindCandidateToConnectCellCornerToLoops(Cell);
 
-#ifdef DEBUG_CONNECTCELLCORNERTOINNERLOOP
+#ifdef DEBUG_CONNECT_CELL_CORNER_TO_INNER_LOOP
 			if (bDisplay)
 			{
 				DisplayCellConnexions(FString::Printf(TEXT("Cell %d CandidateSegments"), Cell.Id), Cell.LoopConnexions, EVisuProperty::BlueCurve);
@@ -996,7 +1036,7 @@ void FIsoTriangulator::FindCellContainingBoundaryNodes(TArray<FCell>& Cells)
 				if (bDisplay)
 				{
 					F3DDebugSession S(FString::Printf(TEXT("Index %d"), IndexV * CountU + IndexU));
-					DisplayPoint(LoopPoint.GetPoint(EGridSpace::UniformScaled, Grid), EVisuProperty::BluePoint);
+					DisplayPoint2DWithScale(LoopPoint.GetPoint(EGridSpace::UniformScaled, Grid), EVisuProperty::BluePoint);
 					DrawCellBoundary(IndexV * CountU + IndexU, EVisuProperty::YellowPoint);
 				}
 #endif
@@ -1009,7 +1049,7 @@ void FIsoTriangulator::FindCellContainingBoundaryNodes(TArray<FCell>& Cells)
 				if (bDisplay)
 				{
 					F3DDebugSession S(TEXT("Deleted Node"));
-					DisplayPoint(LoopPoint.GetPoint(EGridSpace::UniformScaled, Grid), EVisuProperty::RedPoint);
+					DisplayPoint2DWithScale(LoopPoint.GetPoint(EGridSpace::UniformScaled, Grid), EVisuProperty::RedPoint);
 				}
 #endif
 
@@ -1061,7 +1101,11 @@ void FIsoTriangulator::FindCellContainingBoundaryNodes(TArray<FCell>& Cells)
 				CellNodes.Reset(LoopNodeCount);
 
 #ifdef DEBUG_BUILD_CELLS
-				DisplayCell(Cells.Last());
+				if(bDisplay)
+				{
+					DisplayCell(Cells.Last());
+					Wait();
+				}
 #endif
 			}
 
@@ -1074,8 +1118,11 @@ void FIsoTriangulator::FindCellContainingBoundaryNodes(TArray<FCell>& Cells)
 		Cells.Emplace(CellIndex, CellNodes, *this);
 
 #ifdef DEBUG_BUILD_CELLS
-		DisplayCell(Cells.Last());
-		Wait(bDisplay);
+		if (bDisplay)
+		{
+			DisplayCell(Cells.Last());
+			Wait();
+		}
 #endif
 	}
 
@@ -1278,6 +1325,7 @@ void FIsoTriangulator::TriangulateOverCycle(const EGridSpace Space)
 	}
 #endif
 
+	// first the external segments (loop segments) are processed 
 	for (FIsoSegment* Segment : LoopSegments)
 	{
 		if (!Segment->HasCycleOnLeft())
@@ -1285,7 +1333,10 @@ void FIsoTriangulator::TriangulateOverCycle(const EGridSpace Space)
 			Cycle.Empty();
 			CycleOrientation.Empty();
 			bool bLeftSide = true;
-			FindCycle(Segment, bLeftSide, Cycle, CycleOrientation);
+			if (!FindCycle(Segment, bLeftSide, Cycle, CycleOrientation))
+			{
+				continue;
+			}
 #ifdef FIND_CYCLE
 			if (Grid.bDisplay)
 			{
@@ -1298,6 +1349,7 @@ void FIsoTriangulator::TriangulateOverCycle(const EGridSpace Space)
 		}
 	}
 
+	// then all segments are processed 
 	for (FIsoSegment* Segment : FinalToLoops)
 	{
 		if (!Segment->HasCycleOnLeft())
@@ -1635,6 +1687,13 @@ void FCell::SelectSegmentToConnectLoops(TFactory<FIsoSegment>& SegmentFactory)
 	{
 #ifdef DEBUG_CONNECT_CELL_SUB_LOOPS_BY_NEIGHBORHOOD
 		F3DDebugSession _(Grid.bDisplay, TEXT("Best Path"));
+#endif
+#ifdef DEBUG_CONNECT_CELL_SUB_LOOPS_BY_NEIGHBORHOOD
+		if (Grid.bDisplay)
+		{
+			Grid.DisplayIsoSegments(TEXT("SelectSegmentToConnectLoops: before"), EGridSpace::UniformScaled, FinalSegments, false, false, false, EVisuProperty::YellowCurve);
+			Wait(Grid.bDisplay);
+		}
 #endif
 		const int32 LoopCount = CellLoops.Num();
 		for (FCellConnexion* LoopConnexion : LoopConnexionPtrs)
@@ -1657,7 +1716,7 @@ void FCell::SelectSegmentToConnectLoops(TFactory<FIsoSegment>& SegmentFactory)
 #ifdef DEBUG_CONNECT_CELL_SUB_LOOPS_BY_NEIGHBORHOOD
 	if (Grid.bDisplay)
 	{
-		Grid.DisplayIsoSegments(TEXT("Build Segments Connect inner close loops: Candidates"), EGridSpace::UniformScaled, FinalSegments, false, false, false, EVisuProperty::YellowCurve);
+		Grid.DisplayIsoSegments(TEXT("SelectSegmentToConnectLoops: Final"), EGridSpace::UniformScaled, FinalSegments, false, false, false, EVisuProperty::YellowCurve);
 		Wait(Grid.bDisplay);
 	}
 #endif
@@ -1728,7 +1787,7 @@ void FCell::SelectSegmentToConnectLoopToCorner(TFactory<FIsoSegment>& SegmentFac
 #endif
 }
 
-void FCellLoop::SpreadAsConnected()
+void FCellLoop::PropagateAsConnected()
 {
 	bIsConnected = true;
 	for (FCellConnexion* Connexion : Connexions)
@@ -1741,7 +1800,7 @@ void FCellLoop::SpreadAsConnected()
 		FCellLoop* OtherCell = Connexion->GetOtherLoop(this);
 		if (!OtherCell->bIsConnected)
 		{
-			OtherCell->SpreadAsConnected();
+			OtherCell->PropagateAsConnected();
 		}
 	}
 }
@@ -1750,14 +1809,14 @@ void FCell::CheckAllLoopsConnectedTogetherAndConnect()
 {
 	for (FCellCorner& CellCorner : CellCorners)
 	{
-		CellCorner.SpreadAsConnected();
+		CellCorner.PropagateAsConnected();
 	}
 
 	for (FCellLoop& CellLoop : CellLoops)
 	{
 		if (CellLoop.bIsOuterLoop)
 		{
-			CellLoop.SpreadAsConnected();
+			CellLoop.PropagateAsConnected();
 		}
 	}
 }
@@ -1771,11 +1830,28 @@ void FCell::TryToCreateSegment(FCellConnexion& LoopConnexion)
 	if (LoopConnexion.Segment && !LoopConnexion.Segment->IsAFinalSegment())
 	{
 		LoopConnexion.Segment->SetFinalMarker();
-		LoopConnexion.Segment->ConnectToNode();
-
-		IntersectionTool.AddSegment(*LoopConnexion.Segment);
-		FinalSegments.Add(LoopConnexion.Segment);
+		if(LoopConnexion.Segment->ConnectToNode())
+		{
+			IntersectionTool.AddSegment(*LoopConnexion.Segment);
+			FinalSegments.Add(LoopConnexion.Segment);
+		}
+		else
+		{
+#ifdef DEBUG_BUILD_SEGMENT_IF_NEEDED
+			F3DDebugSession _(FString::Printf(TEXT("Not expected error Segment")));
+			Grid.DisplayIsoSegment(EGridSpace::UniformScaled, *LoopConnexion.NodeA, *LoopConnexion.NodeB, 0, EVisuProperty::RedCurve);
+			Wait();
+#endif
+		}
+	}
+#ifdef DEBUG_BUILD_SEGMENT_IF_NEEDED2
+	else if(Grid.bDisplay)
+	{
+		F3DDebugSession _(FString::Printf(TEXT("Failed Segment %d %d"), LoopConnexion.Loop1.Id, LoopConnexion.Loop2.Id));
+		Grid.DisplayIsoSegment(EGridSpace::UniformScaled, *LoopConnexion.NodeA, *LoopConnexion.NodeB, 0, EVisuProperty::RedCurve);
+		Wait();
 	}
+#endif
 }
 
 
@@ -2284,7 +2360,7 @@ void FIsoTriangulator::FindCandidateToConnectCellCornerToLoops(FCell& Cell)
 		return;
 	}
 
-#ifdef DEBUG_CONNECTCELLCORNERTOINNERLOOP
+#ifdef DEBUG_CONNECT_CELL_CORNER_TO_INNER_LOOP
 	F3DDebugSession _(bDisplay, TEXT("With cell corners"));
 #endif
 
@@ -2352,7 +2428,7 @@ void FIsoTriangulator::FindCandidateToConnectCellCornerToLoops(FCell& Cell)
 
 	for (FCellCorner& CellCorner : Cell.CellCorners)
 	{
-#ifdef DEBUG_CONNECTCELLCORNERTOINNERLOOP
+#ifdef DEBUG_CONNECT_CELL_CORNER_TO_INNER_LOOP
 		if (bDisplay)
 		{
 			F3DDebugSession _(bDisplay, FString::Printf(TEXT("Cell corner %d"), CellCorner.Id));
@@ -2365,7 +2441,7 @@ void FIsoTriangulator::FindCandidateToConnectCellCornerToLoops(FCell& Cell)
 		}
 	}
 
-#ifdef DEBUG_CONNECTCELLCORNERTOINNERLOOP
+#ifdef DEBUG_CONNECT_CELL_CORNER_TO_INNER_LOOP
 	if (bDisplay)
 	{
 		DisplayCellConnexions(TEXT("Cell CandidateSegments"), Cell.LoopConnexions, EVisuProperty::YellowCurve);
@@ -2428,7 +2504,15 @@ void FIsoTriangulator::SelectSegmentsToLinkInnerToLoop()
 		IntersectionTool.AddSegment(*Segment);
 		Segment->SetSelected();
 		Segment->SetFinalMarker();
-		Segment->ConnectToNode();
+		if (!Segment->ConnectToNode())
+		{
+#ifdef DEBUG_BUILD_SEGMENT_IF_NEEDED
+			F3DDebugSession _(FString::Printf(TEXT("ERROR Segment")));
+			Grid.DisplayIsoSegment(EGridSpace::UniformScaled, *Segment, 0, EVisuProperty::RedCurve);
+			Wait();
+#endif
+			IsoSegmentFactory.DeleteEntity(Segment);
+		}
 	}
 	CandidateSegments.Empty();
 }
diff --git a/Engine/Source/Runtime/Datasmith/CADKernel/Private/CADKernel/Mesh/Meshers/IsoTriangulator.h b/Engine/Source/Runtime/Datasmith/CADKernel/Private/CADKernel/Mesh/Meshers/IsoTriangulator.h
--- a/Engine/Source/Runtime/Datasmith/CADKernel/Private/CADKernel/Mesh/Meshers/IsoTriangulator.h
+++ b/Engine/Source/Runtime/Datasmith/CADKernel/Private/CADKernel/Mesh/Meshers/IsoTriangulator.h
@@ -5,6 +5,7 @@
 #include "CADKernel/Math/Geometry.h"
 #include "CADKernel/Math/Point.h"
 #include "CADKernel/Math/SlopeUtils.h"
+#include "CADKernel/Mesh/Meshers/IsoTriangulator/IntersectionIsoSegmentTool.h"
 #include "CADKernel/Mesh/Meshers/IsoTriangulator/IntersectionSegmentTool.h"
 #include "CADKernel/Mesh/Meshers/IsoTriangulator/IsoNode.h"
 #include "CADKernel/Mesh/Meshers/IsoTriangulator/IsoSegment.h"
@@ -126,7 +127,7 @@ protected:
 	FIntersectionSegmentTool InnerSegmentsIntersectionTool;
 	
 	/** To check if the candidate segment intersect a iso line */
-	FIntersectionNodePairTool InnerToOuterIsoSegmentsIntersectionTool;
+	FIntersectionIsoSegmentTool InnerToOuterIsoSegmentsIntersectionTool;
 	
 	/** To check if a candidate segment intersect a thin zone that is already meshed */
 	FIntersectionSegmentTool ThinZoneIntersectionTool;
diff --git a/Engine/Source/Runtime/Datasmith/CADKernel/Private/CADKernel/Mesh/Meshers/IsoTriangulator/IntersectionIsoSegmentTool.cpp b/Engine/Source/Runtime/Datasmith/CADKernel/Private/CADKernel/Mesh/Meshers/IsoTriangulator/IntersectionIsoSegmentTool.cpp
new file mode 100644
--- /dev/null
+++ b/Engine/Source/Runtime/Datasmith/CADKernel/Private/CADKernel/Mesh/Meshers/IsoTriangulator/IntersectionIsoSegmentTool.cpp
@@ -0,0 +1,256 @@
+// Copyright Epic Games, Inc. All Rights Reserved.
+#pragma once
+
+#include "CADKernel/Mesh/Meshers/IsoTriangulator/IntersectionIsoSegmentTool.h"
+#include "CADKernel/Geo/GeoEnum.h"
+
+namespace UE::CADKernel
+{
+
+FIntersectionIsoSegmentTool::FIntersectionIsoSegmentTool(const FGrid& InGrid, const double Tolerance)
+	: Grid(InGrid)
+{
+	TFunction<void(TArray<TPair<double, int32>>&, const TArray<double>&)> Reserve = [](TArray<TPair<double, int32>>& InCoordToIndex, const TArray<double>& InCuttingCoordinates)
+	{
+		InCoordToIndex.Reserve(InCuttingCoordinates.Num());
+		for (const double Coordinate : InCuttingCoordinates)
+		{
+			InCoordToIndex.Emplace(Coordinate, 0);
+		}
+	};
+
+	Reserve(CoordToIndex[EIso::IsoU], Grid.GetUniformCuttingCoordinatesAlongIso(EIso::IsoU));
+	Reserve(CoordToIndex[EIso::IsoV], Grid.GetUniformCuttingCoordinatesAlongIso(EIso::IsoV));
+
+	const int32 MaxSegmentCount = FMath::Min(Grid.InnerNodesCount(), Grid.GetCuttingCount(EIso::IsoU) * Grid.GetCuttingCount(EIso::IsoV) - Grid.InnerNodesCount());
+	USegments.Reserve(MaxSegmentCount);
+	VSegments.Reserve(MaxSegmentCount);
+}
+
+void FIntersectionIsoSegmentTool::AddIsoSegment(const FPoint2D& StartPoint, const FPoint2D& EndPoint, const ESegmentType InType)
+{
+	switch(InType)
+	{ 
+	case ESegmentType::IsoU:
+		USegments.Emplace(StartPoint, EndPoint);
+		break;
+
+	case ESegmentType::IsoV:
+		VSegments.Emplace(StartPoint, EndPoint);
+		break;
+
+	default:
+		break;
+	}
+}
+
+template <class TIsoSegment>
+int32 TCountIntersections(const TArray<TIsoSegment>& Segments, EIso Iso, int32 Index, const FSegment2D& InSegment, const FSurfacicBoundary& InSegmentBoundary)
+{
+	using namespace IntersectionIsoSegmentTool;
+
+	const double Max = InSegmentBoundary.Get(Iso).GetMax();
+	
+	int32 IntersectionCount = 0;
+	EIso OtherIso = Other(Iso);
+	for (; Index < Segments.Num(); ++Index)
+	{
+		const FIntersectionIsoSegment& IsoSegment = Segments[Index];
+		if (Max < IsoSegment.GetIsoCoordinate())
+		{
+			break;
+		}
+		if (InSegmentBoundary.Get(OtherIso).GetMin() > IsoSegment.GetMaxCoordinate())
+		{
+			continue;
+		}
+		if (InSegmentBoundary.Get(OtherIso).GetMax() < IsoSegment.GetMinCoordinate())
+		{
+			break;
+		}
+
+		if (InSegmentBoundary.Get(OtherIso).GetMin() >= IsoSegment.GetMinCoordinate() && InSegmentBoundary.Get(OtherIso).GetMax() <= IsoSegment.GetMaxCoordinate())
+		{
+			// Allow Overlapping: is InSegmentBoundary coincident with IsoSegment ?
+			if (InSegmentBoundary.Get(Iso).GetMin() + DOUBLE_KINDA_SMALL_NUMBER > IsoSegment.GetIsoCoordinate()
+				&& InSegmentBoundary.Get(Iso).GetMax() - DOUBLE_KINDA_SMALL_NUMBER < IsoSegment.GetIsoCoordinate())
+			{
+				continue;
+			}
+
+			IntersectionCount++;
+			continue;
+		}
+
+		const FSegment2D IsoSegment2D = IsoSegment.GetSegment2D();
+		if (DoIntersect(InSegment, IsoSegment2D))
+		{
+			IntersectionCount++;
+		}
+	}
+	return IntersectionCount;
+}
+
+template <class TIsoSegment>
+bool TDoesIntersect(const TArray<TIsoSegment>& Segments, EIso Iso, int32 Index, const FSegment2D& InSegment, const FSurfacicBoundary& InSegmentBoundary)
+{
+	using namespace IntersectionIsoSegmentTool;
+
+	const double Max = InSegmentBoundary.Get(Iso).GetMax();
+
+	int32 IntersectionCount = 0;
+
+	EIso OtherIso = Other(Iso);
+	for (; Index < Segments.Num(); ++Index)
+	{
+		const FIntersectionIsoSegment& IsoSegment = Segments[Index];
+		if (Max < IsoSegment.GetIsoCoordinate())
+		{
+			break;
+		}
+		if (InSegmentBoundary.Get(OtherIso).GetMin() > IsoSegment.GetMaxCoordinate())
+		{
+			continue;
+		}
+		if (InSegmentBoundary.Get(OtherIso).GetMax() < IsoSegment.GetMinCoordinate())
+		{
+			break;
+		}
+
+		if (InSegmentBoundary.Get(OtherIso).GetMin() >= IsoSegment.GetMinCoordinate() && InSegmentBoundary.Get(OtherIso).GetMax() <= IsoSegment.GetMaxCoordinate())
+		{
+			// Allow Overlapping: is InSegmentBoundary coincident with IsoSegment ?
+			if (InSegmentBoundary.Get(Iso).GetMin() + DOUBLE_KINDA_SMALL_NUMBER > IsoSegment.GetIsoCoordinate()
+				&& InSegmentBoundary.Get(Iso).GetMax() - DOUBLE_KINDA_SMALL_NUMBER < IsoSegment.GetIsoCoordinate())
+			{
+				continue;
+			}
+
+			return true;
+		}
+
+		const FSegment2D IsoSegment2D = IsoSegment.GetSegment2D();
+		if (DoIntersect(InSegment, IsoSegment2D))
+		{
+			return true;
+		}
+	}
+	return false;
+}
+
+int32 FIntersectionIsoSegmentTool::CountIntersections(const FIsoNode& StartPoint, const FIsoNode& EndPoint) const
+{
+	const FSegment2D Segment(StartPoint.Get2DPoint(EGridSpace::UniformScaled, Grid), EndPoint.Get2DPoint(EGridSpace::UniformScaled, Grid));
+	const FSurfacicBoundary SegmentBoundary(Segment.Point0, Segment.Point1);
+
+	int32 IndexU = GetStartIndex(EIso::IsoU, SegmentBoundary);
+	int32 IndexV = GetStartIndex(EIso::IsoV, SegmentBoundary);
+
+	return TCountIntersections(USegments, EIso::IsoU, IndexU, Segment, SegmentBoundary) + TCountIntersections(VSegments, EIso::IsoV, IndexV, Segment, SegmentBoundary);
+}
+
+bool FIntersectionIsoSegmentTool::DoesIntersect(const FIsoNode& StartPoint, const FIsoNode& EndPoint) const
+{
+	const FSegment2D Segment(StartPoint.Get2DPoint(EGridSpace::UniformScaled, Grid), EndPoint.Get2DPoint(EGridSpace::UniformScaled, Grid));
+	const FSurfacicBoundary SegmentBoundary(Segment.Point0, Segment.Point1);
+
+	int32 IndexU = GetStartIndex(EIso::IsoU, SegmentBoundary);
+
+	if (TDoesIntersect(USegments, EIso::IsoU, IndexU, Segment, SegmentBoundary))
+	{
+		return true;
+	}
+
+	int32 IndexV = GetStartIndex(EIso::IsoV, SegmentBoundary);
+	return TDoesIntersect(VSegments, EIso::IsoV, IndexV, Segment, SegmentBoundary);
+}
+
+int32 FIntersectionIsoSegmentTool::GetStartIndex(EIso Iso, const FSurfacicBoundary& SegmentBoundary) const
+{
+	double Min = SegmentBoundary.Get(Iso).GetMin() - SMALL_NUMBER;
+	for (const TPair<double, int32>& CoordIndex : CoordToIndex[Iso])
+	{
+		if (Min < CoordIndex.Key)
+		{
+			return CoordIndex.Value;
+		}
+	}
+	return CoordToIndex[Iso].Last().Value;
+}
+
+int32 FIntersectionIsoSegmentTool::GetStartIndex(EIso Iso, double Min) const
+{
+	Min -= SMALL_NUMBER;
+	for (const TPair<double, int32>& CoordIndex : CoordToIndex[Iso])
+	{
+		if (Min < CoordIndex.Key)
+		{
+			return CoordIndex.Value;
+		}
+	}
+	return CoordToIndex[Iso].Last().Value;
+}
+
+template <class TIsoSegment>
+void SortCoordToIndex(TArray<TPair<double, int32>>& CoordToIndex, const TArray<TIsoSegment>& Segments)
+{
+	int32 IndexOfCuttingCoordinate = 0;
+	double Coordinate = CoordToIndex[IndexOfCuttingCoordinate].Key + SMALL_NUMBER;
+	int32 SegmentCount = 0;
+	for (const TIsoSegment& Segment : Segments)
+	{
+		const double IsoCoordinate = Segment.GetIsoCoordinate();
+		while (Coordinate < IsoCoordinate)
+		{
+			++IndexOfCuttingCoordinate;
+			CoordToIndex[IndexOfCuttingCoordinate].Value = SegmentCount;
+			Coordinate = CoordToIndex[IndexOfCuttingCoordinate].Key + SMALL_NUMBER;
+		}
+		++SegmentCount;
+	}
+	for (++IndexOfCuttingCoordinate; IndexOfCuttingCoordinate < CoordToIndex.Num(); ++IndexOfCuttingCoordinate)
+	{
+		CoordToIndex[IndexOfCuttingCoordinate].Value = SegmentCount;
+	}
+};
+
+void FIntersectionIsoSegmentTool::Sort()
+{
+	USegments.Sort();
+	VSegments.Sort();
+
+	SortCoordToIndex(CoordToIndex[EIso::IsoU], USegments);
+	SortCoordToIndex(CoordToIndex[EIso::IsoV], VSegments);
+
+	bIsSorted = true;
+}
+
+
+
+#ifdef CADKERNEL_DEBUG
+void FIntersectionIsoSegmentTool::Display(bool bDisplay, const TCHAR* Message, EVisuProperty Property) const
+{
+	if (!bDisplay)
+	{
+		return;
+	}
+
+	int32 Index = 0;
+	Open3DDebugSession(Message);
+	for (const IntersectionIsoSegmentTool::FIntersectionIsoSegment& Segment : USegments)
+	{
+		const FSegment2D Segment2D = Segment.GetSegment2D();
+		DisplaySegment(Segment2D[0] * DisplayScale, Segment2D[1] * DisplayScale, Index++, Property);
+	}
+	for (const IntersectionIsoSegmentTool::FIntersectionIsoSegment& Segment : VSegments)
+	{
+		const FSegment2D Segment2D = Segment.GetSegment2D();
+		DisplaySegment(Segment2D[0] * DisplayScale, Segment2D[1] * DisplayScale, Index++, Property);
+	}
+	Close3DDebugSession();
+	Wait();
+}
+#endif
+
+} // namespace UE::CADKernel
+
diff --git a/Engine/Source/Runtime/Datasmith/CADKernel/Private/CADKernel/Mesh/Meshers/IsoTriangulator/IntersectionIsoSegmentTool.h b/Engine/Source/Runtime/Datasmith/CADKernel/Private/CADKernel/Mesh/Meshers/IsoTriangulator/IntersectionIsoSegmentTool.h
new file mode 100644
--- /dev/null
+++ b/Engine/Source/Runtime/Datasmith/CADKernel/Private/CADKernel/Mesh/Meshers/IsoTriangulator/IntersectionIsoSegmentTool.h
@@ -0,0 +1,148 @@
+// Copyright Epic Games, Inc. All Rights Reserved.
+#pragma once
+
+#include "CADKernel/Core/Types.h"
+#include "CADKernel/Mesh/Meshers/IsoTriangulator/IsoSegment.h"
+#include "CADKernel/Math/Boundary.h"
+#include "CADKernel/UI/Visu.h"
+
+namespace UE::CADKernel
+{
+class FGrid;
+class FIsoNode;
+class FPoint2D;
+
+namespace IntersectionIsoSegmentTool
+{
+
+class FIntersectionIsoSegment
+{
+private:
+	const FPoint2D& Point0;
+	const FPoint2D& Point1;
+	double IsoCoordinate;
+	double MinCoordinate;
+	double MaxCoordinate;
+
+public:
+	
+	FIntersectionIsoSegment(const FPoint2D& StartPoint, const FPoint2D& EndPoint, const double InIsoCoordinate, const double InStartCoordinate, const double InEndCoordinate)
+		: Point0(StartPoint)
+		, Point1(EndPoint)
+		, IsoCoordinate(InIsoCoordinate)
+		, MinCoordinate(FMath::Min(InStartCoordinate, InEndCoordinate))
+		, MaxCoordinate(FMath::Max(InStartCoordinate, InEndCoordinate))
+	{}
+
+	virtual ~FIntersectionIsoSegment()
+	{}
+
+	virtual FPoint2D GetMinPoint() const = 0;
+	virtual FPoint2D GetMaxPoint() const = 0;
+
+	double GetIsoCoordinate() const 
+	{
+		return IsoCoordinate;
+	}
+
+	double GetMinCoordinate() const
+	{
+		return MinCoordinate;
+	}
+
+	double GetMaxCoordinate() const
+	{
+		return MaxCoordinate;
+	}
+
+	FSegment2D GetSegment2D() const
+	{
+		return FSegment2D(Point0, Point1);
+	}
+
+	friend bool operator<(const FIntersectionIsoSegment& A, const FIntersectionIsoSegment& B)
+	{
+		if (FMath::IsNearlyEqual(A.IsoCoordinate, B.IsoCoordinate))
+		{
+			return A.MinCoordinate < B.MinCoordinate;
+		}
+		else
+		{
+			return A.IsoCoordinate < B.IsoCoordinate;
+		}
+	}
+};
+
+class FIsoUSegment : public FIntersectionIsoSegment
+{
+public:
+	FIsoUSegment(const FPoint2D& StartPoint, const FPoint2D& EndPoint)
+		: FIntersectionIsoSegment(StartPoint, EndPoint, StartPoint.U, StartPoint.V, EndPoint.V)
+	{}
+
+	virtual FPoint2D GetMinPoint() const override
+	{
+		return FPoint2D(GetIsoCoordinate(), GetMinCoordinate());
+	}
+	virtual FPoint2D GetMaxPoint() const override
+	{
+		return FPoint2D(GetIsoCoordinate(), GetMaxCoordinate());
+	}
+};
+
+class FIsoVSegment : public FIntersectionIsoSegment
+{
+public:
+	FIsoVSegment(const FPoint2D& StartPoint, const FPoint2D& EndPoint)
+		: FIntersectionIsoSegment(StartPoint, EndPoint, StartPoint.V, StartPoint.U, EndPoint.U)
+	{}
+
+	virtual FPoint2D GetMinPoint() const override
+	{
+		return FPoint2D(GetMinCoordinate(), GetIsoCoordinate());
+	}
+
+	virtual FPoint2D GetMaxPoint() const override
+	{
+		return FPoint2D(GetMaxCoordinate(), GetIsoCoordinate());
+	}
+
+};
+
+}
+
+class FIntersectionIsoSegmentTool
+{
+private:
+	const FGrid& Grid;
+
+	TArray<TPair<double, int32>> CoordToIndex[2];
+	
+	TArray<IntersectionIsoSegmentTool::FIsoUSegment> USegments;
+	TArray<IntersectionIsoSegmentTool::FIsoVSegment> VSegments;
+
+	bool bIsSorted = false;
+
+public:
+	FIntersectionIsoSegmentTool(const FGrid& InGrid, const double Tolerance);
+
+	void AddIsoSegment(const FPoint2D& StartPoint, const FPoint2D& EndPoint, const ESegmentType InType);
+
+	bool DoesIntersect(const FIsoNode& StartNode, const FIsoNode& EndNode) const;
+
+	int32 CountIntersections(const FIsoNode& StartNode, const FIsoNode& EndNode) const;
+
+	void Sort();
+
+private:
+	int32 GetStartIndex(EIso Iso, double Min) const;
+	int32 GetStartIndex(EIso Iso, const FSurfacicBoundary& Boundary) const;
+
+public:
+#ifdef CADKERNEL_DEBUG
+	void Display(bool bDisplay, const TCHAR* Message, EVisuProperty Property = EVisuProperty::BlueCurve) const;
+#endif
+};
+
+} // namespace UE::CADKernel
+
diff --git a/Engine/Source/Runtime/Datasmith/CADKernel/Private/CADKernel/Mesh/Meshers/IsoTriangulator/IntersectionSegmentTool.cpp b/Engine/Source/Runtime/Datasmith/CADKernel/Private/CADKernel/Mesh/Meshers/IsoTriangulator/IntersectionSegmentTool.cpp
--- a/Engine/Source/Runtime/Datasmith/CADKernel/Private/CADKernel/Mesh/Meshers/IsoTriangulator/IntersectionSegmentTool.cpp
+++ b/Engine/Source/Runtime/Datasmith/CADKernel/Private/CADKernel/Mesh/Meshers/IsoTriangulator/IntersectionSegmentTool.cpp
@@ -41,20 +41,6 @@ IntersectionSegmentTool::FSegment::FSegment(const FGrid& Grid, const double Tole
 {
 }
 
-IntersectionNodePairTool::FSegment::FSegment(const double Tolerance, const FIsoNode* InStartNode, const FIsoNode* InEndNode, const FPoint2D& InStartPoint, const FPoint2D& InEndPoint)
-	: IntersectionToolBase::FSegment(Tolerance, InStartPoint, InEndPoint)
-	, StartNode(InStartNode)
-	, EndNode(InEndNode)
-{
-}
-
-IntersectionNodePairTool::FSegment::FSegment(const FGrid& Grid, const double Tolerance, const FIsoNode& InStartNode, const FIsoNode& InEndNode)
-	: IntersectionToolBase::FSegment(Tolerance, InStartNode.Get2DPoint(EGridSpace::UniformScaled, Grid), InEndNode.Get2DPoint(EGridSpace::UniformScaled, Grid))
-	, StartNode(&InStartNode)
-	, EndNode(&InEndNode)
-{
-}
-
 bool IntersectionSegmentTool::FSegment::IsValid() const
 {
 	return IsoSegment != nullptr && !IsoSegment->IsDelete();
diff --git a/Engine/Source/Runtime/Datasmith/CADKernel/Private/CADKernel/Mesh/Meshers/IsoTriangulator/IntersectionSegmentTool.h b/Engine/Source/Runtime/Datasmith/CADKernel/Private/CADKernel/Mesh/Meshers/IsoTriangulator/IntersectionSegmentTool.h
--- a/Engine/Source/Runtime/Datasmith/CADKernel/Private/CADKernel/Mesh/Meshers/IsoTriangulator/IntersectionSegmentTool.h
+++ b/Engine/Source/Runtime/Datasmith/CADKernel/Private/CADKernel/Mesh/Meshers/IsoTriangulator/IntersectionSegmentTool.h
@@ -4,7 +4,6 @@
 #include "CADKernel/Core/Types.h"
 #include "CADKernel/Math/Boundary.h"
 #include "CADKernel/Math/Geometry.h"
-#include "CADKernel/UI/Message.h"
 #include "CADKernel/UI/Visu.h"
 
 namespace UE::CADKernel
@@ -205,46 +204,6 @@ struct FSegment : public IntersectionToolBase::FSegment
 
 }
 
-namespace IntersectionNodePairTool
-{
-
-struct FSegment : public IntersectionToolBase::FSegment
-{
-	const FIsoNode* StartNode;
-	const FIsoNode* EndNode;
-
-	FSegment(const double Tolerance, const FIsoNode* StartNode, const FIsoNode* EndNode, const FPoint2D& StartPoint, const FPoint2D& EndPoint);
-	FSegment(const FGrid& Grid, const double Tolerance, const FIsoNode& StartNode, const FIsoNode& EndNode);
-
-	virtual bool IsValid() const override
-	{
-		return true;
-	}
-
-	virtual const FIsoNode* GetFirstNode() const override
-	{
-		return StartNode;
-	}
-
-	virtual const FIsoNode* GetSecondNode() const override
-	{
-		return EndNode;
-	}
-
-	virtual bool DoesItIntersect(const IntersectionToolBase::FSegment& Segment) const override
-	{
-		if (!CouldItIntersect(Segment.Boundary))
-		{
-			return false;
-		}
-
-		return DoIntersectInside(Segment2D, Segment.Segment2D);
-	}
-
-};
-
-}
-
 template<typename SegmentType>
 class TIntersectionSegmentTool
 {
@@ -256,17 +215,11 @@ protected:
 
 	const double Tolerance;
 
-	// Allow a candidate segment to be parallel with and overlap an existing segment.
-	// This is used by InnerToOuterIsoSegmentsIntersectionTool, which existing segment are Isolines. 
-	// So this allow candidate segment to overlap iso segment. Otherwise, the candidate segment has one intersection with an isoline when we just want that is parallel with.
-	const bool bAllowOverlapping;
-
 public:
-	TIntersectionSegmentTool(const FGrid& InGrid, const double InTolerance, const bool bInAllowOverlapping = false)
+	TIntersectionSegmentTool(const FGrid& InGrid, const double InTolerance)
 		: Grid(InGrid)
 		, bSegmentsAreSorted(false)
 		, Tolerance(InTolerance)
-		, bAllowOverlapping(bInAllowOverlapping)
 	{
 	}
 
@@ -274,7 +227,7 @@ public:
 
 	void Empty(int32 InMaxNum)
 	{
-		Segments.Empty(InMaxNum);
+		Segments.Reset(InMaxNum);
 		bSegmentsAreSorted = false;
 	}
 
@@ -336,16 +289,6 @@ public:
 				}
 			}
 
-			if (Segment.IsParallelWith(InSegment))
-			{
-				printf("");
-			}
-
-			if (bAllowOverlapping && Segment.IsParallelWith(InSegment))
-			{
-				continue;
-			}
-
 			switch (Segment.DoesItStartFromAndSuperimposed(StartExtremity, EndExtremity, InSegment.Segment2D))
 			{
 			case EConnectionType::SameSegment:
@@ -373,7 +316,7 @@ public:
 		SegmentType InSegment(Grid, Tolerance, *StartExtremity, *EndExtremity);
 		if (OutIntersectedSegments)
 		{
-			OutIntersectedSegments->Empty(10);
+			OutIntersectedSegments->Reset(10);
 		}
 
 		int32 IntersectionCount = 0;
@@ -397,11 +340,6 @@ public:
 				}
 			}
 
-			if (bAllowOverlapping && Segment.IsParallelWith(InSegment))
-			{
-				continue;
-			}
-
 			switch (Segment.DoesItStartFromAndSuperimposed(StartExtremity, EndExtremity, InSegment.Segment2D))
 			{
 			case EConnectionType::StartFrom:
@@ -604,37 +542,5 @@ public:
 	}
 };
 
-class FIntersectionNodePairTool : public TIntersectionSegmentTool<IntersectionNodePairTool::FSegment>
-{
-
-public:
-	FIntersectionNodePairTool(const FGrid& InGrid, const double Tolerance, bool bInAllowOverlaping)
-		: TIntersectionSegmentTool<IntersectionNodePairTool::FSegment>(InGrid, Tolerance, bInAllowOverlaping)
-	{
-	}
-
-	void AddSegment(const FIsoNode& StartNode, const FIsoNode& EndNode)
-	{
-		bSegmentsAreSorted = false;
-		Segments.Emplace(Grid, Tolerance, StartNode, EndNode);
-	}
-
-	void AddSegment(const FIsoNode* StartNode, const FIsoNode* EndNode, const FPoint2D& StartPoint, const FPoint2D& EndPoint)
-	{
-		bSegmentsAreSorted = false;
-		Segments.Emplace(Tolerance, StartNode, EndNode, StartPoint, EndPoint);
-	}
-
-	int32 CountIntersections(const FIsoNode& StartNode, const FIsoNode& EndNode) const
-	{
-		return TIntersectionSegmentTool<IntersectionNodePairTool::FSegment>::FindIntersectingSegments(&StartNode, &EndNode, nullptr);
-	}
-
-	bool DoesIntersect(const FIsoNode& StartNode, const FIsoNode& EndNode) const
-	{
-		return TIntersectionSegmentTool<IntersectionNodePairTool::FSegment>::FindIntersectingSegment(&StartNode, &EndNode) != nullptr;
-	}
-};
-
 } // namespace UE::CADKernel
 
diff --git a/Engine/Source/Runtime/Datasmith/CADKernel/Private/CADKernel/Mesh/Meshers/IsoTriangulator/IsoCell.h b/Engine/Source/Runtime/Datasmith/CADKernel/Private/CADKernel/Mesh/Meshers/IsoTriangulator/IsoCell.h
--- a/Engine/Source/Runtime/Datasmith/CADKernel/Private/CADKernel/Mesh/Meshers/IsoTriangulator/IsoCell.h
+++ b/Engine/Source/Runtime/Datasmith/CADKernel/Private/CADKernel/Mesh/Meshers/IsoTriangulator/IsoCell.h
@@ -57,7 +57,7 @@ struct FCellLoop
 		return bIsOuterLoop;
 	}
 
-	void SpreadAsConnected();
+	void PropagateAsConnected();
 };
 
 struct FCellCorner : public FCellLoop
@@ -413,20 +413,42 @@ struct FCell
 		PreviousNode = nullptr;
 		for (FLoopNode* Node : InNodes)
 		{
-			if (Node->IsDelete())
+			Node->SetMarker1();
+		}
+
+		for (FLoopNode* Node : InNodes)
+		{
+			if (Node->IsDeleteOrHasMarker2() )
 			{
 				continue;
 			}
 
-			if (&Node->GetPreviousNode() != PreviousNode)
+			FLoopNode* StartNode = &Node->GetPreviousNode();
+			while (StartNode->HasMarker1() && (StartNode != Node))
 			{
-				MakeLoopCell(LoopNodes);
-				LoopNodes.Reset(NodeCount);
+				StartNode = &StartNode->GetPreviousNode();
 			}
-			LoopNodes.Add(Node);
-			PreviousNode = Node;
+			if (!StartNode->HasMarker1())
+			{
+				StartNode = &StartNode->GetNextNode();
+			}
+
+			LoopNodes.Reset(NodeCount);
+			FLoopNode* NextNode = StartNode;
+			while (NextNode->HasMarker1NotMarker2())
+			{
+				LoopNodes.Add(NextNode);
+				NextNode->SetMarker2();
+				NextNode = &NextNode->GetNextNode();
+			}
+
+			MakeLoopCell(LoopNodes);
+		}
+
+		for (FLoopNode* Node : InNodes)
+		{
+			Node->ResetMarkers();
 		}
-		MakeLoopCell(LoopNodes);
 
 		for(const FCellLoop& LoopCell : CellLoops)
 		{
diff --git a/Engine/Source/Runtime/Datasmith/CADKernel/Private/CADKernel/Mesh/Meshers/IsoTriangulator/IsoNode.h b/Engine/Source/Runtime/Datasmith/CADKernel/Private/CADKernel/Mesh/Meshers/IsoTriangulator/IsoNode.h
--- a/Engine/Source/Runtime/Datasmith/CADKernel/Private/CADKernel/Mesh/Meshers/IsoTriangulator/IsoNode.h
+++ b/Engine/Source/Runtime/Datasmith/CADKernel/Private/CADKernel/Mesh/Meshers/IsoTriangulator/IsoNode.h
@@ -46,6 +46,9 @@ enum class EIsoNodeStates : uint16
 	NearlyIsoU = 0x1000u,  // Loop nodes
 	NearlyIsoV = 0x2000u,  // Loop nodes
 
+	HasMarker1 = 0x4000u,
+	HasMarker2 = 0x8000u,
+
 	All = 0xFFFFu
 };
 
@@ -222,6 +225,44 @@ public:
 		}
 	}
 
+	bool HasMarker1() const
+	{
+		return ((States & EIsoNodeStates::HasMarker1) == EIsoNodeStates::HasMarker1);
+	}
+
+	bool HasMarker2() const
+	{
+		return ((States & EIsoNodeStates::HasMarker2) == EIsoNodeStates::HasMarker2);
+	}
+
+	bool IsDeleteOrHasMarker2() const
+	{
+		constexpr EIsoNodeStates HasMarkerDeleteOrMarker2 = EIsoNodeStates::Delete | EIsoNodeStates::HasMarker2;
+		return ((States & HasMarkerDeleteOrMarker2) != EIsoNodeStates::None);
+	}
+
+	bool HasMarker1NotMarker2() const
+	{
+		constexpr EIsoNodeStates HasMarker12 = EIsoNodeStates::HasMarker1 | EIsoNodeStates::HasMarker2;
+		return ((States & HasMarker12) == EIsoNodeStates::HasMarker1);
+	}
+
+	void SetMarker1()
+	{
+		States |= EIsoNodeStates::HasMarker1;
+	}
+
+	void SetMarker2()
+	{
+		States |= EIsoNodeStates::HasMarker2;
+	}
+
+	void ResetMarkers()
+	{
+		constexpr EIsoNodeStates ResetAllMarkers = ~(EIsoNodeStates::HasMarker1 | EIsoNodeStates::HasMarker2);
+		States &= ResetAllMarkers;
+	}
+
 	/**
 	 * Return the 2d coordinate of the node according to the space
 	 */
diff --git a/Engine/Source/Runtime/Datasmith/CADKernel/Private/CADKernel/Mesh/Meshers/IsoTriangulator/IsoSegment.h b/Engine/Source/Runtime/Datasmith/CADKernel/Private/CADKernel/Mesh/Meshers/IsoTriangulator/IsoSegment.h
--- a/Engine/Source/Runtime/Datasmith/CADKernel/Private/CADKernel/Mesh/Meshers/IsoTriangulator/IsoSegment.h
+++ b/Engine/Source/Runtime/Datasmith/CADKernel/Private/CADKernel/Mesh/Meshers/IsoTriangulator/IsoSegment.h
@@ -86,22 +86,16 @@ public:
 		}
 	}
 
-	void ConnectToNode()
+	bool ConnectToNode()
 	{
-		TFunction<bool()> CheckExistingSegment = [&]()
+		if (IsItAlreadyDefined(FirstNode, SecondNode))
 		{
-			if (IsItAlreadyDefined(FirstNode, SecondNode))
-			{
-				//Wait();
-				return false;
-			}
-			return true;
-		};
-
-		ensureCADKernel(CheckExistingSegment());
+			return false;
+		}
 
 		FirstNode->ConnectSegment(*this);
 		SecondNode->ConnectSegment(*this);
+		return true;
 	}
 
 	static bool	IsItAlreadyDefined(const FIsoNode* StartNode, const FIsoNode* EndNode)
diff --git a/Engine/Source/Runtime/Datasmith/CADKernel/Private/CADKernel/Mesh/Meshers/ParametricFaceMesher.cpp b/Engine/Source/Runtime/Datasmith/CADKernel/Private/CADKernel/Mesh/Meshers/ParametricFaceMesher.cpp
--- a/Engine/Source/Runtime/Datasmith/CADKernel/Private/CADKernel/Mesh/Meshers/ParametricFaceMesher.cpp
+++ b/Engine/Source/Runtime/Datasmith/CADKernel/Private/CADKernel/Mesh/Meshers/ParametricFaceMesher.cpp
@@ -48,6 +48,8 @@ void FParametricFaceMesher::Mesh()
 		return;
 	}
 
+	MeshVerticesOfFace(Face);
+
 	FMessage::Printf(EVerboseLevel::Debug, TEXT("Meshing of surface %d\n"), Face.GetId());
 
 	FProgress _(1, TEXT("Meshing Entities : Mesh Surface"));
@@ -58,7 +60,6 @@ void FParametricFaceMesher::Mesh()
 
 	FTimePoint StartTime = FChrono::Now();
 
-
 	if (!GenerateCloud() || Grid.IsDegenerated())
 	{
 #ifdef CADKERNEL_DEV
@@ -145,6 +146,7 @@ void FParametricFaceMesher::MeshThinZones(TArray<FTopologicalEdge*>& EdgesToMesh
 			}
 
 #ifdef DEBUG_MESH_THINZONE_EDGES
+			if(bDisplay)
 			{
 				F3DDebugSession B(bDisplay, FString::Printf(TEXT("Edge %d"), Edge->GetId()));
 				Display2DWithScale(*Edge, EVisuProperty::GreenCurve);
@@ -434,7 +436,7 @@ void FParametricFaceMesher::DefineImposedCuttingPointsBasedOnOtherSideMesh(FThin
 
 	FAddMeshNodeFunc AddToCrossZoneElements = [&CrossZoneElements](const int32 NodeIndice, const FPoint2D& MeshNode2D, double MeshingTolerance3D, const FEdgeSegment& EdgeSegment, const FPairOfIndex& OppositeNodeIndices)
 	{
-		if (CrossZoneElements.Num() && CrossZoneElements.Last().VertexId == NodeIndice)
+		if (CrossZoneElements.Num() && CrossZoneElements.Last().VertexId >= 0 && CrossZoneElements.Last().VertexId == NodeIndice)
 		{
 			CrossZoneElements.Last().Add(OppositeNodeIndices);
 		}
@@ -496,12 +498,14 @@ void FParametricFaceMesher::DefineImposedCuttingPointsBasedOnOtherSideMesh(FThin
 				}
 			}
 		}
-		//Wait();
+		Wait(false);
 	}
 
 	F3DDebugSession A(bDisplay, TEXT("Find the best projection of existing mesh vertices"));
 #endif
 
+	const double MaxSquareThickness = FrontSide.GetMaxThickness() > SideToConstrain.GetMaxThickness() ? FMath::Square(3. * FrontSide.GetMaxThickness()) : FMath::Square(3. * SideToConstrain.GetMaxThickness());
+
 	// Find the best projection of existing mesh vertices (CrossZone Vertex)
 	for (FCrossZoneElement& CrossZoneElement : CrossZoneElements)
 	{
@@ -516,16 +520,16 @@ void FParametricFaceMesher::DefineImposedCuttingPointsBasedOnOtherSideMesh(FThin
 			continue;
 		}
 
-#ifdef DEBUG_MESHTHINSURF_
+#ifdef DEBUG_MESHTHINSURF
 		if (bDisplay)
 		{
 			F3DDebugSession A(TEXT("Point"));
 			DisplayPoint2DWithScale(CrossZoneElement.VertexPoint2D, EVisuProperty::BluePoint, CrossZoneElement.VertexId);
-			Wait();
+			Wait(false);
 		}
 #endif
 
-		double MinSquareThickness = HUGE_VALUE;
+		double MinSquareThickness = MaxSquareThickness;
 		FPoint2D ClosePoint;
 		FEdgeSegment* CloseSegment = nullptr;
 		double ClosePointCoordinate = -1;
@@ -654,7 +658,7 @@ void FParametricFaceMesher::DefineImposedCuttingPointsBasedOnOtherSideMesh(FThin
 				DisplaySegmentWithScale(CrossZoneElement.VertexPoint2D, CrossZoneElement.OppositePoint2D, 0, EVisuProperty::YellowCurve);
 			}
 		}
-		//Wait();
+		Wait(false);
 	}
 #endif
 
@@ -697,6 +701,18 @@ void FParametricFaceMesher::Mesh(FTopologicalVertex& InVertex)
 	InVertex.GetOrCreateMesh(MeshModel);
 }
 
+void FParametricFaceMesher::MeshVerticesOfFace(FTopologicalFace& FaceToProcess)
+{
+	for (const TSharedPtr<FTopologicalLoop>& Loop : FaceToProcess.GetLoops())
+	{
+		for (FOrientedEdge& Edge : Loop->GetEdges())
+		{
+			Mesh(*Edge.Entity->GetStartVertex());
+			Mesh(*Edge.Entity->GetEndVertex());
+		}
+	}
+}
+
 #ifdef DEBUG_INTERSECTEDGEISOS
 void DebugIntersectEdgeIsos(const FTopologicalFace& Face, const TArray<double>& IsoCoordinates, EIso TypeIso);
 #endif
@@ -907,8 +923,9 @@ void FParametricFaceMesher::Mesh(FTopologicalEdge& InEdge, bool bFinalMeshing)
 	}
 
 
-	const double EdgeDeltaUAtMin = DeltaUs[0] * AQuarter;
-	const double EdgeDeltaUAtMax = DeltaUs.Last() * AQuarter;
+	const double EdgeBoundsLength = EdgeBounds.Length();
+	const double EdgeDeltaUAtMin = FMath::Min(DeltaUs[0] * AQuarter, EdgeBoundsLength * AEighth);
+	const double EdgeDeltaUAtMax = FMath::Min(DeltaUs.Last() * AQuarter, EdgeBoundsLength * AEighth);
 
 	ImposedIsoCuttingPoints.Emplace(EdgeBounds.GetMin(), ECoordinateType::VertexCoordinate, FPairOfIndex::Undefined, EdgeDeltaUAtMin);
 	ImposedIsoCuttingPoints.Emplace(EdgeBounds.GetMax(), ECoordinateType::VertexCoordinate, FPairOfIndex::Undefined, EdgeDeltaUAtMax);
@@ -1155,9 +1172,12 @@ void FParametricFaceMesher::ApplyEdgeCriteria(FTopologicalEdge& Edge)
 {
 	FTopologicalEdge& ActiveEdge = *Edge.GetLinkActiveEdge();
 
-	if (Edge.Length() < Edge.GetTolerance3D())
+	if (Edge.Length() < 2. * Tolerances.GeometricTolerance)
 	{
-		Edge.SetAsDegenerated();
+		for (FTopologicalEdge* TwinEdge : Edge.GetTwinEntities())
+		{
+			TwinEdge->SetAsDegenerated();
+		}
 	}
 
 	Edge.ComputeCrossingPointCoordinates();
diff --git a/Engine/Source/Runtime/Datasmith/CADKernel/Private/CADKernel/Mesh/Meshers/ParametricFaceMesher.h b/Engine/Source/Runtime/Datasmith/CADKernel/Private/CADKernel/Mesh/Meshers/ParametricFaceMesher.h
--- a/Engine/Source/Runtime/Datasmith/CADKernel/Private/CADKernel/Mesh/Meshers/ParametricFaceMesher.h
+++ b/Engine/Source/Runtime/Datasmith/CADKernel/Private/CADKernel/Mesh/Meshers/ParametricFaceMesher.h
@@ -156,6 +156,7 @@ private:
 
 	void Mesh(FTopologicalEdge& InEdge, bool bFinalMeshing = true);
 	void Mesh(FTopologicalVertex& Vertex);
+	void MeshVerticesOfFace(FTopologicalFace& Face);
 
 	void MeshFaceLoops();
 
diff --git a/Engine/Source/Runtime/Datasmith/CADKernel/Private/CADKernel/Mesh/Meshers/ParametricMesher.cpp b/Engine/Source/Runtime/Datasmith/CADKernel/Private/CADKernel/Mesh/Meshers/ParametricMesher.cpp
--- a/Engine/Source/Runtime/Datasmith/CADKernel/Private/CADKernel/Mesh/Meshers/ParametricMesher.cpp
+++ b/Engine/Source/Runtime/Datasmith/CADKernel/Private/CADKernel/Mesh/Meshers/ParametricMesher.cpp
@@ -79,7 +79,7 @@ void FParametricMesher::MeshEntities(TArray<FTopologicalShapeEntity*>& InEntitie
 		Face->ResetMarkers();
 	}
 
-	// Get independent Faces and spread body's shells orientation
+	// Get independent Faces and propagate body's shells orientation
 	for (FTopologicalShapeEntity* TopologicalEntity : InEntities)
 	{
 		if (TopologicalEntity == nullptr)
diff --git a/Engine/Source/Runtime/Datasmith/CADKernel/Private/CADKernel/Mesh/Meshers/ParametricMesherConstantes.h b/Engine/Source/Runtime/Datasmith/CADKernel/Private/CADKernel/Mesh/Meshers/ParametricMesherConstantes.h
--- a/Engine/Source/Runtime/Datasmith/CADKernel/Private/CADKernel/Mesh/Meshers/ParametricMesherConstantes.h
+++ b/Engine/Source/Runtime/Datasmith/CADKernel/Private/CADKernel/Mesh/Meshers/ParametricMesherConstantes.h
@@ -17,6 +17,7 @@ class FMeshingTolerances
 {
 public:
 	const double GeometricTolerance;
+	const double GeometricToleranceSqrt2;
 	const double SquareGeometricTolerance;
 	const double SquareGeometricTolerance2;
 	const double MeshingTolerance;
@@ -24,6 +25,7 @@ public:
 
 	FMeshingTolerances(double InGeometricTolerance3D)
 		: GeometricTolerance(InGeometricTolerance3D)
+		, GeometricToleranceSqrt2(InGeometricTolerance3D * UE_DOUBLE_SQRT_2)
 		, SquareGeometricTolerance(FMath::Square(GeometricTolerance))
 		, SquareGeometricTolerance2(2. * SquareGeometricTolerance)
 		, MeshingTolerance(GeometricTolerance * MeshConst::GeometricToMeshingToleranceFactor)
diff --git a/Engine/Source/Runtime/Datasmith/CADKernel/Private/CADKernel/Mesh/Structure/Grid.cpp b/Engine/Source/Runtime/Datasmith/CADKernel/Private/CADKernel/Mesh/Structure/Grid.cpp
--- a/Engine/Source/Runtime/Datasmith/CADKernel/Private/CADKernel/Mesh/Structure/Grid.cpp
+++ b/Engine/Source/Runtime/Datasmith/CADKernel/Private/CADKernel/Mesh/Structure/Grid.cpp
@@ -11,7 +11,6 @@
 #include "CADKernel/Utils/Util.h"
 #include "CADKernel/Utils/ArrayUtils.h"
 
-//#define DEBUG_GETPREFERREDUVCOORDINATESFROMNEIGHBOURS
 namespace UE::CADKernel
 {
 
@@ -46,7 +45,7 @@ void FGrid::ProcessPointCloud()
 		DisplayGridLoops(TEXT("FGrid::Loop 2D Scaled"), GetLoops2D(EGridSpace::Scaled), true, false, false);
 		DisplayInnerDomainPoints(TEXT("FGrid::PointCloud 2D Scaled"), GetInner2DPoints(EGridSpace::Scaled));
 
-		DisplayGridLoops(TEXT("FGrid::Loop 2D UniformScaled"), GetLoops2D(EGridSpace::UniformScaled), true, false, false);
+		DisplayGridLoops(TEXT("FGrid::Loop 2D UniformScaled"), GetLoops2D(EGridSpace::UniformScaled), true, true, false);
 		DisplayInnerDomainPoints(TEXT("FGrid::PointCloud 2D UniformScaled"), GetInner2DPoints(EGridSpace::UniformScaled));
 
 		DisplayGridLoops(TEXT("FGrid::Loop 3D"), GetLoops3D(), true, false, false);
@@ -1173,20 +1172,16 @@ void FGrid::GetMeshOfThinZone(const FThinZone2D& ThinZone)
 		}
 	};
 
-	TArray<FPoint2D>* ThinZoneMesh = &FaceLoops2D[EGridSpace::UniformScaled].Emplace_GetRef();
+	const bool FirstSideIsClosed = ThinZone.GetFirstSide().IsClosed();
+	const bool SecondSideIsClosed = ThinZone.GetSecondSide().IsClosed();
 
+	TArray<FPoint2D>* ThinZoneMesh = &FaceLoops2D[EGridSpace::UniformScaled].Emplace_GetRef();
 	GetThinZoneSideMesh(ThinZone.GetFirstSide(), *ThinZoneMesh);
-	if (ThinZone.GetCategory() == EThinZone2DType::BetweenLoops)
+	if (FirstSideIsClosed && SecondSideIsClosed)
 	{
 		RemoveDuplicatedNode(*ThinZoneMesh);
-		if (ThinZoneMesh->Num() < 3)
-		{
-			FaceLoops2D[EGridSpace::UniformScaled].Pop();
-		}
-
 		ThinZoneMesh = &FaceLoops2D[EGridSpace::UniformScaled].Emplace_GetRef();
 	}
-
 	GetThinZoneSideMesh(ThinZone.GetSecondSide(), *ThinZoneMesh);
 	RemoveDuplicatedNode(*ThinZoneMesh);
 
diff --git a/Engine/Source/Runtime/Datasmith/CADKernel/Private/CADKernel/Mesh/Structure/LoopCleaner.cpp b/Engine/Source/Runtime/Datasmith/CADKernel/Private/CADKernel/Mesh/Structure/LoopCleaner.cpp
--- a/Engine/Source/Runtime/Datasmith/CADKernel/Private/CADKernel/Mesh/Structure/LoopCleaner.cpp
+++ b/Engine/Source/Runtime/Datasmith/CADKernel/Private/CADKernel/Mesh/Structure/LoopCleaner.cpp
@@ -40,7 +40,9 @@ bool FLoopCleaner::Run()
 		return false;
 	}
 
-	if (!UncrossLoops())
+	// Two step: 
+	// first intersection between inner loops and outer loop are processed 
+	if (!UncrossLoops(/*bAddProcessedLoop*/ false))
 	{
 #ifdef CADKERNEL_DEV
 		FMesherReport::GetLogs().AddRemoveCrossingLoopsFailure();
@@ -48,6 +50,24 @@ bool FLoopCleaner::Run()
 		return false;
 	}
 
+	// second, inner loop are added to LoopSegmentsIntersectionTool, so intersection between inner loop can be fixed
+	if (!UncrossLoops(/*bAddProcessedLoop*/ true))
+	{
+#ifdef CADKERNEL_DEV
+		FMesherReport::GetLogs().AddRemoveCrossingLoopsFailure();
+#endif
+		return false;
+	}
+
+	// a new call of CleanLoops because UncrossLoops could have generated auto intersection.
+	if (!CleanLoops())
+	{
+#ifdef CADKERNEL_DEV
+		FMesherReport::GetLogs().AddRemoveSelfIntersectionFailure();
+#endif
+		return false;
+	}
+
 	return true;
 }
 
@@ -75,6 +95,11 @@ bool FLoopCleaner::CleanLoops()
 
 	for (FLoopNode* StartNode : BestStartNodeOfLoops)
 	{
+		if (StartNode->IsDelete())
+		{
+			continue;
+		}
+
 		LoopIndex++;
 
 		StartSegmentIndex = NextLoopFirstSegmentIndex;
@@ -159,8 +184,7 @@ bool FLoopCleaner::CleanLoops()
 	return true;
 }
 
-//#define DEBUG_UNCROSS_LOOPS
-bool FLoopCleaner::UncrossLoops()
+bool FLoopCleaner::UncrossLoops(bool bAddProcessedLoop)
 {
 	if (Grid.GetLoopCount() == 1)
 	{
@@ -169,122 +193,163 @@ bool FLoopCleaner::UncrossLoops()
 
 #ifdef DEBUG_UNCROSS_LOOPS
 	int32 Iteration = 0;
+	int32 MaxIndex = 10000;
 	F3DDebugSession _(bDisplay, TEXT("FixIntersectionBetweenLoops"));
 #endif
 
-	TSet<uint32> IntersectionAlreadyProcessed;
+	TArray<int32> SegmentAlreadyProcessed;
 
 	LoopSegmentsIntersectionTool.Empty(LoopSegments.Num());
-	int32 Index = 0;
-	for (; Index < LoopSegments.Num(); ++Index)
+
+	TArray<TArray<FIsoSegment*>> SegmentsOfLoops;
 	{
-		if (((FLoopNode&)LoopSegments[Index]->GetFirstNode()).GetLoopIndex() != 0)
+		int32 LoopCount = ((FLoopNode&)LoopSegments.Last()->GetFirstNode()).GetLoopIndex();
+		SegmentsOfLoops.Reserve(LoopCount + 1);
+	}
+
+	{
+		TArray<FIsoSegment*>* SegmentsOfLoop = &SegmentsOfLoops.Emplace_GetRef();
+		int32 IndexOfLoop = 0;
+		for (FIsoSegment* Segment : LoopSegments)
 		{
-			break;
+			if (IndexOfLoop != ((FLoopNode&)Segment->GetFirstNode()).GetLoopIndex())
+			{
+				SegmentsOfLoop = &SegmentsOfLoops.Emplace_GetRef();
+				IndexOfLoop = ((FLoopNode&)Segment->GetFirstNode()).GetLoopIndex();
+			}
+			SegmentsOfLoop->Add(Segment);
 		}
-		LoopSegmentsIntersectionTool.AddSegment(*LoopSegments[Index]);
+	}
+	for (FIsoSegment* Segment : SegmentsOfLoops[0])
+	{
+		LoopSegmentsIntersectionTool.AddSegment(*Segment);
 	}
 
-	for (; Index < LoopSegments.Num(); )
+	for (int32 Index = 1; Index < SegmentsOfLoops.Num(); ++Index)
 	{
-		ensureCADKernel(LoopSegments[Index]);
+		TArray<FIsoSegment*>& SegmentsOfLoop = SegmentsOfLoops[Index];
 
-		FIsoSegment& Segment = *LoopSegments[Index];
-		ensureCADKernel(!Segment.IsDelete());
+		SegmentAlreadyProcessed.Init(0, SegmentsOfLoop.Num());
 
 #ifdef DEBUG_UNCROSS_LOOPS
 		if (bDisplay)
 		{
-			LoopSegmentsIntersectionTool.Display(bDisplay, TEXT("IntersectionTool"), EVisuProperty::BlueCurve);
-			F3DDebugSession _(*FString::Printf(TEXT("Segment to be processed %d %d"), Index, Iteration++));
-			Grid.DisplayIsoSegment(EGridSpace::UniformScaled, Segment, 0, EVisuProperty::BlueCurve);
+			F3DDebugSession B(*FString::Printf(TEXT("Before Move Inside")));
+			Grid.DisplayIsoSegments(EGridSpace::UniformScaled, SegmentsOfLoop, true, false, false, EVisuProperty::BlueCurve);
+			Wait();
 		}
 #endif
 
-		if (const FIsoSegment* IntersectingSegment = LoopSegmentsIntersectionTool.FindIntersectingSegment(Segment))
-		{
 #ifdef DEBUG_UNCROSS_LOOPS
-			if (bDisplay)
+		F3DDebugSession _(*FString::Printf(TEXT("Loop %d"), Index));
+#endif
+		int32 LoopSize = SegmentsOfLoop.Num();
+		for (int32 SegmentIndex = 0; SegmentIndex < LoopSize; ++SegmentIndex)
+		{
+
+			FIsoSegment& Segment = *SegmentsOfLoop[SegmentIndex];
+			if (Segment.IsDelete())
 			{
+				continue;
+			}
+
+			if (const FIsoSegment* IntersectingSegment = LoopSegmentsIntersectionTool.FindIntersectingSegment(Segment))
+			{
+				SegmentAlreadyProcessed[SegmentIndex]++;
+				if (SegmentAlreadyProcessed[SegmentIndex] > 5)
 				{
-					LoopSegmentsIntersectionTool.Display(bDisplay, TEXT("IntersectionTool"), EVisuProperty::BlueCurve);
-					F3DDebugSession _(*FString::Printf(TEXT("Segment to be processed %d %d"), Index, Iteration++));
-					Grid.DisplayIsoSegment(EGridSpace::UniformScaled, Segment, 0, EVisuProperty::BlueCurve);
+					// the segment is processed more than 5 times... infinite loop... the loop process is canceled
+					break;
 				}
+
+#ifdef DEBUG_UNCROSS_LOOPS2
+				++Iteration;
+				if (bDisplay && Index > MaxIndex)
 				{
-					F3DDebugSession _(TEXT("Intersecting Segments"));
-					Grid.DisplayIsoSegment(EGridSpace::UniformScaled, *IntersectingSegment, 0, EVisuProperty::RedCurve/*, true*/);
+					F3DDebugSession _(*FString::Printf(TEXT("Intersection %d"), Iteration));
+					LoopSegmentsIntersectionTool.Display(bDisplay, TEXT("IntersectionTool"), EVisuProperty::BlueCurve);
+					{
+						F3DDebugSession B(*FString::Printf(TEXT("Loop")));
+						Grid.DisplayIsoSegments(EGridSpace::UniformScaled, SegmentsOfLoop, true, false, false, EVisuProperty::BlueCurve);
+					}
+					{
+						F3DDebugSession C(*FString::Printf(TEXT("Segment to be processed %d %d"), Index, Iteration));
+						Grid.DisplayIsoSegment(EGridSpace::UniformScaled, Segment, 0, EVisuProperty::YellowCurve);
+					}
+					{
+						F3DDebugSession C(TEXT("Intersecting Segments"));
+						Grid.DisplayIsoSegment(EGridSpace::UniformScaled, *IntersectingSegment, 0, EVisuProperty::RedCurve/*, true*/);
+					}
 				}
-				Wait(bDisplay);
-			}
 #endif
 
-			uint32 IntersectionHash = GetTypeHash(*IntersectingSegment, Segment);
-			bool bNotProcessed = !IntersectionAlreadyProcessed.Find(IntersectionHash);
-			IntersectionAlreadyProcessed.Add(IntersectionHash);
-
-			bool bIsFixed = true;
-			bool bIsSameLoop = ((FLoopNode&)Segment.GetFirstNode()).GetLoopIndex() && ((FLoopNode&)Segment.GetFirstNode()).GetLoopIndex() == ((FLoopNode&)IntersectingSegment->GetFirstNode()).GetLoopIndex();
-
-			// Swapping segments is possible only with outer loop
-			if (bNotProcessed)
-			{
 				if (!TryToRemoveIntersectionOfTwoConsecutiveIntersectingSegments(*IntersectingSegment, Segment))
 				{
 					if (!TryToRemoveIntersectionOfTwoConsecutiveIntersectingSegments(Segment, const_cast<FIsoSegment&>(*IntersectingSegment)))
 					{
-						RemoveIntersectionByMovingOutsideSegmentNodeInside(*IntersectingSegment, Segment, ((FLoopNode&)Segment.GetFirstNode()).GetLoopIndex() && bIsSameLoop);
+						RemoveIntersectionByMovingOutsideSegmentNodeInside(*IntersectingSegment, Segment, /*((FLoopNode&)Segment.GetFirstNode()).GetLoopIndex() && */ false);
 					}
 				}
 
-				// segment is moved, the previous segment is retested. Thanks to that, only one loop is enough.
-				if (Index > 1)
+#ifdef DEBUG_UNCROSS_LOOPS2
+				if (bDisplay && Index > MaxIndex)
 				{
-					LoopSegmentsIntersectionTool.RemoveLast();
+					F3DDebugSession B(*FString::Printf(TEXT("After fix")));
+					Grid.DisplayIsoSegments(EGridSpace::UniformScaled, SegmentsOfLoop, true, false, false, EVisuProperty::YellowCurve);
+					//Grid.DisplayGridPolyline(TEXT("After fix"), EGridSpace::UniformScaled, LoopNodes, true, EVisuProperty::YellowCurve);
+					//Wait();
 				}
+#endif
+				if (SegmentIndex > 0)
+				{
+					SegmentIndex--;
+				}
+				while (SegmentIndex > 0 && SegmentsOfLoop[SegmentIndex]->IsDelete())
+				{
+					SegmentIndex--;
+				}
+				SegmentIndex--;
 			}
-			else if (bIsSameLoop)
-			{
-				SwapSegments(const_cast<FIsoSegment&>(*IntersectingSegment), Segment);
-			}
-			else
-			{
-				return false;
-			}
+		}
 
 #ifdef DEBUG_UNCROSS_LOOPS
-			if (bDisplay)
-			{
-				Grid.DisplayGridPolyline(TEXT("After fix"), EGridSpace::UniformScaled, LoopNodes, true, EVisuProperty::YellowCurve);
-				Wait(true);
-			}
-#endif
+		if(bDisplay)
+		{
+			F3DDebugSession B(*FString::Printf(TEXT("Before RemoveLoopPicks")));
+			Grid.DisplayIsoSegments(EGridSpace::UniformScaled, SegmentsOfLoop, true, false, false, EVisuProperty::BlueCurve);
+			Wait();
 		}
-		else
+#endif
+
+		RemoveLoopPicks(SegmentsOfLoop);
+
+#ifdef DEBUG_UNCROSS_LOOPS
+		if (bDisplay)
 		{
-			LoopSegmentsIntersectionTool.AddSegment(Segment);
+			F3DDebugSession B(*FString::Printf(TEXT("After RemoveLoopPicks")));
+			Grid.DisplayIsoSegments(EGridSpace::UniformScaled, SegmentsOfLoop, true, false, false, EVisuProperty::BlueCurve);
+			Wait();
 		}
+#endif
 
-		if (!Segment.IsDelete())
+		if(bAddProcessedLoop)
 		{
-			RemovePickRecursively((FLoopNode*) &Segment.GetFirstNode(), (FLoopNode*) &Segment.GetSecondNode());
-#ifdef DEBUG_UNCROSS_LOOPS
-			if (bDisplay)
+			for (FIsoSegment* Segment : SegmentsOfLoop)
 			{
-				Grid.DisplayGridPolyline(TEXT("After RemovePickRecursively"), EGridSpace::UniformScaled, LoopNodes, true, EVisuProperty::YellowCurve);
-				Wait(true);
+				if (Segment->IsDelete())
+				{
+					continue;
+				}
+				LoopSegmentsIntersectionTool.AddSegment(*Segment);
 			}
-#endif
 		}
-
-		Index = LoopSegmentsIntersectionTool.Count();
 	}
 
 #ifdef DEBUG_UNCROSS_LOOPS
-	if (bDisplay)
+	if (bDisplay && Iteration > MaxIndex)
 	{
 		Grid.DisplayGridPolyline(TEXT("After UncrossLoops"), EGridSpace::UniformScaled, LoopNodes, true, EVisuProperty::YellowCurve);
-		Wait(true);
+		Wait();
 	}
 #endif
 
@@ -297,8 +362,8 @@ bool FLoopCleaner::TryToRemoveIntersectionOfTwoConsecutiveIntersectingSegments(c
 	FLoopNode* PreviousNode = nullptr;
 	FLoopNode* NextNode = nullptr;
 
-	TSegment<FPoint2D> IntersectingSegment2D(IntersectingSegment.GetFirstNode().Get2DPoint(EGridSpace::UniformScaled, Grid), IntersectingSegment.GetSecondNode().Get2DPoint(EGridSpace::UniformScaled, Grid));
-	TSegment<FPoint2D> Segment2D(Segment.GetFirstNode().Get2DPoint(EGridSpace::UniformScaled, Grid), Segment.GetSecondNode().Get2DPoint(EGridSpace::UniformScaled, Grid));
+	FSegment2D IntersectingSegment2D(IntersectingSegment.GetFirstNode().Get2DPoint(EGridSpace::UniformScaled, Grid), IntersectingSegment.GetSecondNode().Get2DPoint(EGridSpace::UniformScaled, Grid));
+	FSegment2D Segment2D(Segment.GetFirstNode().Get2DPoint(EGridSpace::UniformScaled, Grid), Segment.GetSecondNode().Get2DPoint(EGridSpace::UniformScaled, Grid));
 
 	double IntersectingSegmentSlope = ComputeOrientedSlope(IntersectingSegment2D.Point0, IntersectingSegment2D.Point1, 0);
 	double SegmentSlope = ComputeUnorientedSlope(Segment2D.Point1, Segment2D.Point0, IntersectingSegmentSlope);
@@ -312,7 +377,7 @@ bool FLoopCleaner::TryToRemoveIntersectionOfTwoConsecutiveIntersectingSegments(c
 	{
 		double StartPointSquareDistance = SquareDistanceOfPointToSegment(Segment2D.Point0, IntersectingSegment2D.Point0, IntersectingSegment2D.Point1);
 		double EndPointSquareDistance = SquareDistanceOfPointToSegment(Segment2D.Point1, IntersectingSegment2D.Point0, IntersectingSegment2D.Point1);
-		if (StartPointSquareDistance < Tolerances.SquareGeometricTolerance && EndPointSquareDistance < Tolerances.SquareGeometricTolerance)
+		if (StartPointSquareDistance < Tolerances.SquareGeometricTolerance2 && EndPointSquareDistance < Tolerances.SquareGeometricTolerance2)
 		{
 			OffsetSegment(Segment, Segment2D, IntersectingSegment2D);
 			return true;
@@ -323,13 +388,13 @@ bool FLoopCleaner::TryToRemoveIntersectionOfTwoConsecutiveIntersectingSegments(c
 	{
 		double Coordinate = 0;
 		FindIntersectionOfSegments2D(Segment2D, IntersectingSegment2D, Coordinate);
-		if (FMath::IsNearlyZero(Coordinate))
+		if (FMath::IsNearlyZero(Coordinate, DOUBLE_KINDA_SMALL_NUMBER))
 		{
 			// can add a test to offset the outside node and not the node a 
 			OffsetNode((FLoopNode&)Segment.GetFirstNode(), IntersectingSegment2D);
 			return true;
 		}
-		else if (FMath::IsNearlyEqual(Coordinate, 1))
+		else if (FMath::IsNearlyEqual(Coordinate, 1, DOUBLE_KINDA_SMALL_NUMBER))
 		{
 			OffsetNode((FLoopNode&)Segment.GetSecondNode(), IntersectingSegment2D);
 			return true;
@@ -362,13 +427,13 @@ bool FLoopCleaner::TryToRemoveIntersectionOfTwoConsecutiveIntersectingSegments(c
 	}
 #endif
 
-	TSegment<FPoint2D> NextSegment2D(Node->Get2DPoint(EGridSpace::UniformScaled, Grid), NextNode->Get2DPoint(EGridSpace::UniformScaled, Grid));
-	if (!FastDoIntersect(NextSegment2D, IntersectingSegment2D))
+	FSegment2D NextSegment2D(Node->Get2DPoint(EGridSpace::UniformScaled, Grid), NextNode->Get2DPoint(EGridSpace::UniformScaled, Grid));
+	if (!DoIntersect(NextSegment2D, IntersectingSegment2D))
 	{
 		return false;
 	}
 
-	TSegment<FPoint2D>  PreviousSegment2D(Node->Get2DPoint(EGridSpace::UniformScaled, Grid), PreviousNode->Get2DPoint(EGridSpace::UniformScaled, Grid));
+	FSegment2D  PreviousSegment2D(Node->Get2DPoint(EGridSpace::UniformScaled, Grid), PreviousNode->Get2DPoint(EGridSpace::UniformScaled, Grid));
 
 	FPoint2D Intersection1 = FindIntersectionOfSegments2D(PreviousSegment2D, IntersectingSegment2D);
 	FPoint2D Intersection2 = FindIntersectionOfSegments2D(NextSegment2D, IntersectingSegment2D);
@@ -432,7 +497,7 @@ bool FLoopCleaner::TryToRemoveIntersectionOfTwoConsecutiveIntersectingSegments(c
 	return true;
 }
 
-void FLoopCleaner::OffsetSegment(FIsoSegment& Segment, TSegment<FPoint2D>& Segment2D, TSegment<FPoint2D>& IntersectingSegment2D)
+void FLoopCleaner::OffsetSegment(FIsoSegment& Segment, FSegment2D& Segment2D, FSegment2D& IntersectingSegment2D)
 {
 	FPoint2D SegmentTangent = IntersectingSegment2D.Point1 - IntersectingSegment2D.Point0;
 	SegmentTangent.Normalize();
@@ -446,7 +511,7 @@ void FLoopCleaner::OffsetSegment(FIsoSegment& Segment, TSegment<FPoint2D>& Segme
 	Segment.GetSecondNode().Set2DPoint(EGridSpace::UniformScaled, Grid, NewPoint1);
 }
 
-void FLoopCleaner::OffsetNode(FLoopNode& Node, TSegment<FPoint2D>& IntersectingSegment2D)
+void FLoopCleaner::OffsetNode(FLoopNode& Node, FSegment2D& IntersectingSegment2D)
 {
 	FPoint2D SegmentTangent = IntersectingSegment2D.Point1 - IntersectingSegment2D.Point0;
 	SegmentTangent.Normalize();
@@ -1503,7 +1568,6 @@ bool FLoopCleaner::TryToSwapSegmentsOrRemoveLoop(const TPair<double, double>& In
 	return RemoveSubLoop(Segment0_Node0, Segment1_Node1);
 }
 
-//#define DEBUG_BY_SWAPPING_SEGMENTS
 void FLoopCleaner::SwapSegments(FIsoSegment& IntersectingSegment, FIsoSegment& Segment)
 {
 	int32 IntersectingSegmentIndex = LoopSegments.IndexOfByKey(&IntersectingSegment);
@@ -1523,8 +1587,16 @@ void FLoopCleaner::SwapSegments(FIsoSegment& IntersectingSegment, FIsoSegment& S
 	if (bDisplay)
 	{
 		F3DDebugSession _(TEXT("New Segments"));
-		Grid.DisplayIsoSegment(EGridSpace::UniformScaled, Segment, StartSegmentIndex, EVisuProperty::RedCurve/*, true*/);
-		Grid.DisplayIsoSegment(EGridSpace::UniformScaled, IntersectingSegment, EndSegmentIndex, EVisuProperty::RedCurve/*, true*/);
+		Grid.DisplayIsoSegment(EGridSpace::UniformScaled, Segment, SegmentIndex, EVisuProperty::RedCurve/*, true*/);
+		Grid.DisplayIsoSegment(EGridSpace::UniformScaled, IntersectingSegment, IntersectingSegmentIndex, EVisuProperty::RedCurve/*, true*/);
+		Wait();
+	}
+#endif
+
+#ifdef DEBUG_BY_SWAPPING_SEGMENTS
+	if (bDisplay)
+	{
+		Grid.DisplayIsoSegments(TEXT("Before orientation"), EGridSpace::UniformScaled, LoopSegments, false, true);
 		Wait();
 	}
 #endif
@@ -1532,8 +1604,11 @@ void FLoopCleaner::SwapSegments(FIsoSegment& IntersectingSegment, FIsoSegment& S
 	SwapSubLoopOrientation(IntersectingSegmentIndex + 1, SegmentIndex);
 
 #ifdef DEBUG_BY_SWAPPING_SEGMENTS
-	Grid.DisplayIsoSegments(TEXT("After orientation"), EGridSpace::UniformScaled, LoopSegments, false, true);
-	Wait();
+	if (bDisplay)
+	{
+		Grid.DisplayIsoSegments(TEXT("After orientation"), EGridSpace::UniformScaled, LoopSegments, false, true);
+		Wait();
+	}
 #endif
 
 	LoopSegmentsIntersectionTool.Update(&IntersectingSegment);
@@ -1687,14 +1762,21 @@ void FLoopCleaner::FindLoopIntersections()
 					TmpNode = GetNext(TmpNode);
 				}
 
-				TSegment<FPoint2D> SegmentPoints(Node->Get2DPoint(EGridSpace::UniformScaled, Grid), NextNode->Get2DPoint(EGridSpace::UniformScaled, Grid));
-				TSegment<FPoint2D> IntersectedSegmentPoints(GetFirst(IntersectedSegment)->Get2DPoint(EGridSpace::UniformScaled, Grid), GetSecond(IntersectedSegment)->Get2DPoint(EGridSpace::UniformScaled, Grid));
+				FSegment2D SegmentPoints(Node->Get2DPoint(EGridSpace::UniformScaled, Grid), NextNode->Get2DPoint(EGridSpace::UniformScaled, Grid));
+				FSegment2D IntersectedSegmentPoints(GetFirst(IntersectedSegment)->Get2DPoint(EGridSpace::UniformScaled, Grid), GetSecond(IntersectedSegment)->Get2DPoint(EGridSpace::UniformScaled, Grid));
 
 				double SegmentIntersectionCoordinate;
 				FPoint2D IntersectionPoint = FindIntersectionOfSegments2D(SegmentPoints, IntersectedSegmentPoints, SegmentIntersectionCoordinate);
 
-				double SquareLength = IntersectedSegmentPoints.SquareLength();
-				double IntersectedSegmentIntersectionCoordinate = (SquareLength > SMALL_NUMBER_SQUARE) ? FMath::Sqrt(IntersectedSegmentPoints[0].SquareDistance(IntersectionPoint) / IntersectedSegmentPoints.SquareLength()) : 0;
+				const double IntersectedSegmentPointsSquareLength = IntersectedSegmentPoints.SquareLength();
+				const double SquareDistanceToIntersectionPoint = IntersectedSegmentPoints[0].SquareDistance(IntersectionPoint);
+				if (SquareDistanceToIntersectionPoint > IntersectedSegmentPointsSquareLength)
+				{
+					// the IntersectionPoint is not inside IntersectedSegmentPoints
+					continue;
+				}
+
+				double IntersectedSegmentIntersectionCoordinate = (IntersectedSegmentPointsSquareLength > DOUBLE_KINDA_SMALL_NUMBER_SQUARE) ? FMath::Sqrt(SquareDistanceToIntersectionPoint / IntersectedSegmentPointsSquareLength) : 0.;
 				IntersectedSegmentIntersectionCoordinate += IntersectionIndex;
 
 				SegmentIntersectionCoordinate += SegmentIndex;
@@ -1732,10 +1814,23 @@ void FLoopCleaner::FindLoopIntersections()
 				if (bDisplay)
 				{
 					F3DDebugSession _(*FString::Printf(TEXT("Intersection %f %f"), IntersectedSegmentIntersectionCoordinate, SegmentIntersectionCoordinate));
-					Grid.DisplayIsoSegment(EGridSpace::UniformScaled, *Node, *NextNode, 0, EVisuProperty::BlueCurve);
-					Grid.DisplayIsoSegment(EGridSpace::UniformScaled, *IntersectedSegment, 0, EVisuProperty::BlueCurve);
-					DisplayPoint(IntersectionPoint, EVisuProperty::RedPoint);
-					//Wait();
+					{
+						F3DDebugSession A(TEXT("Seg 1"));
+						Grid.DisplayIsoSegment(EGridSpace::UniformScaled, *Node, *NextNode, 0, EVisuProperty::BlueCurve);
+						Grid.DisplayIsoNode(EGridSpace::UniformScaled, *Node, 0, EVisuProperty::BluePoint);
+						Grid.DisplayIsoNode(EGridSpace::UniformScaled, *NextNode, 0, EVisuProperty::BluePoint);
+					}
+					{
+						F3DDebugSession A(TEXT("IntersectedSegment"));
+						Grid.DisplayIsoSegment(EGridSpace::UniformScaled, *IntersectedSegment, 0, EVisuProperty::RedCurve);
+						Grid.DisplayIsoNode(EGridSpace::UniformScaled, IntersectedSegment->GetFirstNode(), 0, EVisuProperty::RedPoint);
+						Grid.DisplayIsoNode(EGridSpace::UniformScaled, IntersectedSegment->GetSecondNode(), 0, EVisuProperty::RedPoint);
+					}
+					{
+						F3DDebugSession A(TEXT("IntersectionPoint"));
+						DisplayPoint(IntersectionPoint * DisplayScale, EVisuProperty::YellowPoint);
+					}
+					Wait();
 				}
 #endif
 			}
@@ -1840,6 +1935,23 @@ bool FLoopCleaner::RemoveLoopPicks()
 	return UpdateNodesOfLoop();
 }
 
+bool FLoopCleaner::RemoveLoopPicks(TArray<FIsoSegment*>& Loop)
+{
+	for (FIsoSegment* Segment : Loop)
+	{
+		if (Segment == nullptr || Segment->IsDelete())
+		{
+			continue;
+		}
+
+		if (!RemovePickRecursively((FLoopNode*)&Segment->GetFirstNode(), (FLoopNode*)&Segment->GetSecondNode()))
+		{
+			return false;
+		}
+	}
+	return UpdateNodesOfLoop();
+}
+
 bool FLoopCleaner::FindAndRemoveCoincidence(FLoopNode*& StartNode)
 {
 #ifdef DEBUG_FIND_AND_REMOVE_COINCIDENCE
@@ -2084,6 +2196,10 @@ void FLoopCleaner::FindBestLoopExtremity()
 
 		for (FLoopNode* Node : ExtremityNodes)
 		{
+			if (Node->IsDelete())
+			{
+				continue;
+			}
 			CompareWithSlopeAt(Node);
 		}
 		BestStartNodeOfLoops.Add(BestNode);
@@ -2105,6 +2221,11 @@ void FLoopCleaner::FindBestLoopExtremity()
 	int32 Index = 0;
 	for (FLoopNode& Node : LoopNodes)
 	{
+		if (Node.IsDelete())
+		{
+			continue;
+		}
+
 		if (Node.GetLoopIndex() != LoopIndex)
 		{
 			FindLoopExtremity();
@@ -2134,17 +2255,6 @@ void FLoopCleaner::FindBestLoopExtremity()
 			VMin = Point.V;
 			ExtremityNodes[3] = &Node;
 		}
-
-		//if (Node.GetIndex() == 44)
-		//{
-		//	UMax = VMax = 10000;
-		//	UMin = VMin = -10000;
-
-		//	ExtremityNodes[0] = \
-		//	ExtremityNodes[1] = \
-		//	ExtremityNodes[2] = \
-		//	ExtremityNodes[3] = &Node;
-		//}
 	}
 	FindLoopExtremity();
 }
diff --git a/Engine/Source/Runtime/Datasmith/CADKernel/Private/CADKernel/Mesh/Structure/LoopCleaner.h b/Engine/Source/Runtime/Datasmith/CADKernel/Private/CADKernel/Mesh/Structure/LoopCleaner.h
--- a/Engine/Source/Runtime/Datasmith/CADKernel/Private/CADKernel/Mesh/Structure/LoopCleaner.h
+++ b/Engine/Source/Runtime/Datasmith/CADKernel/Private/CADKernel/Mesh/Structure/LoopCleaner.h
@@ -150,7 +150,7 @@ public:
 private:
 
 	bool CleanLoops();
-	bool UncrossLoops();
+	bool UncrossLoops(bool bAddProcessedLoop);
 
 	/**
 	 * For each loop, find the best starting node i.e. a node well oriented
@@ -166,6 +166,8 @@ private:
 	 * @return false if the process failed => the surface cannot be meshed
 	 */
 	bool RemoveLoopPicks();
+	bool RemoveLoopPicks(TArray<FIsoSegment*>& Loop);
+
 	bool RemovePickRecursively(FLoopNode* Node0, FLoopNode* Node1);
 	bool FindAndRemoveCoincidence(FLoopNode*& StartNode);
 
@@ -266,8 +268,8 @@ private:
 	//bool TryToRemoveSelfIntersectionByMovingTheClosedOusidePoint(const FIsoSegment& Segment0, const FIsoSegment& Segment1);
 	//bool TryToRemoveIntersectionByMovingTheClosedOusidePoint(const FIsoSegment& Segment0, const FIsoSegment& Segment1);
 
-	void OffsetSegment(FIsoSegment& Segment, TSegment<FPoint2D>& Segment2D, TSegment<FPoint2D>& IntersectingSegment2D);
-	void OffsetNode(FLoopNode& Node, TSegment<FPoint2D>& IntersectingSegment2D);
+	void OffsetSegment(FIsoSegment& Segment, FSegment2D& Segment2D, FSegment2D& IntersectingSegment2D);
+	void OffsetNode(FLoopNode& Node, FSegment2D& IntersectingSegment2D);
 
 
 	//
diff --git a/Engine/Source/Runtime/Datasmith/CADKernel/Private/CADKernel/Mesh/Structure/ThinZone2DFinder.cpp b/Engine/Source/Runtime/Datasmith/CADKernel/Private/CADKernel/Mesh/Structure/ThinZone2DFinder.cpp
--- a/Engine/Source/Runtime/Datasmith/CADKernel/Private/CADKernel/Mesh/Structure/ThinZone2DFinder.cpp
+++ b/Engine/Source/Runtime/Datasmith/CADKernel/Private/CADKernel/Mesh/Structure/ThinZone2DFinder.cpp
@@ -404,7 +404,7 @@ void FThinZone2DFinder::CheckIfCloseSideOfThinSidesAreNotDegenerated()
 		if (!CheckIfCloseSideOfThinSideIsNotDegenerated(Side))
 		{
 #ifdef DEBUG_CHECK_THIN_SIDE
-			if(bDisplay)
+			if (bDisplay)
 			{
 				F3DDebugSession _(TEXT("Delete Thin Zone"));
 				ThinZone::DisplayThinZoneSide(Side, 0, EVisuProperty::RedCurve);
@@ -415,7 +415,7 @@ void FThinZone2DFinder::CheckIfCloseSideOfThinSidesAreNotDegenerated()
 				Segment->ResetCloseData();
 				Segment->SetChainIndex(Ident::Undefined);
 			}
- 			Side.Empty();
+			Side.Empty();
 		}
 	}
 #ifdef DEBUG_CHECK_THIN_SIDE
@@ -603,7 +603,11 @@ bool FThinZone2DFinder::CheckIfCloseSideOfThinSideIsNotDegenerated(TArray<FEdgeS
 void FThinZone2DFinder::ImproveThinSide()
 {
 #ifdef DEBUG_THIN_ZONES_IMPROVE
-	F3DDebugSession _(bDisplay, ("Improve Thin Zone"));
+	if (bDisplay)
+	{
+		F3DDebugSession _(bDisplay, ("Improve Thin Zone"));
+		Wait(false);
+	}
 #endif
 
 	TFunction<TArray<FEdgeSegment*>(const FEdgeSegment*)> GetComplementary = [this](const FEdgeSegment* EdgeSegment) ->TArray<FEdgeSegment*>
@@ -951,8 +955,11 @@ void FThinZone2DFinder::ImproveThinSide()
 
 
 #ifdef DEBUG_THIN_ZONES_IMPROVE_
-	ThinZone::DisplayThinZoneSides(ThinZoneSides);
-	Wait();
+	if (bDisplay)
+	{
+		ThinZone::DisplayThinZoneSides(ThinZoneSides);
+		Wait(false);
+	}
 #endif
 }
 
@@ -969,8 +976,8 @@ void FThinZone2DFinder::SplitThinSide()
 		}
 
 		FIdent CloseSideIndex = ThinZoneSide[0]->GetCloseSegment()->GetChainIndex();
-		if (!Algo::AllOf(ThinZoneSide, [&](const FEdgeSegment* EdgeSegment) 
-			{ 
+		if (!Algo::AllOf(ThinZoneSide, [&CloseSideIndex](const FEdgeSegment* EdgeSegment)
+			{
 				if (!EdgeSegment->GetCloseSegment())
 				{
 					return false;
@@ -979,6 +986,11 @@ void FThinZone2DFinder::SplitThinSide()
 			}))
 		{
 			int32 Index = 0;
+
+			const FEdgeSegment* FirstSegment = ThinZoneSide[0];
+			const FEdgeSegment* LastSegment = ThinZoneSide.Last();
+			const bool bIsCloseZone = (LastSegment->GetNext() == FirstSegment);
+
 			for (; Index < ThinZoneSide.Num(); ++Index)
 			{
 				if (!ThinZoneSide[Index]->GetCloseSegment() || ThinZoneSide[Index]->GetCloseSegment()->GetChainIndex() != CloseSideIndex)
@@ -1009,7 +1021,18 @@ void FThinZone2DFinder::SplitThinSide()
 				}
 				++NewSideIndex;
 			}
+
 			ThinZoneSide.SetNum(LastIndexFirstSide);
+			if (bIsCloseZone && (FirstSegment->GetCloseSegment()->GetChainIndex() == LastSegment->GetCloseSegment()->GetChainIndex()))
+			{
+				--NewSideIndex;
+				TArray<FEdgeSegment*>& NewThinZoneSide = NewThinZoneSides.Last();
+				for (FEdgeSegment* Segment : ThinZoneSide)
+				{
+					Segment->SetChainIndex(NewSideIndex);
+				}
+				NewThinZoneSide += MoveTemp(ThinZoneSide);
+			}
 		}
 	}
 
@@ -1021,6 +1044,27 @@ void FThinZone2DFinder::SplitThinSide()
 
 }
 
+namespace ThinZone2DFinderTool
+{
+void AlignStartEndPoints(TArray<FEdgeSegment*>& ClosedSide, TArray<FEdgeSegment*>& OpenSide)
+{
+	FEdgeSegment* OpenSideFirst = OpenSide[0];
+	FEdgeSegment* OpenSideLast = OpenSide.Last();
+
+	TArray<FEdgeSegment*> NewClosedSide;
+	NewClosedSide.Reserve(ClosedSide.Num());
+
+	FEdgeSegment* CloseSideSegment = OpenSideLast->GetCloseSegment();
+	FEdgeSegment* CloseSideEnd = OpenSideFirst->GetCloseSegment();
+	while (CloseSideSegment != CloseSideEnd)
+	{
+		NewClosedSide.Add(CloseSideSegment);
+		CloseSideSegment = CloseSideSegment->GetNext();
+	}
+	ClosedSide = MoveTemp(NewClosedSide);
+}
+}
+
 void FThinZone2DFinder::BuildThinZone()
 {
 	// The number of ThinZone should be less than the number of ThinZoneSides, 
@@ -1046,6 +1090,26 @@ void FThinZone2DFinder::BuildThinZone()
 			continue;
 		}
 
+		if (SecondSide[0]->GetCloseSegment()->GetChainIndex() != FirstSideSegment->GetChainIndex())
+		{
+			continue;
+		}
+
+		const bool FirstSideIsClosed = (FirstSide[0]->GetPrevious() == FirstSide.Last());
+		const bool SecondSideIsClosed = (SecondSide[0]->GetPrevious() == SecondSide.Last());
+
+		if (FirstSideIsClosed ^ SecondSideIsClosed)
+		{
+			if (!SecondSideIsClosed)
+			{
+				ThinZone2DFinderTool::AlignStartEndPoints(FirstSide, SecondSide);
+			}
+			else
+			{
+				ThinZone2DFinderTool::AlignStartEndPoints(SecondSide, FirstSide);
+			}
+		}
+
 		BuildThinZone(FirstSide, SecondSide);
 
 		FirstSide.Empty(); // to avoid to rebuild a second thin zone with SecondSide & FirstSide
@@ -1091,6 +1155,11 @@ void FThinZone2DFinder::GetThinZoneSideConnectionsLength(const TArray<FEdgeSegme
 
 void FThinZone2DFinder::BuildThinZone(const TArray<FEdgeSegment*>& FirstSide, const TArray<FEdgeSegment*>& SecondSide)
 {
+	if (FirstSide.IsEmpty() || SecondSide.IsEmpty())
+	{
+		return;
+	}
+
 	TFunction<void(const TArray<FEdgeSegment*>&, double&, double&)> ComputeThicknessAndLength = [](const TArray<FEdgeSegment*>& Side, double& MaxThickness, double& SideLength)
 	{
 		SideLength = 0;
@@ -1227,7 +1296,7 @@ bool FThinZone2DFinder::SearchThinZones(double InTolerance)
 	if (bDisplay)
 	{
 		DisplayCloseSegments();
-		//Wait();
+		Wait(false);
 	}
 #endif
 
@@ -1238,7 +1307,7 @@ bool FThinZone2DFinder::SearchThinZones(double InTolerance)
 	if (bDisplay)
 	{
 		ThinZone::DisplayThinZoneSides(ThinZoneSides);
-		//Wait();
+		Wait(false);
 	}
 #endif
 
@@ -1248,7 +1317,7 @@ bool FThinZone2DFinder::SearchThinZones(double InTolerance)
 	if (bDisplay)
 	{
 		ThinZone::DisplayThinZoneSides(ThinZoneSides);
-		//Wait();
+		Wait(false);
 	}
 #endif
 
@@ -1258,6 +1327,7 @@ bool FThinZone2DFinder::SearchThinZones(double InTolerance)
 	if (bDisplay)
 	{
 		ThinZone::DisplayThinZoneSides2(ThinZoneSides);
+		Wait(false);
 	}
 #endif
 
@@ -1266,8 +1336,8 @@ bool FThinZone2DFinder::SearchThinZones(double InTolerance)
 #ifdef DEBUG_SEARCH_THIN_ZONES
 	if (bDisplay)
 	{
-		ThinZone::DisplayThinZoneSides(ThinZoneSides);
-		//Wait();
+		ThinZone::DisplayThinZoneSidesAndCloses(ThinZoneSides);
+		Wait(false);
 	}
 #endif
 
@@ -1660,6 +1730,8 @@ void FThinZone2D::AddToEdge()
 void FThinZoneSide::GetExistingMeshNodes(const FTopologicalFace& Face, FModelMesh& MeshModel, FReserveContainerFunc& Reserve, FAddMeshNodeFunc& AddMeshNode, const bool bWithTolerance) const
 {
 	bool bEdgeIsForward = true;
+	bool bEdgeIsClosed = false;
+
 	int32 Index = 0;
 	int32 Increment = 1;
 	double SegmentUMin = 0.;
@@ -1683,46 +1755,50 @@ void FThinZoneSide::GetExistingMeshNodes(const FTopologicalFace& Face, FModelMes
 		}
 	};
 
-	TFunction<void(const FTopologicalEdge*, const FEdgeSegment&, ELimit)> AddEdgeVertexIfInclude = [&](const FTopologicalEdge* Edge, const FEdgeSegment& EdgeSegment, ELimit VertexExtremity)
+	TFunction<double(const int32)> ComputeTolerance = [&EdgeElementLength, &EdgeMeshNodeCount](const int32 Index)
 	{
-		const bool bIsStartVertex = (bEdgeIsForward == (VertexExtremity == ELimit::Start));
-
-		double VertexCoordinate = bIsStartVertex ? Edge->GetStartCurvilinearCoordinates() : Edge->GetEndCurvilinearCoordinates();
-
-		if (VertexCoordinate > SegmentUMin - DOUBLE_SMALL_NUMBER && VertexCoordinate < SegmentUMax + DOUBLE_SMALL_NUMBER)
+		if (Index == 0)
 		{
-			FTopologicalVertex& Vertex = bIsStartVertex ? *Edge->GetStartVertex() : *Edge->GetEndVertex();
-			int32 MeshVertexId = Vertex.GetOrCreateMesh(MeshModel).GetMesh();
-			FPoint2D Vertex2DCoordinate = EdgeSegment.ComputeEdgePoint(VertexCoordinate);
-
-			FPairOfIndex OppositeNodeIndices = FPairOfIndex::Undefined;
-
-			const TArray<FCuttingPoint>& CuttingPoints = Edge->GetCuttingPoints();
-			if (CuttingPoints.Num())
-			{
-				const FCuttingPoint& CuttingPoint = bIsStartVertex ? CuttingPoints[0] : CuttingPoints.Last();
-				if (FMath::IsNearlyEqual(CuttingPoint.Coordinate, VertexCoordinate))
-				{
-					OppositeNodeIndices = CuttingPoint.OppositNodeIndices;
-				}
-			}
-
-			const double MeshingTolerance = FMath::Max(DOUBLE_KINDA_SMALL_NUMBER, bWithTolerance ? (bIsStartVertex ? EdgeElementLength[0] : EdgeElementLength.Last()) * ASixth : DOUBLE_KINDA_SMALL_NUMBER);
-			AddMeshNode(MeshVertexId, Vertex2DCoordinate, MeshingTolerance, EdgeSegment, OppositeNodeIndices);
+			return EdgeElementLength[0] * ASixth;
+		}
+		else if (Index == EdgeMeshNodeCount - 1)
+		{
+			return EdgeElementLength[EdgeMeshNodeCount - 2] * ASixth;
 		}
+		return FMath::Min(EdgeElementLength[Index - 1], EdgeElementLength[Index]) * AThird;
 	};
 
 	TFunction<void(const FEdgeSegment&)> AddEdgeMeshNodeIfInclude = [&](const FEdgeSegment& EdgeSegment)
 	{
-		for (; Index >= 0 && Index < EdgeMeshNodeCount; Index += Increment)
+		bool bAlreadyDone = false;
+		for (;; Index += Increment)
 		{
+			if (Index < 0)
+			{
+				if (bAlreadyDone)
+				{
+					break;
+				}
+				bAlreadyDone = true;
+				Index = EdgeMeshNodeCount - 1;
+			}
+			else if (Index == EdgeMeshNodeCount)
+			{
+				if (bAlreadyDone)
+				{
+					break;
+				}
+				bAlreadyDone = true;
+				Index = 0;
+			}
+
 			const double MeshNodeCoordinate = CoordinatesOfMesh[Index];
 			if (MeshNodeCoordinate < SegmentUMin || MeshNodeCoordinate > SegmentUMax)
 			{
 				break;
 			}
 
-			const double MeshingTolerance = bWithTolerance ? FMath::Min(EdgeElementLength[Index], EdgeElementLength[Index + 1]) * AThird : 0;
+			const double MeshingTolerance = bWithTolerance ? ComputeTolerance(Index) : 0;
 			const FPoint2D MeshNode2D = EdgeSegment.ComputeEdgePoint(MeshNodeCoordinate);
 			AddMeshNode(NodeIndices[Index], MeshNode2D, MeshingTolerance, EdgeSegment, OppositeNodeIndices[Index]);
 		}
@@ -1776,116 +1852,101 @@ void FThinZoneSide::GetExistingMeshNodes(const FTopologicalFace& Face, FModelMes
 				bEdgeIsForward = EdgeSegment.IsForward();
 				EdgeMeshNodeCount = 0;
 
-				if (bWithTolerance)
-				{
-					if (ActiveEdge->IsMeshed())
-					{
-						EdgeElementLength = Edge->GetMesh()->GetElementLengths();
-					}
-					else
-					{
-						EdgeElementLength = Edge->GetPreElementLengths();
-					}
-				}
-
-				AddEdgeVertexIfInclude(Edge, EdgeSegment, ELimit::Start);
+				bEdgeIsClosed = Edge->IsClosed();
 
-				if (ActiveEdge->IsPreMeshed() && !Edge->IsDegenerated())
+				if (!Edge->IsDegenerated())
 				{
+					if (ActiveEdge->IsPreMeshed())
 					{
-						FAddCuttingPointFunc AddCuttingPoint = [&Edge](const double Coordinate, const ECoordinateType Type, const FPairOfIndex OppositNodeIndices, const double DeltaU)
 						{
-							Edge->AddTwinsCuttingPoint(Coordinate, DeltaU);
-						};
-
-						const bool bOnlyWithOppositeNode = false;
-						Edge->TransferCuttingPointFromMeshedEdge(bOnlyWithOppositeNode, AddCuttingPoint);
-					}
+							FAddCuttingPointFunc AddCuttingPoint = [&Edge](const double Coordinate, const ECoordinateType Type, const FPairOfIndex OppositNodeIndices, const double DeltaU)
+							{
+								Edge->AddTwinsCuttingPoint(Coordinate, DeltaU);
+							};
 
-					const TArray<FCuttingPoint>& CuttingPoints = Edge->GetCuttingPoints();
-					EdgeMeshNodeCount = CuttingPoints.Num();
-					if (EdgeMeshNodeCount > 2)
-					{
-						CoordinatesOfMesh.Empty(EdgeMeshNodeCount);
-						OppositeNodeIndices.Empty(EdgeMeshNodeCount);
-						EdgeMeshNodeCount -= 2;
-						for (int32 Cndex = 1; Cndex <= EdgeMeshNodeCount; ++Cndex)
-						{
-							const FCuttingPoint& CuttingPoint = CuttingPoints[Cndex];
-							CoordinatesOfMesh.Add(CuttingPoint.Coordinate);
-							OppositeNodeIndices.Emplace(CuttingPoint.OppositNodeIndices);
+							const bool bOnlyWithOppositeNode = false;
+							Edge->TransferCuttingPointFromMeshedEdge(bOnlyWithOppositeNode, AddCuttingPoint);
 						}
 
-						if (ActiveEdge->IsMeshed())
+						const TArray<FCuttingPoint>& CuttingPoints = Edge->GetCuttingPoints();
+						EdgeMeshNodeCount = CuttingPoints.Num();
+						if (EdgeMeshNodeCount >= 2)
 						{
-							NodeIndices.Empty(EdgeMeshNodeCount);
-							const FEdgeMesh* EdgeMesh = Edge->GetMesh();
-							if (!EdgeMesh)
+							CoordinatesOfMesh.Empty(EdgeMeshNodeCount);
+							OppositeNodeIndices.Empty(EdgeMeshNodeCount);
+							for (int32 Cndex = 0; Cndex < EdgeMeshNodeCount; ++Cndex)
 							{
-								continue;
+								const FCuttingPoint& CuttingPoint = CuttingPoints[Cndex];
+								CoordinatesOfMesh.Add(CuttingPoint.Coordinate);
+								OppositeNodeIndices.Emplace(CuttingPoint.OppositNodeIndices);
 							}
 
-							int32 StartMeshVertexId = EdgeMesh->GetStartVertexId();
-							if (Edge->IsSameDirection(*ActiveEdge))
+							if (ActiveEdge->IsMeshed())
 							{
-								for (int32 Cndex = 0; Cndex < EdgeMeshNodeCount; ++Cndex)
+								NodeIndices.Empty(EdgeMeshNodeCount);
+								const FEdgeMesh* EdgeMesh = Edge->GetMesh();
+								if (!EdgeMesh)
 								{
-									NodeIndices.Add(StartMeshVertexId++);
+									continue;
 								}
+
+								NodeIndices.Add(Edge->GetStartVertex()->GetMesh()->GetMesh());
+
+								const int32 InnerEdgeMeshNodeCount = EdgeMeshNodeCount - 2;
+								int32 StartMeshVertexId = EdgeMesh->GetStartVertexId();
+								if (Edge->IsSameDirection(*ActiveEdge))
+								{
+									for (int32 Cndex = 0; Cndex < InnerEdgeMeshNodeCount; ++Cndex)
+									{
+										NodeIndices.Add(StartMeshVertexId++);
+									}
+								}
+								else
+								{
+									StartMeshVertexId += InnerEdgeMeshNodeCount;
+									for (int32 Cndex = 0; Cndex < InnerEdgeMeshNodeCount; ++Cndex)
+									{
+										NodeIndices.Add(--StartMeshVertexId);
+									}
+								}
+
+								NodeIndices.Add(Edge->GetEndVertex()->GetMesh()->GetMesh());
 							}
 							else
 							{
-								StartMeshVertexId += EdgeMeshNodeCount;
-								for (int32 Cndex = 0; Cndex < EdgeMeshNodeCount; ++Cndex)
-								{
-									NodeIndices.Add(--StartMeshVertexId);
-								}
+								NodeIndices.Init(-1, EdgeMeshNodeCount);
+								NodeIndices[0] = Edge->GetStartVertex()->GetMesh()->GetMesh();
+								NodeIndices.Last() = Edge->GetEndVertex()->GetMesh()->GetMesh();
 							}
 						}
-						else
-						{
-							NodeIndices.Init(-1, EdgeMeshNodeCount);
-						}
-					}
-					else if (EdgeMeshNodeCount > 0)
-					{
-						EdgeMeshNodeCount = 0;
 					}
 					else
 					{
-						const FEdgeMesh* EdgeMesh = Edge->GetMesh();
-						if (!EdgeMesh)
-						{
-							continue;
-						}
-
-						const TArray<FPoint>& MeshVertices = EdgeMesh->GetNodeCoordinates();
-						int32 LocIncrement = 1;
+						CoordinatesOfMesh.Reset(2);
+						NodeIndices.Reset(2);
+						OppositeNodeIndices.Reset(2);
+
+						CoordinatesOfMesh.Add(Edge->GetBoundary().GetMin());
+						CoordinatesOfMesh.Add(Edge->GetBoundary().GetMax());
+						NodeIndices.Add(Edge->GetStartVertex()->GetMesh()->GetMesh());
+						NodeIndices.Add(Edge->GetEndVertex()->GetMesh()->GetMesh());
+						OppositeNodeIndices.Add(FPairOfIndex::Undefined);
+						OppositeNodeIndices.Add(FPairOfIndex::Undefined);
+						EdgeMeshNodeCount = 2;
+					}
 
-						int32 StartMeshVertexId = EdgeMesh->GetStartVertexId();
-						TArray<FPoint> ProjectedPoints;
-						Edge->ProjectPoints(MeshVertices, CoordinatesOfMesh, ProjectedPoints);
-						if (EdgeMeshNodeCount > 1 && CoordinatesOfMesh[0] > CoordinatesOfMesh[1])
+					if (bWithTolerance)
+					{
+						if (ActiveEdge->IsMeshed())
 						{
-							Algo::Reverse(CoordinatesOfMesh);
-							LocIncrement = -1;
-							StartMeshVertexId += EdgeMeshNodeCount - 1;
+							EdgeElementLength = Edge->GetMesh()->GetElementLengths();
 						}
-
-						NodeIndices.Empty(EdgeMeshNodeCount);
-						OppositeNodeIndices.Init(FPairOfIndex::Undefined, EdgeMeshNodeCount);
-						for (const double& Coordinate : CoordinatesOfMesh)
+						else
 						{
-							NodeIndices.Add(StartMeshVertexId);
-							StartMeshVertexId += LocIncrement;
+							EdgeElementLength = Edge->GetPreElementLengths();
 						}
 					}
 
-					if (EdgeElementLength.Num() != CuttingPoints.Num())
-					{
-						EdgeElementLength = Edge->GetPreElementLengths();
-					}
-
 					if (bEdgeIsForward)
 					{
 						Index = 0;
@@ -1899,28 +1960,15 @@ void FThinZoneSide::GetExistingMeshNodes(const FTopologicalFace& Face, FModelMes
 						FindFirstIndex(SegmentUMax, [](double Value1, double Value2) { return (Value1 > Value2); });
 					}
 				}
-				else if (bWithTolerance)
-				{
-					EdgeElementLength.Empty(1);
-					EdgeElementLength.Add(Edge->Length());
-				}
 			}
 
-			if (EdgeMeshNodeCount > Index)
+			if(EdgeMeshNodeCount)
 			{
 				AddEdgeMeshNodeIfInclude(EdgeSegment);
 			}
 		}
-
-		// Last vertex, last segment 
-		{
-			const FEdgeSegment& EdgeSegment = GetSegments().Last();
-			AddEdgeVertexIfInclude(Edge, EdgeSegment, ELimit::End);
-		}
 	}
 }
 
-
-
 }
 
diff --git a/Engine/Source/Runtime/Datasmith/CADKernel/Private/CADKernel/Mesh/Structure/ThinZone2DFinder.h b/Engine/Source/Runtime/Datasmith/CADKernel/Private/CADKernel/Mesh/Structure/ThinZone2DFinder.h
--- a/Engine/Source/Runtime/Datasmith/CADKernel/Private/CADKernel/Mesh/Structure/ThinZone2DFinder.h
+++ b/Engine/Source/Runtime/Datasmith/CADKernel/Private/CADKernel/Mesh/Structure/ThinZone2DFinder.h
@@ -164,10 +164,12 @@ void DisplayEdgeSegmentAndProjection(const FEdgeSegment* Segment, EVisuProperty
 void DisplayEdgeSegmentAndProjection(const FEdgeSegment* Segment, const FEdgeSegment* Opposite, EVisuProperty SegColor, EVisuProperty OppositeColor, EVisuProperty ProjectionColor);
 void DisplayEdgeSegment(const FEdgeSegment* EdgeSegment, EVisuProperty Color);
 void DisplayEdgeSegment(const FEdgeSegment* EdgeSegment, EVisuProperty Color, int32 Index);
+void DisplayThinZoneSidesAndCloses(const TArray<TArray<FEdgeSegment*>>& ThinZoneSides);
 void DisplayThinZoneSides(const TArray<TArray<FEdgeSegment*>>& ThinZoneSides);
 void DisplayThinZoneSides2(const TArray<TArray<FEdgeSegment*>>& ThinZoneSides);
 void DisplayThinZoneSide2(const TArray<FEdgeSegment*>& Side, int32 Index, bool bSplitBySide = false, bool bSplitBySegment = false);
 void DisplayThinZoneSide(const TArray<FEdgeSegment*>& Side, int32 Index, EVisuProperty Color, bool bSplitBySegment = false);
+void DisplayThinZoneSideAndClose(const TArray<FEdgeSegment*>& Side, int32 Index, bool bSplitBySegment = false);
 void DisplayThinZoneSide(const FThinZoneSide& Side, int32 Index, EVisuProperty Color, bool bSplitBySegment = false);
 void DisplayThinZones(const TArray<FThinZone2D>& ThinZones);
 }
diff --git a/Engine/Source/Runtime/Datasmith/CADKernel/Private/CADKernel/Topo/Shell.cpp b/Engine/Source/Runtime/Datasmith/CADKernel/Private/CADKernel/Topo/Shell.cpp
--- a/Engine/Source/Runtime/Datasmith/CADKernel/Private/CADKernel/Topo/Shell.cpp
+++ b/Engine/Source/Runtime/Datasmith/CADKernel/Private/CADKernel/Topo/Shell.cpp
@@ -266,7 +266,7 @@ void FShell::CheckTopology(TArray<FFaceSubset>& Subshells)
 		}
 	};
 	
-	TFunction<void(FFaceSubset&)> SpreadFront = [&](FFaceSubset& Shell)
+	TFunction<void(FFaceSubset&)> PropagateFront = [&](FFaceSubset& Shell)
 	{
 		while (Front.Num())
 		{
@@ -292,7 +292,7 @@ void FShell::CheckTopology(TArray<FFaceSubset>& Subshells)
 		Front.Empty(TopologicalFaceCount);
 		Face->SetMarker1();
 		Front.Add(Face);
-		SpreadFront(Shell);
+		PropagateFront(Shell);
 		ProcessFaceCount += Shell.Faces.Num();
 
 		if (ProcessFaceCount == TopologicalFaceCount)
@@ -394,7 +394,7 @@ int32 FShell::Orient()
 		}
 	};
 
-	TFunction<void()> SpreadFront = [&]()
+	TFunction<void()> PropagateFront = [&]()
 	{
 		while (Front.Num())
 		{
@@ -537,7 +537,7 @@ int32 FShell::Orient()
 		Face.Entity->SetProcessedMarker();
 		Front.Add(Face.Entity.Get());
 
-		SpreadFront();
+		PropagateFront();
 
 		CheckAndFixSubshellMainOrientation();
 
diff --git a/Engine/Source/Runtime/Datasmith/CADKernel/Private/CADKernel/Topo/TopologicalEdge.cpp b/Engine/Source/Runtime/Datasmith/CADKernel/Private/CADKernel/Topo/TopologicalEdge.cpp
--- a/Engine/Source/Runtime/Datasmith/CADKernel/Private/CADKernel/Topo/TopologicalEdge.cpp
+++ b/Engine/Source/Runtime/Datasmith/CADKernel/Private/CADKernel/Topo/TopologicalEdge.cpp
@@ -892,7 +892,7 @@ void FTopologicalEdge::TransferCuttingPointFromMeshedEdge(bool bOnlyWithOpposite
 {
 	const FTopologicalEdge* PreMeshedTwin = GetPreMeshedTwin();
 
-	if ((PreMeshedTwin == nullptr) || (PreMeshedTwin == this) || (PreMeshedTwin->GetFace() == GetFace()))
+	if ((PreMeshedTwin == nullptr) || (PreMeshedTwin == this) /*|| (PreMeshedTwin->GetFace() == GetFace())*/ )
 	{
 		return;
 	}
diff --git a/Engine/Source/Runtime/Datasmith/CADKernel/Private/CADKernel/Topo/Topomaker.cpp b/Engine/Source/Runtime/Datasmith/CADKernel/Private/CADKernel/Topo/Topomaker.cpp
--- a/Engine/Source/Runtime/Datasmith/CADKernel/Private/CADKernel/Topo/Topomaker.cpp
+++ b/Engine/Source/Runtime/Datasmith/CADKernel/Private/CADKernel/Topo/Topomaker.cpp
@@ -1127,7 +1127,7 @@ void FTopomaker::SplitIntoConnectedShells()
 		}
 	};
 
-	TFunction<void(FFaceSubset&)> SpreadFront = [&](FFaceSubset& Shell)
+	TFunction<void(FFaceSubset&)> PropagateFront = [&](FFaceSubset& Shell)
 	{
 		while (Front.Num())
 		{
@@ -1157,7 +1157,7 @@ void FTopomaker::SplitIntoConnectedShells()
 
 		Face->SetProcessedMarker();
 		Front.Add(Face.Get());
-		SpreadFront(Shell);
+		PropagateFront(Shell);
 		ProcessFaceCount += Shell.Faces.Num();
 
 		if (ProcessFaceCount == TopologicalFaceCount)
diff --git a/Engine/Source/Runtime/Datasmith/CADKernel/Public/CADKernel/Core/Session.h b/Engine/Source/Runtime/Datasmith/CADKernel/Public/CADKernel/Core/Session.h
--- a/Engine/Source/Runtime/Datasmith/CADKernel/Public/CADKernel/Core/Session.h
+++ b/Engine/Source/Runtime/Datasmith/CADKernel/Public/CADKernel/Core/Session.h
@@ -5,6 +5,7 @@
 #include "CADKernel/Core/Database.h"
 #include "CADKernel/Core/Entity.h"
 #include "CADKernel/Core/Types.h"
+#include "CADKernel/Math/Geometry.h"
 
 class FArchive;
 
@@ -32,6 +33,7 @@ public:
 	FSession(double InGeometricTolerance)
 		: GeometricTolerance(InGeometricTolerance)
 	{
+		IntersectionTool::SetTolerance(InGeometricTolerance);
 	}
 
 	FModel& GetModel();
@@ -39,6 +41,7 @@ public:
 	void Serialize(FCADKernelArchive& Ar)
 	{
 		Ar << GeometricTolerance;
+		IntersectionTool::SetTolerance(GeometricTolerance);
 	}
 
 	FDatabase& GetDatabase()
diff --git a/Engine/Source/Runtime/Datasmith/CADKernel/Public/CADKernel/Core/Types.h b/Engine/Source/Runtime/Datasmith/CADKernel/Public/CADKernel/Core/Types.h
--- a/Engine/Source/Runtime/Datasmith/CADKernel/Public/CADKernel/Core/Types.h
+++ b/Engine/Source/Runtime/Datasmith/CADKernel/Public/CADKernel/Core/Types.h
@@ -7,14 +7,22 @@
 #include "CoreMinimal.h"
 #endif
 
-#define SMALL_NUMBER_SQUARE 10e-16
-#define HUGE_VALUE 10e8
-#define HUGE_VALUE_SQUARE 10e16
-
-#define AThird 0.33333333333333333333333333333333
-#define AQuarter 0.25
-constexpr double ASixth = 1. / 6.; //0.16666666666666666666666666666667
-#define AEighth 0.125
+#define SMALL_NUMBER_SQUARE (1.e-16)
+#define DOUBLE_SMALL_NUMBER_SQUARE (1.e-16)
+
+#define KINDA_SMALL_NUMBER_SQUARE (1.e-8)
+#define DOUBLE_KINDA_SMALL_NUMBER_SQUARE (1.e-8)
+
+#define HUGE_VALUE (1e8)
+#define DOUBLE_HUGE_VALUE (1e8)
+
+#define HUGE_VALUE_SQUARE (1e16)
+#define DOUBLE_HUGE_VALUE_SQUARE (1e16)
+
+constexpr double AThird = 1. / 3.;
+constexpr double AQuarter = 0.25;
+constexpr double ASixth = 1. / 6.;
+constexpr double AEighth  = 0.125;
 
 typedef uint32 FIdent;
 
diff --git a/Engine/Source/Runtime/Datasmith/CADKernel/Public/CADKernel/Geo/Sampling/PolylineTools.h b/Engine/Source/Runtime/Datasmith/CADKernel/Public/CADKernel/Geo/Sampling/PolylineTools.h
--- a/Engine/Source/Runtime/Datasmith/CADKernel/Public/CADKernel/Geo/Sampling/PolylineTools.h
+++ b/Engine/Source/Runtime/Datasmith/CADKernel/Public/CADKernel/Geo/Sampling/PolylineTools.h
@@ -206,34 +206,46 @@ protected:
 
 		OutCurvilinearCoordinates.Reserve(BoundaryIndices[1] - BoundaryIndices[0] + 2);
 
-		double LastSegmentLength;
-		double Length = 0;
+		double LastEdgeSegmentLength;
+		double EdgeLength = 0;
+		double LengthOfSegment = 0;
 
 		OutCurvilinearCoordinates.Add(0);
-		if (BoundaryIndices[1] > BoundaryIndices[0] + 1)
+		if (BoundaryIndices[1] > BoundaryIndices[0])
 		{
 			PointType StartPoint = ComputePoint(BoundaryIndices[0], InBoundary.Min);
+			PointType PointIndice0 = PolylinePoints[BoundaryIndices[0]];
+
 			PointType EndPoint = ComputePoint(BoundaryIndices[1], InBoundary.Max);
-			Length = StartPoint.Distance(PolylinePoints[BoundaryIndices[0] + 1]);
-			LastSegmentLength = EndPoint.Distance(PolylinePoints[BoundaryIndices[1]]);
+			PointType NextPointIndice1 = PolylinePoints[BoundaryIndices[1] + 1];
+
+			EdgeLength = StartPoint.Distance(PolylinePoints[BoundaryIndices[0] + 1]);
+			double LengthOfSegments = PointIndice0.Distance(PolylinePoints[BoundaryIndices[0] + 1]);
+
+			LastEdgeSegmentLength = EndPoint.Distance(PolylinePoints[BoundaryIndices[1]]);
+			const double LastSegmentLength = NextPointIndice1.Distance(PolylinePoints[BoundaryIndices[1]]);
 
-			OutCurvilinearCoordinates.Add(Length);
+			OutCurvilinearCoordinates.Add(LengthOfSegments);
 			for (int32 Index = BoundaryIndices[0] + 1; Index < BoundaryIndices[1]; ++Index)
 			{
-				Length += PolylinePoints[Index].Distance(PolylinePoints[Index + 1]);
-				OutCurvilinearCoordinates.Add(Length);
+				const double SegLength = PolylinePoints[Index].Distance(PolylinePoints[Index + 1]);
+				EdgeLength += SegLength;
+				LengthOfSegments += SegLength;
+				OutCurvilinearCoordinates.Add(LengthOfSegments);
 			}
-			Length += LastSegmentLength;
-			OutCurvilinearCoordinates.Add(Length);
+			EdgeLength += LastEdgeSegmentLength;
+			LengthOfSegments += LastSegmentLength;
+			OutCurvilinearCoordinates.Add(LengthOfSegments);
 		}
 		else
 		{
 			PointType StartPoint = ComputePoint(BoundaryIndices[0], InBoundary.Min);
 			PointType EndPoint = ComputePoint(BoundaryIndices[1], InBoundary.Max);
-			Length = StartPoint.Distance(EndPoint);
-			OutCurvilinearCoordinates.Add(Length);
+			EdgeLength = StartPoint.Distance(EndPoint);
+			const double SegLength = PolylinePoints[BoundaryIndices[0]].Distance(PolylinePoints[BoundaryIndices[0] + 1]);
+			OutCurvilinearCoordinates.Add(SegLength);
 		}
-		return Length;
+		return EdgeLength;
 	}
 
 	PointType ComputePoint(const int32 Index, const double PointCoordinate) const
@@ -446,6 +458,9 @@ public:
 		TArray<double> CurvilinearCoordinates;
 		double CurveLength = ComputeCurvilinearCoordinatesOfPolyline(InBoundary, CurvilinearCoordinates, BoundaryIndices);
 
+		PointType StartPoint = ComputePoint(BoundaryIndices[0], InBoundary.Min);
+		double FromStartSegmentLength = PolylinePoints[BoundaryIndices[0]].Distance(StartPoint);
+
 		if (CurveLength < 2. * DesiredSegmentLength)
 		{
 			OutCoordinates.SetNum(3);
@@ -465,19 +480,22 @@ public:
 
 		TFunction<double(const int32, const int32, const double, const double)> ComputeSamplePointCoordinate = [&](const int32 IndexCurvilinear, const int32 IndexCoordinate, const double Length, const double Coordinate)
 		{
-			return Coordinate + (PolylineCoordinates[IndexCoordinate] - Coordinate) * (Length - CurvilinearCoordinates[IndexCurvilinear - 1u]) / (CurvilinearCoordinates[IndexCurvilinear] - CurvilinearCoordinates[IndexCurvilinear - 1u]);
+			return Coordinate + (PolylineCoordinates[IndexCoordinate] - Coordinate) * (Length - CurvilinearCoordinates[IndexCurvilinear - 1]) / (CurvilinearCoordinates[IndexCurvilinear] - CurvilinearCoordinates[IndexCurvilinear - 1]);
 		};
 
 		double CurvilinearLength = SectionLength;
+		FromStartSegmentLength += SectionLength;
+
 		double LastCoordinate = InBoundary.Min;
 		for (int32 IndexCurvilinear = 1, IndexCoordinate = BoundaryIndices[0] + 1; IndexCurvilinear < CurvilinearCoordinates.Num(); ++IndexCurvilinear, ++IndexCoordinate)
 		{
-			while (CurvilinearLength + DOUBLE_SMALL_NUMBER < CurvilinearCoordinates[IndexCurvilinear])
+			while (FromStartSegmentLength + DOUBLE_SMALL_NUMBER < CurvilinearCoordinates[IndexCurvilinear])
 			{
-				double Coordinate = ComputeSamplePointCoordinate(IndexCurvilinear, IndexCoordinate, CurvilinearLength, LastCoordinate);
+				double Coordinate = ComputeSamplePointCoordinate(IndexCurvilinear, IndexCoordinate, FromStartSegmentLength, LastCoordinate);
 				OutCoordinates.Add(Coordinate);
 				CurvilinearLength += SectionLength;
-				if (CurvilinearLength + DOUBLE_SMALL_NUMBER > CurveLength)
+				FromStartSegmentLength += SectionLength;
+				if (OutCoordinates.Num() == SegmentNum)
 				{
 					OutCoordinates.Add(InBoundary.Max);
 					break;
diff --git a/Engine/Source/Runtime/Datasmith/CADKernel/Public/CADKernel/Math/Geometry.h b/Engine/Source/Runtime/Datasmith/CADKernel/Public/CADKernel/Math/Geometry.h
--- a/Engine/Source/Runtime/Datasmith/CADKernel/Public/CADKernel/Math/Geometry.h
+++ b/Engine/Source/Runtime/Datasmith/CADKernel/Public/CADKernel/Math/Geometry.h
@@ -19,6 +19,12 @@ enum EPolygonSide : uint8
 	Side30,
 };
 
+namespace IntersectionTool
+{
+CADKERNEL_API void SetTolerance(const double Tolerance);
+}
+
+
 /**
  * https://en.wikipedia.org/wiki/Circumscribed_circle#Cartesian_coordinates_2
  * With A = (0, 0)
@@ -501,37 +507,6 @@ inline double CoordinateOfProjectedPointOnSegment(const PointType& Point, const
 
 CADKERNEL_API void FindLoopIntersectionsWithIso(const EIso Iso, const double IsoParameter, const TArray<TArray<FPoint2D>>& Loops, TArray<double>& OutIntersections);
 
-/**
- * Similar as DoIntersect but do not check intersection if both segment are carried by the same line.
- * Must be done before (with BBox comparison for example)
- * This method is 50% faster than IntersectSegments2D even if
- */
-inline bool FastDoIntersect(const FSegment2D& SegmentAB, const FSegment2D& SegmentCD)
-{
-	constexpr const double Min = -DOUBLE_SMALL_NUMBER;
-	constexpr const double Max = 1. + DOUBLE_SMALL_NUMBER;
-
-	FPoint2D AB = SegmentAB[1] - SegmentAB[0];
-	FPoint2D CD = SegmentCD[1] - SegmentCD[0];
-	FPoint2D CA = SegmentAB[0] - SegmentCD[0];
-
-	double ParallelCoef = CD ^ AB;
-	double ABIntersectionCoordinate = (CA ^ CD) / ParallelCoef;
-	double CDIntersectionCoordinate = (CA ^ AB) / ParallelCoef;
-
-	if (FMath::IsNearlyZero(ParallelCoef))
-	{
-		ParallelCoef = CA ^ AB;
-		if (!FMath::IsNearlyZero(ParallelCoef))
-		{
-			return false;
-		}
-		return true;
-	}
-
-	return (ABIntersectionCoordinate <= Max && ABIntersectionCoordinate >= Min && CDIntersectionCoordinate <= Max && CDIntersectionCoordinate >= Min);
-}
-
 /**
  * The segments must intersect because no check is done
  */
diff --git a/Engine/Source/Runtime/Datasmith/CADKernel/Public/CADKernel/Math/Point.h b/Engine/Source/Runtime/Datasmith/CADKernel/Public/CADKernel/Math/Point.h
--- a/Engine/Source/Runtime/Datasmith/CADKernel/Public/CADKernel/Math/Point.h
+++ b/Engine/Source/Runtime/Datasmith/CADKernel/Public/CADKernel/Math/Point.h
@@ -467,8 +467,18 @@ PRAGMA_ENABLE_DEPRECATION_WARNINGS
 class CADKERNEL_API FPoint2D
 {
 public:
-	double U;
-	double V;
+	union
+	{
+		struct
+		{
+			double U;
+			double V;
+		};
+
+		UE_DEPRECATED(all, "For internal use only")
+		double UV[2];
+	};
+
 
 public:
 
@@ -537,9 +547,9 @@ public:
 		return FMath::Square(U) + FMath::Square(V);
 	}
 
-	FPoint2D& Normalize()
+	FPoint2D& Normalize(double& Norm)
 	{
-		double Norm = Length();
+		Norm = Length();
 		if (FMath::IsNearlyZero(Norm))
 		{
 			U = V = 0.0;
@@ -552,14 +562,24 @@ public:
 		return *this;
 	}
 
+	FPoint2D& Normalize()
+	{
+		double Norm = 0;
+		return Normalize(Norm);
+	}
+
 	constexpr double& operator[](int32 Index)
 	{
-		return *(&U + Index);
+PRAGMA_DISABLE_DEPRECATION_WARNINGS
+		return UV[Index];
+PRAGMA_ENABLE_DEPRECATION_WARNINGS
 	}
 
 	constexpr double operator[](int32 Index) const
 	{
-		return *(&U + Index);
+PRAGMA_DISABLE_DEPRECATION_WARNINGS
+		return UV[Index];
+PRAGMA_ENABLE_DEPRECATION_WARNINGS
 	}
 
 	bool operator==(const FPoint2D& Point) const
diff --git a/Engine/Source/Runtime/Datasmith/CADKernel/Public/CADKernel/Mesh/Structure/Mesh.h b/Engine/Source/Runtime/Datasmith/CADKernel/Public/CADKernel/Mesh/Structure/Mesh.h
--- a/Engine/Source/Runtime/Datasmith/CADKernel/Public/CADKernel/Mesh/Structure/Mesh.h
+++ b/Engine/Source/Runtime/Datasmith/CADKernel/Public/CADKernel/Mesh/Structure/Mesh.h
@@ -30,6 +30,9 @@ public:
 		: FEntityGeom()
 		, ModelMesh(InMeshModel)
 		, TopologicalEntity(InTopologicalEntity)
+		, StartNodeId(0)
+		, LastNodeIndex(0)
+		, MeshModelIndex(0)
 	{
 	}
 
diff --git a/Engine/Source/Runtime/Datasmith/CADKernel/Public/CADKernel/Mesh/Structure/ThinZone2D.h b/Engine/Source/Runtime/Datasmith/CADKernel/Public/CADKernel/Mesh/Structure/ThinZone2D.h
--- a/Engine/Source/Runtime/Datasmith/CADKernel/Public/CADKernel/Mesh/Structure/ThinZone2D.h
+++ b/Engine/Source/Runtime/Datasmith/CADKernel/Public/CADKernel/Mesh/Structure/ThinZone2D.h
@@ -127,6 +127,15 @@ namespace UE::CADKernel
 			return Segments[0].IsInner();
 		}
 
+		bool IsClosed() const
+		{
+			if (Segments.Num() && (Segments[0].GetPrevious() == &Segments.Last()))
+			{
+				return true;
+			}
+			return false;
+		}
+
 		FThinZoneSide& GetFrontThinZoneSide()
 		{
 			return FrontSide;
diff --git a/Engine/Source/Runtime/Datasmith/CADKernel/Public/CADKernel/Topo/TopologicalShapeEntity.h b/Engine/Source/Runtime/Datasmith/CADKernel/Public/CADKernel/Topo/TopologicalShapeEntity.h
--- a/Engine/Source/Runtime/Datasmith/CADKernel/Public/CADKernel/Topo/TopologicalShapeEntity.h
+++ b/Engine/Source/Runtime/Datasmith/CADKernel/Public/CADKernel/Topo/TopologicalShapeEntity.h
@@ -58,7 +58,7 @@ public:
 
 	/**
 	 * Each face of model is set by its orientation. This allow to make oriented mesh and to keep the face orientation in topological function.
-	 * Marker2 of spread face is set. It must be reset after the process
+	 * Marker2 of propagate face is set. It must be reset after the process
 	 */
 	virtual void PropagateBodyOrientation() = 0;
 
