Fix  Compilation Tests / NonUnity Compile UnrealGame Android

diff --git a/Engine/Source/Runtime/Datasmith/CADKernel/Private/CADKernel/Mesh/Meshers/IsoTriangulator/IntersectionIsoSegmentTool.cpp b/Engine/Source/Runtime/Datasmith/CADKernel/Private/CADKernel/Mesh/Meshers/IsoTriangulator/IntersectionIsoSegmentTool.cpp
--- a/Engine/Source/Runtime/Datasmith/CADKernel/Private/CADKernel/Mesh/Meshers/IsoTriangulator/IntersectionIsoSegmentTool.cpp
+++ b/Engine/Source/Runtime/Datasmith/CADKernel/Private/CADKernel/Mesh/Meshers/IsoTriangulator/IntersectionIsoSegmentTool.cpp
@@ -1,5 +1,4 @@
 // Copyright Epic Games, Inc. All Rights Reserved.
-#pragma once
 
 #include "CADKernel/Mesh/Meshers/IsoTriangulator/IntersectionIsoSegmentTool.h"
 #include "CADKernel/Geo/GeoEnum.h"
