[NO NON-PII COMMIT MESSAGE]

diff --git a/Engine/Plugins/Enterprise/DatasmithCADImporter/Source/DatasmithWireTranslator/Private/DatasmithWireTranslator.cpp b/Engine/Plugins/Enterprise/DatasmithCADImporter/Source/DatasmithWireTranslator/Private/DatasmithWireTranslator.cpp
--- a/Engine/Plugins/Enterprise/DatasmithCADImporter/Source/DatasmithWireTranslator/Private/DatasmithWireTranslator.cpp
+++ b/Engine/Plugins/Enterprise/DatasmithCADImporter/Source/DatasmithWireTranslator/Private/DatasmithWireTranslator.cpp
@@ -2431,6 +2431,35 @@ bool FDatasmithWireTranslator::LoadScene(TSharedRef<IDatasmithScene> OutScene)
 	UE_LOG(LogDatasmithWireTranslator, Display, TEXT("CAD translation [%s]."), *Filename);
 	UE_LOG(LogDatasmithWireTranslator, Display, TEXT(" - Parsing Library:      %s"), TEXT("Alias"));
 	UE_LOG(LogDatasmithWireTranslator, Display, TEXT(" - Tessellation Library: %s"), CADLibrary::FImportParameters::bGDisableCADKernelTessellation ? TEXT("TechSoft") : TEXT("CADKernel"));
+	UE_LOG(LogDatasmithWireTranslator, Display, TEXT(" - Import parameters:"));
+
+	const FDatasmithTessellationOptions& TessellationOptions = GetCommonTessellationOptions();
+	UE_LOG(LogDatasmithWireTranslator, Display, TEXT("     - ChordTolerance:     %lf"), TessellationOptions.ChordTolerance);
+	UE_LOG(LogDatasmithWireTranslator, Display, TEXT("     - MaxEdgeLength:      %lf"), TessellationOptions.MaxEdgeLength);
+	UE_LOG(LogDatasmithWireTranslator, Display, TEXT("     - MaxNormalAngle:     %lf"), TessellationOptions.NormalTolerance);
+	FString StitchingTechnique;
+	switch (TessellationOptions.StitchingTechnique)
+	{
+	case EDatasmithCADStitchingTechnique::StitchingHeal:
+		StitchingTechnique = TEXT("Heal");
+		break;
+	case EDatasmithCADStitchingTechnique::StitchingSew:
+		StitchingTechnique = TEXT("Sew");
+		break;
+	default:
+		StitchingTechnique = TEXT("None");
+		break;
+	}
+	UE_LOG(LogDatasmithWireTranslator, Display, TEXT("     - StitchingTechnique: %s"), *StitchingTechnique);
+
+	if (!CADLibrary::FImportParameters::bGDisableCADKernelTessellation)
+	{
+		UE_LOG(LogDatasmithWireTranslator, Display, TEXT("     - Stitching Options:"));
+		UE_LOG(LogDatasmithWireTranslator, Display, TEXT("         - ForceSew:              %s"), CADLibrary::FImportParameters::bGStitchingForceSew ? TEXT("True") : TEXT("False"));
+		UE_LOG(LogDatasmithWireTranslator, Display, TEXT("         - RemoveThinFaces:       %s"), CADLibrary::FImportParameters::bGStitchingRemoveThinFaces ? TEXT("True") : TEXT("False"));
+		UE_LOG(LogDatasmithWireTranslator, Display, TEXT("         - RemoveDuplicatedFaces: %s"), CADLibrary::FImportParameters::bGStitchingRemoveDuplicatedFaces ? TEXT("True") : TEXT("False"));
+		UE_LOG(LogDatasmithWireTranslator, Display, TEXT("         - ForceFactor:           %f"), CADLibrary::FImportParameters::GStitchingForceFactor);
+	}
 
 	FString OutputPath = FPaths::ConvertRelativePathToFull(FPaths::Combine(FDatasmithWireTranslatorModule::Get().GetTempDir(), TEXT("Cache"), GetSource().GetSceneName()));
 	IFileManager::Get().MakeDirectory(*OutputPath, true);
diff --git a/Engine/Source/Runtime/Datasmith/CADKernel/Private/CADKernel/Core/Factory.h b/Engine/Source/Runtime/Datasmith/CADKernel/Private/CADKernel/Core/Factory.h
--- a/Engine/Source/Runtime/Datasmith/CADKernel/Private/CADKernel/Core/Factory.h
+++ b/Engine/Source/Runtime/Datasmith/CADKernel/Private/CADKernel/Core/Factory.h
@@ -12,11 +12,12 @@
 
 namespace UE::CADKernel
 {
-const int32 MaxSize = 256;
 
 template<class ElementType>
 class TFactory
 {
+	static const int32 MaxSize = 256;
+
 protected:
 	TDoubleLinkedList<TArray<ElementType>> StackOfEntitySet;
 	TArray<ElementType*> FreeEntityStack;
diff --git a/Engine/Source/Runtime/Datasmith/CADKernel/Private/CADKernel/Mesh/Meshers/CycleTriangulator.cpp b/Engine/Source/Runtime/Datasmith/CADKernel/Private/CADKernel/Mesh/Meshers/CycleTriangulator.cpp
--- a/Engine/Source/Runtime/Datasmith/CADKernel/Private/CADKernel/Mesh/Meshers/CycleTriangulator.cpp
+++ b/Engine/Source/Runtime/Datasmith/CADKernel/Private/CADKernel/Mesh/Meshers/CycleTriangulator.cpp
@@ -37,7 +37,7 @@ bool FCycleTriangulator::MeshCycle()
 	if (Grid.bDisplay)
 	{
 		Grid.DisplayIsoSegments(TEXT("Cycle"), EGridSpace::UniformScaled, Cycle);
-		Wait(CycleId == 60);
+		Wait(CycleId == 0);
 	}
 #endif
 
@@ -586,22 +586,28 @@ bool FCycleTriangulator::FindTheBestCandidateNode()
 			// if the segment exist, it has already been tested
 			{
 				const FIsoSegment* StartSegment = FirstSideStartNode->GetSegmentConnectedTo(Node);
-				if (!StartSegment && CycleIntersectionTool.DoesIntersect(FirstSideStartNode, Node))
+				if (!StartSegment)
 				{
-					if (ConfirmIntersection(FirstSideStartNode, Node, FirstSideEndNode, StartSegment))
+					if(const FIsoSegment* Intersection = CycleIntersectionTool.FindIntersectingSegment(FirstSideStartNode, Node))
 					{
-						continue;
+						if (ConfirmIntersection(FirstSideStartNode, Node, FirstSideEndNode, Intersection))
+						{
+							continue;
+						}
 					}
 				}
 			}
 
 			{
 				const FIsoSegment* EndSegment = FirstSideEndNode->GetSegmentConnectedTo(Node);
-				if (CycleIntersectionTool.DoesIntersect(FirstSideEndNode, Node))
+				if(!EndSegment)
 				{
-					if (ConfirmIntersection(FirstSideEndNode, Node, FirstSideStartNode, EndSegment))
+					if (const FIsoSegment* Intersection = CycleIntersectionTool.FindIntersectingSegment(FirstSideEndNode, Node))
 					{
-						continue;
+						if (ConfirmIntersection(FirstSideEndNode, Node, FirstSideStartNode, EndSegment))
+						{
+							continue;
+						}
 					}
 				}
 			}
diff --git a/Engine/Source/Runtime/Datasmith/CADKernel/Private/CADKernel/Mesh/Structure/Grid.cpp b/Engine/Source/Runtime/Datasmith/CADKernel/Private/CADKernel/Mesh/Structure/Grid.cpp
--- a/Engine/Source/Runtime/Datasmith/CADKernel/Private/CADKernel/Mesh/Structure/Grid.cpp
+++ b/Engine/Source/Runtime/Datasmith/CADKernel/Private/CADKernel/Mesh/Structure/Grid.cpp
@@ -1012,6 +1012,14 @@ void FGrid::GetMeshOfLoop(const FTopologicalLoop& Loop)
 		if (Edge->IsDegenerated())
 		{
 			EdgeVerticesIndex.Init(ActiveEdge->GetStartVertex()->GetLinkActiveEntity()->GetOrCreateMesh(MeshModel).GetMesh(), CuttingPolyline.Coordinates.Num());
+			if (OrientedEdge.Direction == Front)
+			{
+				EdgeVerticesIndex[0] = ActiveEdge->GetEndVertex()->GetLinkActiveEntity()->GetOrCreateMesh(MeshModel).GetMesh();
+			}
+			else
+			{
+				EdgeVerticesIndex.Last() = ActiveEdge->GetEndVertex()->GetLinkActiveEntity()->GetOrCreateMesh(MeshModel).GetMesh();
+			}
 		}
 		else if (Edge->IsVirtuallyMeshed())
 		{
@@ -1125,28 +1133,87 @@ void FGrid::GetMeshOfLoop(const FTopologicalLoop& Loop)
 
 void FGrid::GetMeshOfThinZone(const FThinZone2D& ThinZone)
 {
-	TFunction<void(const FThinZoneSide&, TArray<FPoint2D>&)> GetThinZoneSideMesh = [this](const FThinZoneSide& Side, TArray<FPoint2D>& ThinZoneSideMesh)
+	// ThinZones are identified during "ApplyCriteria". "ApplyCriteria" use the FCriteriaGrid, The EGridSpace::UniformScaled are not the same between FCriteriaGrid and FGrid
+	// This is the reason why, to get the ThinZone in the FGrid UniformScaled space, we need to get the mesh of the thinZone defined by the node ids, 
+	// And then, from this, we retrieve the thinZone points in Grid UniformScaled space.
+
+	// Return the mesh defined by the node index
+	TFunction<void(const FThinZoneSide&, TArray<int32>&)> GetThinZoneSideMesh = [this](const FThinZoneSide& Side, TArray<int32>& MeshIndices)
 	{
-		TArray<int32> MeshIndices;
-		FAddMeshNodeFunc AddMeshNode = [&Mesh = ThinZoneSideMesh, &MeshIndices](int32 NodeIndice, const FPoint2D& MeshNode2D, double MeshingTolerance, const FEdgeSegment& EdgeSegment, const FPairOfIndex& OppositeNodeIndices)
+		FAddMeshNodeFunc AddMeshNode = [&MeshIndices](int32 NodeIndice, const FPoint2D& MeshNode2D, double MeshingTolerance, const FEdgeSegment& EdgeSegment, const FPairOfIndex& OppositeNodeIndices)
 		{
-			if(MeshIndices.IsEmpty() || MeshIndices.Last() != NodeIndice)
+			if (MeshIndices.IsEmpty() || MeshIndices.Last() != NodeIndice)
 			{
 				MeshIndices.Add(NodeIndice);
-				Mesh.Emplace(MeshNode2D);
 			}
 		};
 
-		FReserveContainerFunc Reserve = [&Mesh = ThinZoneSideMesh, &MeshIndices](int32 MeshVertexCount)
+		FReserveContainerFunc Reserve = [&MeshIndices](int32 MeshVertexCount)
 		{
-			Mesh.Reserve(Mesh.Num() + MeshVertexCount);
-			MeshIndices.Reserve(Mesh.Num() + MeshVertexCount);
+			MeshIndices.Reserve(MeshIndices.Num() + MeshVertexCount);
 		};
-
 		const bool bWithTolerance = false;
 		Side.GetExistingMeshNodes(GetFace(), MeshModel, Reserve, AddMeshNode, bWithTolerance);
 	};
 
+	TFunction<bool(int32, int32&, int32&)> FindLoopAndNodeIndex = [&NodeIdsOfLoops = NodeIdsOfFaceLoops](int32 NodeId, int32& LoopIndex, int32& LoopNodeIndex)
+	{
+		bool bFind = false;
+		for (LoopIndex = 0; LoopIndex < NodeIdsOfLoops.Num(); ++LoopIndex)
+		{
+			TArray<int32>& NodeIdsOfLoop = NodeIdsOfLoops[LoopIndex];
+			for (LoopNodeIndex = 0; LoopNodeIndex < NodeIdsOfLoop.Num(); ++LoopNodeIndex)
+			{
+				if (NodeIdsOfLoop[LoopNodeIndex] == NodeId)
+				{
+					return true;
+				}
+			}
+		}
+		LoopIndex = -1;
+		LoopNodeIndex = -1;
+		return false;
+	};
+
+	TFunction<void(const TArray<int32>&, TArray<FPoint2D>&)> GetThinZoneMeshCoordinates = [&NodeIdsOfLoops = NodeIdsOfFaceLoops, &Loops2D = FaceLoops2D, &FindLoopAndNodeIndex](const TArray<int32>& ThinZoneNodeIds, TArray<FPoint2D>& ThinZonePoints)
+	{
+		ThinZonePoints.Reset(ThinZoneNodeIds.Num());
+
+		int32 LoopIndex = 0;
+		int32 LoopNodeIndex = 0;
+
+		for (int32 NodeId : ThinZoneNodeIds)
+		{
+			if (LoopIndex >= 0)
+			{
+				const int32 NextLoopNodeIndex = (LoopNodeIndex == NodeIdsOfLoops[LoopIndex].Num() - 1) ? 0 : LoopNodeIndex + 1;
+				if (NodeIdsOfLoops[LoopIndex][NextLoopNodeIndex] == NodeId)
+				{
+					LoopNodeIndex = NextLoopNodeIndex;
+				}
+				else
+				{
+					const int32 PrevLoopNodeIndex = LoopNodeIndex ? LoopNodeIndex - 1 : NodeIdsOfLoops[LoopIndex].Num() - 1;
+					if (NodeIdsOfLoops[LoopIndex][PrevLoopNodeIndex] == NodeId)
+					{
+						LoopNodeIndex = PrevLoopNodeIndex;
+					}
+					else
+					{
+						LoopIndex = -1;
+					}
+				}
+			}
+
+			if (LoopIndex < 0 && !FindLoopAndNodeIndex(NodeId, LoopIndex, LoopNodeIndex))
+			{
+				continue;
+			}
+
+			ThinZonePoints.Emplace(Loops2D[EGridSpace::UniformScaled][LoopIndex][LoopNodeIndex]);
+		}
+	};
+
 	TFunction<void(TArray<FPoint2D>&)> RemoveDuplicatedNode = [](TArray<FPoint2D>& ThinZoneMesh)
 	{
 		if (ThinZoneMesh.Num() > 1)
@@ -1175,17 +1242,39 @@ void FGrid::GetMeshOfThinZone(const FThinZone2D& ThinZone)
 	const bool FirstSideIsClosed = ThinZone.GetFirstSide().IsClosed();
 	const bool SecondSideIsClosed = ThinZone.GetSecondSide().IsClosed();
 
-	TArray<FPoint2D>* ThinZoneMesh = &FaceLoops2D[EGridSpace::UniformScaled].Emplace_GetRef();
-	GetThinZoneSideMesh(ThinZone.GetFirstSide(), *ThinZoneMesh);
+	TArray<int32> ThinZoneNodeIds;
+	GetThinZoneSideMesh(ThinZone.GetFirstSide(), ThinZoneNodeIds);
+	const int32 FirstSideSize = ThinZoneNodeIds.Num();
+
+	if (FirstSideSize < 2)
+	{
+		return;
+	}
+
 	if (FirstSideIsClosed && SecondSideIsClosed)
 	{
-		RemoveDuplicatedNode(*ThinZoneMesh);
-		ThinZoneMesh = &FaceLoops2D[EGridSpace::UniformScaled].Emplace_GetRef();
+		TArray<FPoint2D>& ThinZoneMesh = FaceLoops2D[EGridSpace::UniformScaled].Emplace_GetRef();
+		GetThinZoneMeshCoordinates(ThinZoneNodeIds, ThinZoneMesh);
+		RemoveDuplicatedNode(ThinZoneMesh);
+		ThinZoneNodeIds.Reset();
 	}
-	GetThinZoneSideMesh(ThinZone.GetSecondSide(), *ThinZoneMesh);
-	RemoveDuplicatedNode(*ThinZoneMesh);
 
-	if (ThinZoneMesh->Num() < 3)
+	GetThinZoneSideMesh(ThinZone.GetSecondSide(), ThinZoneNodeIds);
+	const int32 SecondSideSize = (FirstSideIsClosed && SecondSideIsClosed) ? ThinZoneNodeIds.Num() : ThinZoneNodeIds.Num() - FirstSideSize;
+	if (SecondSideSize < 2)
+	{
+		if (FirstSideIsClosed && SecondSideIsClosed)
+		{
+			FaceLoops2D[EGridSpace::UniformScaled].Pop();
+		}
+		return;
+	}
+
+	TArray<FPoint2D>& ThinZoneMesh = FaceLoops2D[EGridSpace::UniformScaled].Emplace_GetRef();
+	GetThinZoneMeshCoordinates(ThinZoneNodeIds, ThinZoneMesh);
+	RemoveDuplicatedNode(ThinZoneMesh);
+
+	if (ThinZoneMesh.Num() < 4)
 	{
 		FaceLoops2D[EGridSpace::UniformScaled].Pop();
 	}
diff --git a/Engine/Source/Runtime/Datasmith/CADKernel/Private/CADKernel/Mesh/Structure/LoopCleaner.cpp b/Engine/Source/Runtime/Datasmith/CADKernel/Private/CADKernel/Mesh/Structure/LoopCleaner.cpp
--- a/Engine/Source/Runtime/Datasmith/CADKernel/Private/CADKernel/Mesh/Structure/LoopCleaner.cpp
+++ b/Engine/Source/Runtime/Datasmith/CADKernel/Private/CADKernel/Mesh/Structure/LoopCleaner.cpp
@@ -618,21 +618,21 @@ void FLoopCleaner::FixLoopOrientation()
 	if (Orientation == EOrientation::Back)
 	{
 #ifdef DEBUG_LOOP_ORIENTATION
+		F3DDebugSession Q(Grid.bDisplay, TEXT("FixLoopOrientation"));
+		Grid.DisplayIsoNodes(TEXT("Nodes of loop Before orientation"), EGridSpace::UniformScaled, (const TArray<const FIsoNode*>&) NodesOfLoop, EVisuProperty::BluePoint);
 		Grid.DisplayIsoSegments(TEXT("Before orientation"), EGridSpace::UniformScaled, LoopSegments, true, true);
-		Wait(bDisplay);
+		Wait(false);
 #endif
 
 		SwapSubLoopOrientation(StartSegmentIndex, NextLoopFirstSegmentIndex);
 
 #ifdef DEBUG_LOOP_ORIENTATION
 		Grid.DisplayIsoSegments(TEXT("After orientation"), EGridSpace::UniformScaled, LoopSegments, true, true);
-		Wait(bDisplay);
+		Wait(false);
 #endif
 	}
 }
 
-
-
 bool FLoopCleaner::RemoveSelfIntersectionsOfLoop()
 {
 	using namespace LoopCleanerImpl;
@@ -1617,28 +1617,13 @@ void FLoopCleaner::SwapSegments(FIsoSegment& IntersectingSegment, FIsoSegment& S
 
 void FLoopCleaner::SwapSubLoopOrientation(int32 FirstSegmentIndex, int32 LastSegmentIndex)
 {
+	FirstSegmentIndex = FitSegmentIndex(FirstSegmentIndex);
 	LastSegmentIndex = FitSegmentIndex(LastSegmentIndex);
 
 	TArray<FIsoSegment*> Segments;
-	{
-		int32 SegCount = 0;
-		if (LastSegmentIndex < FirstSegmentIndex)
-		{
-			int32 SegIndex = FitSegmentIndex(FirstSegmentIndex);
-			do
-			{
-				SegCount++;
-				SegIndex = NextSegmentIndex(SegIndex);
-			} while (SegIndex != LastSegmentIndex);
-		}
-		else
-		{
-			SegCount = LastSegmentIndex - FirstSegmentIndex;
-		}
-		Segments.Reserve(SegCount);
-	}
+	Segments.Reserve(SegmentCount);
 
-	int32 Index = FitSegmentIndex(FirstSegmentIndex);
+	int32 Index = FirstSegmentIndex;
 	do
 	{
 		LoopSegments[Index]->SwapOrientation();
@@ -1647,15 +1632,12 @@ void FLoopCleaner::SwapSubLoopOrientation(int32 FirstSegmentIndex, int32 LastSeg
 	}
 	while(Index != LastSegmentIndex);
 
-	int32 ReverseIndex = Segments.Num() - 1;
-	Index = FitSegmentIndex(FirstSegmentIndex);
-	do
+	Index = FirstSegmentIndex;
+	for (int32 ReverseIndex = Segments.Num() - 1; ReverseIndex >= 0; --ReverseIndex)
 	{
 		LoopSegments[Index] = Segments[ReverseIndex];
 		Index = NextSegmentIndex(Index);
-		--ReverseIndex;
 	}
-	while(Index != LastSegmentIndex);
 }
 
 bool FLoopCleaner::RemoveSubLoop(FLoopNode* StartNode, FLoopNode* EndNode)
@@ -2169,18 +2151,20 @@ void FLoopCleaner::FindBestLoopExtremity()
 	double OptimalSlope = 9.;
 	LoopIndex = 0;
 
-	//F3DDebugSession _(TEXT("FindBestLoopExtremity"));
-
-	BestStartNodeOfLoops.Reserve(Grid.GetLoopCount());
+	BestStartNodeOfLoops.Reset(Grid.GetLoopCount());
 
 	TFunction<void(FLoopNode*)> CompareWithSlopeAt = [&](FLoopNode* Node)
 	{
 		FLoopNode& PreviousNode = Node->GetPreviousNode();
 		FLoopNode& NextNode = Node->GetNextNode();
-		double Slope = ComputePositiveSlope(Node->Get2DPoint(EGridSpace::UniformScaled, Grid), PreviousNode.Get2DPoint(EGridSpace::UniformScaled, Grid), NextNode.Get2DPoint(EGridSpace::UniformScaled, Grid));
 
-		//F3DDebugSession A(FString::Printf(TEXT("Node Slope %f"), Slope));
-		//Display(EGridSpace::UniformScaled, *Node, 0, EVisuProperty::RedPoint);
+		if (Node->Get2DPoint(EGridSpace::UniformScaled, Grid).SquareDistance(PreviousNode.Get2DPoint(EGridSpace::UniformScaled, Grid)) < DOUBLE_SMALL_NUMBER ||
+			Node->Get2DPoint(EGridSpace::UniformScaled, Grid).SquareDistance(NextNode.Get2DPoint(EGridSpace::UniformScaled, Grid)) < DOUBLE_SMALL_NUMBER)
+		{
+			return;
+		}
+
+		double Slope = ComputePositiveSlope(Node->Get2DPoint(EGridSpace::UniformScaled, Grid), PreviousNode.Get2DPoint(EGridSpace::UniformScaled, Grid), NextNode.Get2DPoint(EGridSpace::UniformScaled, Grid));
 
 		if ((Slope > OptimalSlope) == (LoopIndex == 0))
 		{
@@ -2202,7 +2186,10 @@ void FLoopCleaner::FindBestLoopExtremity()
 			}
 			CompareWithSlopeAt(Node);
 		}
-		BestStartNodeOfLoops.Add(BestNode);
+		if(BestNode)
+		{
+			BestStartNodeOfLoops.Add(BestNode);
+		}
 
 		// init for next loop
 		UMin = HUGE_VAL;
@@ -2291,8 +2278,6 @@ EOrientation FLoopCleaner::GetLoopOrientation(const FLoopNode* StartNode)
 
 	const FLoopNode* ExtremityNodes[4] = { nullptr, nullptr, nullptr, nullptr };
 
-	const FLoopNode* BestNode = nullptr;
-
 	LoopIndex = StartNode->GetLoopIndex();
 	double OptimalSlope = (LoopIndex == 0) ? -1 : 9.;
 
@@ -2302,11 +2287,18 @@ EOrientation FLoopCleaner::GetLoopOrientation(const FLoopNode* StartNode)
 #ifdef DEBUG_LOOP_ORIENTATION
 	F3DDebugSession _(Grid.bDisplay, TEXT("GetLoopOrientation"));
 #endif
-	TFunction<void(const FLoopNode*)> CompareWithSlopeAt = [&](const FLoopNode* Node)
+
+	TFunction<void(int32)> CompareWithSlopeAt = [&ExtremityNodes, &MaxFrontSlope , &MinBackSlope, &Grid = Grid](int32 Index)
 	{
-		FLoopNode& PreviousNode = Node->GetPreviousNode();
-		FLoopNode& NextNode = Node->GetNextNode();
-		double Slope = ComputePositiveSlope(Node->Get2DPoint(EGridSpace::UniformScaled, Grid), PreviousNode.Get2DPoint(EGridSpace::UniformScaled, Grid), NextNode.Get2DPoint(EGridSpace::UniformScaled, Grid));
+		const FLoopNode* Node = ExtremityNodes[Index];
+		const FLoopNode& PreviousNode = Node->GetPreviousNode();
+		const FLoopNode& NextNode = Node->GetNextNode();
+
+		const FPoint2D Point = Node->Get2DPoint(EGridSpace::UniformScaled, Grid);
+		const FPoint2D PreviousPoint = PreviousNode.Get2DPoint(EGridSpace::UniformScaled, Grid);
+		const FPoint2D NextPoint = NextNode.Get2DPoint(EGridSpace::UniformScaled, Grid);
+
+		double Slope = ComputePositiveSlope(Point, PreviousPoint, NextPoint);
 
 #ifdef DEBUG_LOOP_ORIENTATION
 		F3DDebugSession A(Grid.bDisplay, FString::Printf(TEXT("Slope %f"), Slope));
@@ -2315,7 +2307,40 @@ EOrientation FLoopCleaner::GetLoopOrientation(const FLoopNode* StartNode)
 		Grid.DisplayIsoSegment(EGridSpace::UniformScaled, *Node, NextNode);
 #endif
 
-		if (Slope > 4)
+		if (FMath::IsNearlyEqual(Slope, 4.))
+		{
+			switch (Index)
+			{
+			case 0:
+			{
+				// UMax
+				Slope += PreviousPoint.V < NextPoint.V ? 0.1 : -0.1;
+				break;
+			}
+			case 1:
+			{
+				// UMin
+				Slope += PreviousPoint.V > NextPoint.V ? 0.1 : -0.1;
+				break;
+			}
+			case 2:
+			{
+				// VMax
+				Slope += PreviousPoint.U > NextPoint.U ? 0.1 : -0.1;
+				break;
+			}
+			case 3:
+			{
+				// VMin
+				Slope += PreviousPoint.U < NextPoint.U ? 0.1 : -0.1;
+				break;
+			}
+			default:
+				break;
+			}
+		}
+		
+		if (Slope > 4.)
 		{
 			if (MaxFrontSlope < Slope) MaxFrontSlope = Slope;
 		}
@@ -2325,27 +2350,7 @@ EOrientation FLoopCleaner::GetLoopOrientation(const FLoopNode* StartNode)
 		}
 	};
 
-	TFunction<void()> FindLoopExtremity = [&]()
-	{
-		BestNode = nullptr;
-
-		for (const FLoopNode* Node : ExtremityNodes)
-		{
-			CompareWithSlopeAt(Node);
-		}
-
-		// init for next loop
-		UMin = HUGE_VAL;
-		UMax = -HUGE_VAL;
-		VMin = HUGE_VAL;
-		VMax = -HUGE_VAL;
-
-		for (const FLoopNode*& Node : ExtremityNodes)
-		{
-			Node = nullptr;
-		}
-	};
-
+	// Find extremities
 	const FLoopNode* Node = StartNode;
 	do
 	{
@@ -2375,7 +2380,10 @@ EOrientation FLoopCleaner::GetLoopOrientation(const FLoopNode* StartNode)
 		Node = LoopCleanerImpl::GetNextConstNodeImpl(Node);
 	} while (Node != StartNode);
 
-	FindLoopExtremity();
+	for (int32 Index = 0; Index<4; ++Index)
+	{
+		CompareWithSlopeAt(Index);
+	}
 
 	MaxFrontSlope = 8. - MaxFrontSlope;
 
diff --git a/Engine/Source/Runtime/Datasmith/CADKernel/Private/CADKernel/Mesh/Structure/ThinZone2DFinder.cpp b/Engine/Source/Runtime/Datasmith/CADKernel/Private/CADKernel/Mesh/Structure/ThinZone2DFinder.cpp
--- a/Engine/Source/Runtime/Datasmith/CADKernel/Private/CADKernel/Mesh/Structure/ThinZone2DFinder.cpp
+++ b/Engine/Source/Runtime/Datasmith/CADKernel/Private/CADKernel/Mesh/Structure/ThinZone2DFinder.cpp
@@ -1940,6 +1940,10 @@ void FThinZoneSide::GetExistingMeshNodes(const FTopologicalFace& Face, FModelMes
 						if (ActiveEdge->IsMeshed())
 						{
 							EdgeElementLength = Edge->GetMesh()->GetElementLengths();
+							if (EdgeElementLength.Num() + 1 != EdgeMeshNodeCount)
+							{
+								EdgeElementLength = Edge->GetPreElementLengths();
+							}
 						}
 						else
 						{
diff --git a/Engine/Source/Runtime/Datasmith/CADKernel/Private/CADKernel/Topo/FaceAnalyzer.cpp b/Engine/Source/Runtime/Datasmith/CADKernel/Private/CADKernel/Topo/FaceAnalyzer.cpp
--- a/Engine/Source/Runtime/Datasmith/CADKernel/Private/CADKernel/Topo/FaceAnalyzer.cpp
+++ b/Engine/Source/Runtime/Datasmith/CADKernel/Private/CADKernel/Topo/FaceAnalyzer.cpp
@@ -117,6 +117,11 @@ void FFaceAnalyzer::FindClosedSegments(Topo::FThinFaceContext& Context)
 
 void FFaceAnalyzer::Analyze(Topo::FThinFaceContext& Context)
 {
+	if (Context.LoopSegments.IsEmpty())
+	{
+		return;
+	}
+
 	const FTopologicalEdge* Edge = Context.LoopSegments[0]->GetEdge();
 	double MaxSquareDistance = 0;
 	double MedSquareDistance = 0;
diff --git a/Engine/Source/Runtime/Datasmith/CADKernel/Private/CADKernel/Topo/TopologicalEdge.cpp b/Engine/Source/Runtime/Datasmith/CADKernel/Private/CADKernel/Topo/TopologicalEdge.cpp
--- a/Engine/Source/Runtime/Datasmith/CADKernel/Private/CADKernel/Topo/TopologicalEdge.cpp
+++ b/Engine/Source/Runtime/Datasmith/CADKernel/Private/CADKernel/Topo/TopologicalEdge.cpp
@@ -871,7 +871,10 @@ const FTopologicalEdge* FTopologicalEdge::GetPreMeshedTwin() const
 
 void FTopologicalEdge::AddImposedCuttingPointU(const double ImposedCuttingPointU, int32 OppositeNodeIndex, const double DeltaU)
 {
-	ImposedCuttingPointUs.Emplace(ImposedCuttingPointU, OppositeNodeIndex, DeltaU);
+	if(Boundary.Contains(ImposedCuttingPointU))
+	{
+		ImposedCuttingPointUs.Emplace(ImposedCuttingPointU, OppositeNodeIndex, DeltaU);
+	}
 }
 
 void FTopologicalEdge::AddTwinsCuttingPoint(double Coord, double DeltaU)
@@ -892,7 +895,7 @@ void FTopologicalEdge::TransferCuttingPointFromMeshedEdge(bool bOnlyWithOpposite
 {
 	const FTopologicalEdge* PreMeshedTwin = GetPreMeshedTwin();
 
-	if ((PreMeshedTwin == nullptr) || (PreMeshedTwin == this) /*|| (PreMeshedTwin->GetFace() == GetFace())*/ )
+	if ((PreMeshedTwin == nullptr) || (PreMeshedTwin == this))
 	{
 		return;
 	}
@@ -1348,7 +1351,6 @@ FTopologicalVertex* FTopologicalEdge::SplitAt(double SplittingCoordinate, const
 	if (GetTwinEntityCount() > 1)
 	{
 		return nullptr;
-		// TODO
 	}
 
 	TSharedRef<FTopologicalVertex> MiddelVertex = FTopologicalVertex::Make(NewVertexCoordinate);
diff --git a/Engine/Source/Runtime/Datasmith/CADKernel/Private/CADKernel/Topo/Topomaker.cpp b/Engine/Source/Runtime/Datasmith/CADKernel/Private/CADKernel/Topo/Topomaker.cpp
--- a/Engine/Source/Runtime/Datasmith/CADKernel/Private/CADKernel/Topo/Topomaker.cpp
+++ b/Engine/Source/Runtime/Datasmith/CADKernel/Private/CADKernel/Topo/Topomaker.cpp
@@ -193,6 +193,24 @@ FTopologicalVertex* SplitAndLink(FTopologicalVertex& StartVertex, FTopologicalEd
 		return &VertexToLink->GetLinkActiveEntity().Get();
 	}
 
+	TArray<FTopologicalVertex*> NewTwinVertices;
+	if (EdgeToSplit.GetTwinEntityCount() > 1)
+	{
+		TArray<FTopologicalEdge*> TwinEdges = EdgeToSplit.GetTwinEntities();
+		EdgeToSplit.UnlinkTwinEntities();
+		for (FTopologicalEdge* TwinEdge : TwinEdges)
+		{
+			if (TwinEdge != &EdgeToSplit)
+			{
+				FTopologicalVertex* NewVertex = SplitAndLink(StartVertex, EdgeToLink, *TwinEdge, SquareSewTolerance, SquareMinEdgeLength);
+				if (NewVertex)
+				{
+					NewTwinVertices.Add(NewVertex);
+				}
+			}
+		}
+	}
+
 	// JoinParallelEdges process all edges connected to startVertex (ConnectedEdges).
 	// Connected edges must remain compliant i.e. all edges of ConnectedEdges must be connected to StartVertex
 	// EdgeToSplit->SplitAt() must keep EdgeToSplit connected to StartVertex
@@ -207,10 +225,19 @@ FTopologicalVertex* SplitAndLink(FTopologicalVertex& StartVertex, FTopologicalEd
 
 	VertexToLink->Link(*NewVertex);
 	EdgeToLink.Link(EdgeToSplit);
+
+	if (NewTwinVertices.Num() > 0)
+	{
+		for (FTopologicalVertex* Vertex : NewTwinVertices)
+		{
+			NewVertex->Link(*Vertex);
+		}
+	}
+
 	return NewVertex;
 }
 
-FTopologicalVertex* StitchParallelEdgesFrom(FTopologicalVertex* Vertex, double SewTolerance, double MinEdgeLength)
+FTopologicalVertex* StitchParallelEdgesFrom(FTopologicalVertex* Vertex, double SewTolerance, double MinEdgeLength, bool bProhibitSewingEdgesOfSameFace)
 {
 	double SquareTolerance = FMath::Square(SewTolerance);
 	double SquareMinEdgeLength = FMath::Square(MinEdgeLength);
@@ -287,6 +314,27 @@ FTopologicalVertex* StitchParallelEdgesFrom(FTopologicalVertex* Vertex, double S
 				continue;
 			}
 
+			if(bProhibitSewingEdgesOfSameFace)
+			{
+				bool bVoidSelection = false;
+				for (FTopologicalEdge* FirstEdgeTwin : Edge->GetTwinEntities())
+				{
+					FTopologicalFace* FirstEdgeFace = FirstEdgeTwin->GetFace();
+					for (FTopologicalEdge* SecondEdgeTwin : SecondEdge->GetTwinEntities())
+					{
+						if (FirstEdgeFace == SecondEdgeTwin->GetFace())
+						{
+							bVoidSelection = true;
+							break;
+						}
+					}
+				}
+				if (bVoidSelection)
+				{
+					continue;
+				}
+			}
+
 			FPoint StartTangentEdge = Edge->GetTangentAt(*Vertex);
 			FPoint StartTangentOtherEdge = SecondEdge->GetTangentAt(*Vertex);
 
@@ -307,6 +355,7 @@ FTopologicalVertex* StitchParallelEdgesFrom(FTopologicalVertex* Vertex, double S
 					F3DDebugSession A(*FString::Printf(TEXT("SecondEdge %d"), SecondEdge->GetId()));
 					Display(*SecondEdge, EVisuProperty::RedCurve);
 				}
+				Wait();
 			}
 #endif
 
@@ -346,7 +395,7 @@ FTopologicalVertex* StitchParallelEdgesFrom(FTopologicalVertex* Vertex, double S
 	return nullptr;
 }
 
-void StitchParallelEdges(TArray<FTopologicalVertex*>& VerticesToProcess, double SewTolerance, double MinEdgeLength)
+void StitchParallelEdges(TArray<FTopologicalVertex*>& VerticesToProcess, double SewTolerance, double MinEdgeLength, bool bProhibitSewingEdgesOfSameFace)
 {
 	double SquareTolerance = FMath::Square(SewTolerance);
 	double SquareMinEdgeLength = FMath::Square(MinEdgeLength);
@@ -364,7 +413,7 @@ void StitchParallelEdges(TArray<FTopologicalVertex*>& VerticesToProcess, double
 
 		while (Vertex)
 		{
-			Vertex = StitchParallelEdgesFrom(Vertex, SewTolerance, MinEdgeLength);
+			Vertex = StitchParallelEdgesFrom(Vertex, SewTolerance, MinEdgeLength, bProhibitSewingEdgesOfSameFace);
 		}
 	}
 
@@ -560,6 +609,18 @@ FTopomaker::FTopomaker(FSession& InSession, const TArray<TSharedPtr<FShell>>& In
 	InitFaces();
 }
 
+FTopomaker::FTopomaker(FSession& InSession, const TArray<TSharedPtr<FTopologicalFace>>& InFaces, const FTopomakerOptions& InOptions)
+	: Session(InSession)
+{
+	SetTolerance(InOptions);
+
+	Faces.Reserve(InFaces.Num());
+	for (const TSharedPtr<FTopologicalFace>& Face : InFaces)
+	{
+		Faces.Add(Face);
+	}
+}
+
 void FTopomaker::InitFaces()
 {
 	int32 FaceCount = 0;
@@ -609,7 +670,7 @@ void FTopomaker::Sew()
 	// advance case: two partially coincident edges connected at one extremity i.e. coincident along the shortest edge
 	BorderVertices.Reset();
 	GetBorderVertices(BorderVertices);
-	TopomakerTools::StitchParallelEdges(BorderVertices, SewTolerance, EdgeLengthTolerance);
+	TopomakerTools::StitchParallelEdges(BorderVertices, SewTolerance, EdgeLengthTolerance, /*bProhibitSewingEdgesOfSameFace*/ false);
 
 	const bool bForceJoining = SewOption::IsForceJoining(SewOptions);
 	if (bForceJoining)
@@ -619,30 +680,35 @@ void FTopomaker::Sew()
 		TopomakerTools::MergeCoincidentVertices(BorderVertices, SewToleranceToForceJoin);
 
 		CheckSelfConnectedEdge(EdgeLengthTolerance, BorderVertices);
-	}
 
-	// re process with new edges from MergeUnconnectedSuccessiveEdges and new merged vertices (if bForceJoining)
-	BorderVertices.Reset();
-	GetBorderVertices(BorderVertices);
-	TopomakerTools::MergeCoincidentEdges(BorderVertices, bForceJoining ? LargeEdgeLengthTolerance : EdgeLengthTolerance);
+		// re process with new edges from MergeUnconnectedSuccessiveEdges and new merged vertices (if bForceJoining)
+		BorderVertices.Reset();
+		GetBorderVertices(BorderVertices);
+		TopomakerTools::MergeCoincidentEdges(BorderVertices, LargeEdgeLengthTolerance);
 
-	if (SewOption::IsRemoveThinFaces(SewOptions))
-	{
-		RemoveThinFaces();
+		// advance case: two partially coincident edges connected at one extremity i.e. coincident along the shortest edge
+		BorderVertices.Reset();
+		GetBorderVertices(BorderVertices);
+		TopomakerTools::StitchParallelEdges(BorderVertices, SewToleranceToForceJoin, LargeEdgeLengthTolerance, /*bProhibitSewingEdgesOfSameFace*/ true);
 	}
 
-	// advance case: two partially coincident edges connected at one extremity i.e. coincident along the shortest edge
-	BorderVertices.Reset();
-	GetBorderVertices(BorderVertices);
-
-	TopomakerTools::MergeCoincidentEdges(BorderVertices, bForceJoining ? LargeEdgeLengthTolerance : EdgeLengthTolerance);
-	TopomakerTools::StitchParallelEdges(BorderVertices, bForceJoining ? SewToleranceToForceJoin : SewTolerance, bForceJoining ? LargeEdgeLengthTolerance : EdgeLengthTolerance);
-
 	if (SewOption::IsRemoveDuplicatedFaces(SewOptions))
 	{
 		RemoveDuplicatedFaces();
 	}
 
+	if (SewOption::IsRemoveThinFaces(SewOptions))
+	{
+		RemoveThinFaces();
+		// advance case: two partially coincident edges connected at one extremity i.e. coincident along the shortest edge
+		BorderVertices.Reset();
+		GetBorderVertices(BorderVertices);
+
+		TopomakerTools::MergeCoincidentEdges(BorderVertices, bForceJoining ? LargeEdgeLengthTolerance : EdgeLengthTolerance);
+		
+		TopomakerTools::StitchParallelEdges(BorderVertices, bForceJoining ? SewToleranceToForceJoin : SewTolerance, bForceJoining ? LargeEdgeLengthTolerance : EdgeLengthTolerance, /*bProhibitSewingEdgesOfSameFace*/ true);
+	}
+
 	ResetMarkersOfFaces();
 
 #ifdef CADKERNEL_DEV
diff --git a/Engine/Source/Runtime/Datasmith/CADKernel/Public/CADKernel/Geo/Sampling/PolylineTools.h b/Engine/Source/Runtime/Datasmith/CADKernel/Public/CADKernel/Geo/Sampling/PolylineTools.h
--- a/Engine/Source/Runtime/Datasmith/CADKernel/Public/CADKernel/Geo/Sampling/PolylineTools.h
+++ b/Engine/Source/Runtime/Datasmith/CADKernel/Public/CADKernel/Geo/Sampling/PolylineTools.h
@@ -456,12 +456,12 @@ public:
 		int32 BoundaryIndices[2];
 
 		TArray<double> CurvilinearCoordinates;
-		double CurveLength = ComputeCurvilinearCoordinatesOfPolyline(InBoundary, CurvilinearCoordinates, BoundaryIndices);
+		const double CurveLength = ComputeCurvilinearCoordinatesOfPolyline(InBoundary, CurvilinearCoordinates, BoundaryIndices);
 
-		PointType StartPoint = ComputePoint(BoundaryIndices[0], InBoundary.Min);
+		const PointType StartPoint = ComputePoint(BoundaryIndices[0], InBoundary.Min);
 		double FromStartSegmentLength = PolylinePoints[BoundaryIndices[0]].Distance(StartPoint);
 
-		if (CurveLength < 2. * DesiredSegmentLength)
+		if (FMath::IsNearlyZero(DesiredSegmentLength, DOUBLE_KINDA_SMALL_NUMBER) || CurveLength < 2. * DesiredSegmentLength)
 		{
 			OutCoordinates.SetNum(3);
 			OutCoordinates[0] = InBoundary.GetMin();
@@ -472,7 +472,7 @@ public:
 
 		int32 SegmentNum = FMath::IsNearlyZero(DesiredSegmentLength) ? 2 : (int32)FMath::Max(CurveLength / DesiredSegmentLength + 0.5, 1.0);
 
-		double SectionLength = CurveLength / (double)(SegmentNum);
+		const double SectionLength = CurveLength / (double)(SegmentNum);
 
 		OutCoordinates.Empty();
 		OutCoordinates.Reserve(SegmentNum + 1);
@@ -483,26 +483,35 @@ public:
 			return Coordinate + (PolylineCoordinates[IndexCoordinate] - Coordinate) * (Length - CurvilinearCoordinates[IndexCurvilinear - 1]) / (CurvilinearCoordinates[IndexCurvilinear] - CurvilinearCoordinates[IndexCurvilinear - 1]);
 		};
 
-		double CurvilinearLength = SectionLength;
-		FromStartSegmentLength += SectionLength;
+		double CurvilinearLength = 0.;
 
 		double LastCoordinate = InBoundary.Min;
 		for (int32 IndexCurvilinear = 1, IndexCoordinate = BoundaryIndices[0] + 1; IndexCurvilinear < CurvilinearCoordinates.Num(); ++IndexCurvilinear, ++IndexCoordinate)
 		{
 			while (FromStartSegmentLength + DOUBLE_SMALL_NUMBER < CurvilinearCoordinates[IndexCurvilinear])
 			{
-				double Coordinate = ComputeSamplePointCoordinate(IndexCurvilinear, IndexCoordinate, FromStartSegmentLength, LastCoordinate);
+				const double Coordinate = ComputeSamplePointCoordinate(IndexCurvilinear, IndexCoordinate, FromStartSegmentLength, LastCoordinate);
 				OutCoordinates.Add(Coordinate);
 				CurvilinearLength += SectionLength;
 				FromStartSegmentLength += SectionLength;
-				if (OutCoordinates.Num() == SegmentNum)
+				if (CurvilinearLength > CurveLength || OutCoordinates.Num() == SegmentNum)
 				{
-					OutCoordinates.Add(InBoundary.Max);
-					break;
+					while (OutCoordinates.Last() + SMALL_NUMBER > InBoundary.GetMax())
+					{
+						OutCoordinates.Pop();
+					}
+					OutCoordinates.Add(InBoundary.GetMax());
+					return;
 				}
 			}
 			LastCoordinate = PolylineCoordinates[IndexCoordinate];
 		}
+
+		while (OutCoordinates.Last() + SMALL_NUMBER > InBoundary.GetMax())
+		{
+			OutCoordinates.Pop();
+		}
+		OutCoordinates.Add(InBoundary.GetMax());
 	}
 
 	/**
@@ -513,7 +522,22 @@ public:
 	{
 		int32 BoundaryIndices[2];
 		GetStartEndIndex(InBoundary, BoundaryIndices);
-		return ProjectPointToPolyline(BoundaryIndices, PointOnEdge, OutProjectedPoint);
+		double Coordinate = ProjectPointToPolyline(BoundaryIndices, PointOnEdge, OutProjectedPoint);
+		if (InBoundary.Contains(Coordinate))
+		{
+			return Coordinate;
+		}
+
+		if (Coordinate < InBoundary.GetMin())
+		{
+			Coordinate = InBoundary.GetMin();
+		}
+		else if (Coordinate > InBoundary.GetMax())
+		{
+			Coordinate = InBoundary.GetMax();
+		}
+		OutProjectedPoint = ApproximatePoint(Coordinate);
+		return Coordinate;
 	}
 
 	/**
diff --git a/Engine/Source/Runtime/Datasmith/CADKernel/Public/CADKernel/Math/MathConst.h b/Engine/Source/Runtime/Datasmith/CADKernel/Public/CADKernel/Math/MathConst.h
--- a/Engine/Source/Runtime/Datasmith/CADKernel/Public/CADKernel/Math/MathConst.h
+++ b/Engine/Source/Runtime/Datasmith/CADKernel/Public/CADKernel/Math/MathConst.h
@@ -68,6 +68,11 @@ uint8 ToUInt8(T Value)
  */
 inline double WrapTo(double Slope, const double StartOfPeriod, const double EndOfPeriod, const double PeriodLength)
 {
+	if (FMath::Abs(Slope) > DOUBLE_BIG_NUMBER)
+	{
+		return 0;
+	}
+
 	while (Slope < StartOfPeriod)
 	{
 		Slope += PeriodLength;
diff --git a/Engine/Source/Runtime/Datasmith/CADKernel/Public/CADKernel/Math/SlopeUtils.h b/Engine/Source/Runtime/Datasmith/CADKernel/Public/CADKernel/Math/SlopeUtils.h
--- a/Engine/Source/Runtime/Datasmith/CADKernel/Public/CADKernel/Math/SlopeUtils.h
+++ b/Engine/Source/Runtime/Datasmith/CADKernel/Public/CADKernel/Math/SlopeUtils.h
@@ -11,19 +11,19 @@ namespace UE::CADKernel
 
 /**
  * "Slope" is a fast angle approximation.
- * 
+ *
  * This file propose all useful methods to use slope instead of angle.
- * 
- * The method "ComputeSlope(const FPoint2D& StartPoint, const FPoint2D& EndPoint) is the main method. 
+ *
+ * The method "ComputeSlope(const FPoint2D& StartPoint, const FPoint2D& EndPoint) is the main method.
  * It compute the slope between the input segment defined by two points and [0, u) axis.
  * The return value is a real in the interval [0, 8] for an angle in the interval [0, 2Pi]
- * 
+ *
  * Warning, it's only an approximation... The conversion is not linear but the error is small near the integer value of slope (0, 1, 2, 3, ...8)
  *
  * To compute an angle value between two segments, the call of acos (and asin for an oriented angle) is necessary while with this approximation, only a division is useful.
- * 
+ *
  * This approximation is very good when only comparison of angles is needed and more faster than acos and/or asin i.e. Slope approximation need only a division and few addition and test
- * 
+ *
  * [0 - 2Pi] is divide into 8 angular sector i.e. [0, Pi/4] = [0,1], [Pi/4, Pi/2] = [1,2], ...
  *
  * The value of the slope for an angle in [0, Pi/4] = tan(angle)
@@ -84,12 +84,12 @@ constexpr double TwoPiSlope = 8.;
  * ThirdPiSlope i.e. Pi/3 angle (60 deg)
  */
 constexpr double ThirdPiSlope = 1.422649730810374235490851219498;
-constexpr double SixtySlope = 1.422649730810374235490851219498;
+constexpr double SixtySlope   = 1.422649730810374235490851219498;
 
 /**
  * ThirdPiSlope i.e. Pi/4 angle (45 deg)
  */
-constexpr double QuaterPiSlope = 1;
+constexpr double QuaterPiSlope  = 1;
 constexpr double FortyFiveSlope = 0.57735026918962576450914878050196;
 
 /**
@@ -103,12 +103,12 @@ constexpr double ThirtySlope  = 0.57735026918962576450914878050196;
  */
 constexpr double ThreeQuaterPiSlope = 3;
                               
-constexpr double OneDegree  = 0.01745506492821758576512889521973;
-constexpr double TwoDegree  = 0.03492076949174773050040262577373;
-constexpr double FiveDegree = 0.08748866352592400522201866943496;
-constexpr double TenDegree = 0.17632698070846497347109038686862;
-constexpr double FifteenDegree = 0.26794919243112270647255365849413;
-constexpr double TwentyDegree = 0.36397023426620236135104788277683;
+constexpr double OneDegree        = 0.01745506492821758576512889521973;
+constexpr double TwoDegree        = 0.03492076949174773050040262577373;
+constexpr double FiveDegree       = 0.08748866352592400522201866943496;
+constexpr double TenDegree        = 0.17632698070846497347109038686862;
+constexpr double FifteenDegree    = 0.26794919243112270647255365849413;
+constexpr double TwentyDegree     = 0.36397023426620236135104788277683;
 constexpr double TwentyFiveDegree = 0.46630765815499859283000619479956;
 
 constexpr double Epsilon = 0.001;
@@ -207,13 +207,17 @@ inline double ComputeSlope(const FPoint2D& StartPoint, const FPoint2D& EndPoint)
 			{
 				Delta = 8. + DeltaV / DeltaU;
 			}
-			else
+			else if (fabs(DeltaV) > DOUBLE_SMALL_NUMBER)
 			{
 				Delta = 6. - DeltaU / DeltaV; // deltaU/deltaV <0
 			}
+			else
+			{
+				Delta = 8.;
+			}
 		}
 	}
-	else
+	else if (-DeltaU > DOUBLE_SMALL_NUMBER)
 	{
 		if (DeltaV > DOUBLE_SMALL_NUMBER)
 		{
@@ -232,10 +236,25 @@ inline double ComputeSlope(const FPoint2D& StartPoint, const FPoint2D& EndPoint)
 			{
 				Delta = 4. + DeltaV / DeltaU;
 			}
-			else
+			else if (fabs(DeltaV) > DOUBLE_SMALL_NUMBER)
 			{
 				Delta = 6. - DeltaU / DeltaV;
 			}
+			else
+			{
+				Delta = 4.;
+			}
+		}
+	}
+	else
+	{
+		if (DeltaV > 0)
+		{
+			Delta = 2.;
+		}
+		else
+		{
+			Delta = 6.;
 		}
 	}
 
@@ -409,7 +428,8 @@ inline bool ArePointsInsideSectorABC(const FPoint2D& PointA, const FPoint2D& Poi
 	double SlopWithNextBoundary = ComputeSlope(PointB, PointC);
 	double BoundaryDeltaSlope = ComputePositiveSlope(PointB, PointA, SlopWithNextBoundary);
 
-	return Algo::AllOf(Points, [&](const FPoint2D* PointP) {
+	return Algo::AllOf(Points, [&](const FPoint2D* PointP) 
+	{
 		double DeltaU = PointB.U - PointP->U;
 		double DeltaV = PointB.V - PointP->V;
 
@@ -424,7 +444,7 @@ inline bool ArePointsInsideSectorABC(const FPoint2D& PointA, const FPoint2D& Poi
 			return false;
 		}
 		return true;
-		});
+	});
 }
 
 inline FPoint2D SlopeToVector(const double Slope)
diff --git a/Engine/Source/Runtime/Datasmith/CADKernel/Public/CADKernel/Topo/Topomaker.h b/Engine/Source/Runtime/Datasmith/CADKernel/Public/CADKernel/Topo/Topomaker.h
--- a/Engine/Source/Runtime/Datasmith/CADKernel/Public/CADKernel/Topo/Topomaker.h
+++ b/Engine/Source/Runtime/Datasmith/CADKernel/Public/CADKernel/Topo/Topomaker.h
@@ -100,6 +100,7 @@ public:
 
 	FTopomaker(FSession& InSession, const FTopomakerOptions& InOptions);
 	FTopomaker(FSession& InSession, const TArray<TSharedPtr<FShell>>& Shells, const FTopomakerOptions& InOptions);
+	FTopomaker(FSession& InSession, const TArray<TSharedPtr<FTopologicalFace>>& Faces, const FTopomakerOptions& InOptions);
 
 	void SetTolerance(const FTopomakerOptions Options)
 	{
@@ -120,10 +121,9 @@ public:
 		SewToleranceToForceJoin = Tolerance * ForceJoinFactor;
 		LargeEdgeLengthTolerance = SewToleranceToForceJoin * 2.;
 
-		ThinFaceWidth = Tolerance * ForceJoinFactor;
+		ThinFaceWidth = Tolerance * ForceJoinFactor * 0.5;
 	}
 
-
 	void Sew();
 
 	/**
