[NO NON-PII COMMIT MESSAGE]

diff --git a/Engine/Source/Runtime/Datasmith/CADKernel/Private/CADKernel/Mesh/Structure/LoopCleaner.cpp b/Engine/Source/Runtime/Datasmith/CADKernel/Private/CADKernel/Mesh/Structure/LoopCleaner.cpp
--- a/Engine/Source/Runtime/Datasmith/CADKernel/Private/CADKernel/Mesh/Structure/LoopCleaner.cpp
+++ b/Engine/Source/Runtime/Datasmith/CADKernel/Private/CADKernel/Mesh/Structure/LoopCleaner.cpp
@@ -131,9 +131,35 @@ bool FLoopCleaner::CleanLoops()
 
 #ifdef DEBUG_CLEAN_LOOPS		
 		Grid.DisplayGridPolyline(TEXT("LoopIntersections: remove loop's picks"), EGridSpace::UniformScaled, NodesOfLoop, true, EVisuProperty::YellowCurve);
-		Wait(bDisplay);
+		if(bDisplay)
+		{
+			Wait(bDisplay);
+		}
 #endif
-		
+
+		FLoopNode** StartNodePtr = NodesOfLoop.FindByPredicate([](FLoopNode* Node) { return !Node->IsDelete(); });
+		if (StartNodePtr == nullptr || *StartNodePtr == nullptr || (*StartNodePtr)->IsDelete())
+		{
+			continue;
+		}
+
+		StartNode = *StartNodePtr;
+
+		if (!FindAndRemoveCoincidence(StartNode))
+		{
+			if (StartSegmentIndex == 0)
+			{
+				FMessage::Printf(Log, TEXT("The outer loop of the surface %d is degenerated. The mesh of this surface is canceled.\n"), Grid.GetFace().GetId());
+				return false;
+			}
+			continue;
+		}
+
+#ifdef DEBUG_CLEAN_LOOPS		
+		Grid.DisplayGridPolyline(TEXT("Loop: start"), EGridSpace::UniformScaled, NodesOfLoop, true, EVisuProperty::YellowCurve);
+		//Wait(bDisplay);
+#endif
+
 		FindLoopIntersections();
 
 		if(Intersections.Num())
@@ -687,7 +713,11 @@ bool FLoopCleaner::RemoveSelfIntersectionsOfLoop()
 
 #ifdef DEBUG_REMOVE_LOOP_INTERSECTIONS		
 		F3DDebugSession _(bDisplay, TEXT("Intersected Segments"));
-		DisplayIntersection(Intersection);
+		if(bDisplay)
+		{
+			DisplayIntersection(Intersection);
+			Wait();
+		}
 #endif
 
 		bool bIntersectionForward = true;
diff --git a/Engine/Source/Runtime/Datasmith/CADKernel/Public/CADKernel/Geo/Sampling/PolylineTools.h b/Engine/Source/Runtime/Datasmith/CADKernel/Public/CADKernel/Geo/Sampling/PolylineTools.h
--- a/Engine/Source/Runtime/Datasmith/CADKernel/Public/CADKernel/Geo/Sampling/PolylineTools.h
+++ b/Engine/Source/Runtime/Datasmith/CADKernel/Public/CADKernel/Geo/Sampling/PolylineTools.h
@@ -723,7 +723,7 @@ public:
 			BoundaryIndices[0]++;
 		}
 
-		if (FMath::IsNearlyEqual(PolylineCoordinates[BoundaryIndices[1]], InBoundary.Max, (double)DOUBLE_SMALL_NUMBER))
+		if (BoundaryIndices[1]> 0 && FMath::IsNearlyEqual(PolylineCoordinates[BoundaryIndices[1]], InBoundary.Max, (double)DOUBLE_SMALL_NUMBER))
 		{
 			BoundaryIndices[1]--;
 		}
