Fixed crash when loading old udatasmith files saved before 5.1

#rnx

diff --git a/Engine/Source/Runtime/Datasmith/DatasmithCore/Private/DatasmithMeshUObject.cpp b/Engine/Source/Runtime/Datasmith/DatasmithCore/Private/DatasmithMeshUObject.cpp
--- a/Engine/Source/Runtime/Datasmith/DatasmithCore/Private/DatasmithMeshUObject.cpp
+++ b/Engine/Source/Runtime/Datasmith/DatasmithCore/Private/DatasmithMeshUObject.cpp
@@ -9,6 +9,7 @@ void FDatasmithMeshSourceModel::SerializeBulkData(FArchive& Ar, UObject* Owner)
 
 void UDatasmithMesh::Serialize(FArchive& Ar)
 {
+	Ar.SetUEVer(FPackageFileVersion::CreateUE4Version(VER_UE4_AUTOMATIC_VERSION));
 	Super::Serialize( Ar );
 
 	for ( FDatasmithMeshSourceModel& SourceModel : SourceModels )
