Fix Datasmith old FRawMesh serialization
#rnx

diff --git a/Engine/Source/Runtime/Datasmith/DatasmithCore/Private/DatasmithMeshSerialization.cpp b/Engine/Source/Runtime/Datasmith/DatasmithCore/Private/DatasmithMeshSerialization.cpp
--- a/Engine/Source/Runtime/Datasmith/DatasmithCore/Private/DatasmithMeshSerialization.cpp
+++ b/Engine/Source/Runtime/Datasmith/DatasmithCore/Private/DatasmithMeshSerialization.cpp
@@ -314,6 +314,7 @@ TArray<FDatasmithMeshModels> GetDatasmithMeshFromMeshPath_Legacy(FArchive* Archi
 		MemoryReader.ArIgnoreClassRef = false;
 		MemoryReader.ArIgnoreArchetypeRef = false;
 		MemoryReader.SetWantBinaryPropertySerialization(true);
+		MemoryReader.SetUEVer(FPackageFileVersion::CreateUE4Version(EUnrealEngineObjectUE4Version::VER_UE4_AUTOMATIC_VERSION));
 		DatasmithMesh->Serialize( MemoryReader );
 
 		FDatasmithMeshModels& MeshInternal = Result.AddDefaulted_GetRef();
diff --git a/Engine/Source/Runtime/Datasmith/DatasmithCore/Private/DatasmithMeshUObject.cpp b/Engine/Source/Runtime/Datasmith/DatasmithCore/Private/DatasmithMeshUObject.cpp
--- a/Engine/Source/Runtime/Datasmith/DatasmithCore/Private/DatasmithMeshUObject.cpp
+++ b/Engine/Source/Runtime/Datasmith/DatasmithCore/Private/DatasmithMeshUObject.cpp
@@ -9,7 +9,6 @@ void FDatasmithMeshSourceModel::SerializeBulkData(FArchive& Ar, UObject* Owner)
 
 void UDatasmithMesh::Serialize(FArchive& Ar)
 {
-	Ar.SetUEVer(FPackageFileVersion::CreateUE4Version(VER_UE4_AUTOMATIC_VERSION));
 	Super::Serialize( Ar );
 
 	for ( FDatasmithMeshSourceModel& SourceModel : SourceModels )
