Fixed crash when loading a specific wire file

#rnx

diff --git a/Engine/Source/Runtime/Datasmith/CADKernel/Private/CADKernel/Topo/TopologicalFace.cpp b/Engine/Source/Runtime/Datasmith/CADKernel/Private/CADKernel/Topo/TopologicalFace.cpp
--- a/Engine/Source/Runtime/Datasmith/CADKernel/Private/CADKernel/Topo/TopologicalFace.cpp
+++ b/Engine/Source/Runtime/Datasmith/CADKernel/Private/CADKernel/Topo/TopologicalFace.cpp
@@ -222,29 +222,37 @@ void FTopologicalFace::ApplyNaturalLoops(const FSurfacicBoundary& Boundaries)
 
 	const bool bIsExternalLoop = true;
 	TSharedPtr<FTopologicalLoop> Loop = FTopologicalLoop::Make(Edges, Orientations, bIsExternalLoop, CarrierSurface->Get3DTolerance());
-	AddLoop(Loop);
+	if (Loop)
+	{
+		Loop->SetSurface(this);
+		Loops.Add(Loop);
+	}
 }
 
 void FTopologicalFace::AddLoops(const TArray<TSharedPtr<FTopologicalLoop>>& InLoops, int32& DoubtfulLoopOrientationCount)
 {
 	for (TSharedPtr<FTopologicalLoop> Loop : InLoops)
 	{
-		AddLoop(Loop);
-	}
-
-	for (TSharedPtr<FTopologicalLoop> Loop : InLoops)
-	{
-		if (!Loop->Orient())
+		if (Loop)
 		{
-			DoubtfulLoopOrientationCount++;
+			Loop->SetSurface(this);
+			Loops.Add(Loop);
+
+			if (Loop->Orient())
+			{
+				DoubtfulLoopOrientationCount++;
+			}
 		}
 	}
 }
 
 void FTopologicalFace::AddLoop(const TSharedPtr<FTopologicalLoop>& InLoop)
 {
-	InLoop->SetSurface(this);
-	Loops.Add(InLoop);
+	if (InLoop.IsValid())
+	{
+		InLoop->SetSurface(this);
+		Loops.Add(InLoop);
+	}
 }
 
 void FTopologicalFace::RemoveLoop(const TSharedPtr<FTopologicalLoop>& Loop)
