Improved CADKernel performance.
   Removed a very frequent call to RemoveAt during the tessellation process.

#rnx

diff --git a/Engine/Source/Runtime/Datasmith/CADKernel/Private/CADKernel/Geo/Sampler/SamplerAbstract.h b/Engine/Source/Runtime/Datasmith/CADKernel/Private/CADKernel/Geo/Sampler/SamplerAbstract.h
--- a/Engine/Source/Runtime/Datasmith/CADKernel/Private/CADKernel/Geo/Sampler/SamplerAbstract.h
+++ b/Engine/Source/Runtime/Datasmith/CADKernel/Private/CADKernel/Geo/Sampler/SamplerAbstract.h
@@ -101,7 +101,7 @@ protected:
 		CandidatePoints.Empty(100);
 
 		// Initialization of the algorithm with the not derivable coordinates of the curve and at least 5 temporary points
-		int32 ComplementatyPointCount = 0;
+		int32 ComplementaryPointOffset = 0;
 		{
 			TArray<double> NotDerivableCoordinates;
 			GetNotDerivableCoordinates(NotDerivableCoordinates);
@@ -109,34 +109,32 @@ protected:
 			NotDerivableCoordinates.Insert(Boundary.Min, 0);
 			NotDerivableCoordinates.Add(Boundary.Max);
 
-			ComplementatyPointCount = NotDerivableCoordinates.Num() < 10 ? 10 : 1;
+			ComplementaryPointOffset = NotDerivableCoordinates.Num() < 10 ? 10 : 1;
 
-			NextCoordinates.Empty(NotDerivableCoordinates.Num() * (ComplementatyPointCount + 1));
+			NextCoordinates.Empty(NotDerivableCoordinates.Num() * (ComplementaryPointOffset + 1));
 			NextCoordinates.Add(Boundary.Min);
 
 			for (int32 Index = 0; Index < NotDerivableCoordinates.Num() - 1; ++Index)
 			{
-				AddIntermediateCoordinates(NotDerivableCoordinates[Index], NotDerivableCoordinates[Index + 1], ComplementatyPointCount);
+				AddIntermediateCoordinates(NotDerivableCoordinates[Index], NotDerivableCoordinates[Index + 1], ComplementaryPointOffset);
 				NextCoordinates.Add(NotDerivableCoordinates[Index + 1]);
 			}
 		}
 
+		// #cadkernel_check: Could we do a MoveTem?
 		CandidatePoints.SwapCoordinates(NextCoordinates);
 		EvaluatesNewCandidatePoints();
 
 		// Not Derivable point initialize the sampling
-		ComplementatyPointCount++;
-		for (int32 Index = 0, ISampling = 0; Index < CandidatePoints.Size(); Index += ComplementatyPointCount, ++ISampling)
+		ComplementaryPointOffset++;
+		for (int32 Index = 0, ISampling = 0; Index < CandidatePoints.Size(); Index += ComplementaryPointOffset, ++ISampling)
 		{
 			Sampling.EmplaceAt(ISampling, CandidatePoints, Index);
 			IsOptimalSegments.Add(false);
 		}
 		IsOptimalSegments.Pop();
 
-		for (int32 Index = CandidatePoints.Size() - 1; Index >= 0; Index -= ComplementatyPointCount)
-		{
-			CandidatePoints.RemoveAt(Index);
-		}
+		CandidatePoints.RemoveComplementaryPoints(ComplementaryPointOffset);
 
 		// first segment is not optimal
 
diff --git a/Engine/Source/Runtime/Datasmith/CADKernel/Private/CADKernel/Geo/Sampler/SamplerOnParam.h b/Engine/Source/Runtime/Datasmith/CADKernel/Private/CADKernel/Geo/Sampler/SamplerOnParam.h
--- a/Engine/Source/Runtime/Datasmith/CADKernel/Private/CADKernel/Geo/Sampler/SamplerOnParam.h
+++ b/Engine/Source/Runtime/Datasmith/CADKernel/Private/CADKernel/Geo/Sampler/SamplerOnParam.h
@@ -324,21 +324,21 @@ public:
 		CandidatePoints.Empty(100);
 
 		// Initialization of the algorithm with the not derivable coordinates of the curve and at least 5 temporary points
-		int32 ComplementatyPointCount = 0;
+		int32 ComplementaryPointOffset = 0;
 		{
 			TArray<double> LocalNotDerivableCoordinates;
 			GetNotDerivableCoordinates(LocalNotDerivableCoordinates);
 			LocalNotDerivableCoordinates.Insert(Boundary.Min, 0);
 			LocalNotDerivableCoordinates.Add(Boundary.Max);
 
-			ComplementatyPointCount = LocalNotDerivableCoordinates.Num() < 5 ? 5 : 1;
+			ComplementaryPointOffset = LocalNotDerivableCoordinates.Num() < 5 ? 5 : 1;
 
-			NextCoordinates.Empty(LocalNotDerivableCoordinates.Num() * (ComplementatyPointCount + 1));
+			NextCoordinates.Empty(LocalNotDerivableCoordinates.Num() * (ComplementaryPointOffset + 1));
 			NextCoordinates.Add(Boundary.Min);
 
 			for (int32 Index = 0; Index < LocalNotDerivableCoordinates.Num() - 1; ++Index)
 			{
-				AddIntermediateCoordinates(LocalNotDerivableCoordinates[Index], LocalNotDerivableCoordinates[Index + 1], ComplementatyPointCount);
+				AddIntermediateCoordinates(LocalNotDerivableCoordinates[Index], LocalNotDerivableCoordinates[Index + 1], ComplementaryPointOffset);
 				NextCoordinates.Add(LocalNotDerivableCoordinates[Index + 1]);
 			}
 		}
@@ -347,18 +347,15 @@ public:
 		EvaluatesNewCandidatePoints();
 
 		// Not Derivable point initialize the sampling
-		ComplementatyPointCount++;
-		for (int32 Index = 0, ISampling = 0; Index < CandidatePoints.Size(); Index += ComplementatyPointCount, ++ISampling)
+		ComplementaryPointOffset++;
+		for (int32 Index = 0, ISampling = 0; Index < CandidatePoints.Size(); Index += ComplementaryPointOffset, ++ISampling)
 		{
 			Sampling.EmplaceAt(ISampling, CandidatePoints, Index);
 			IsOptimalSegments.Add(false);
 		}
 		IsOptimalSegments.Pop();
 
-		for (int32 Index = CandidatePoints.Size() - 1; Index >= 0; Index -= ComplementatyPointCount)
-		{
-			CandidatePoints.RemoveAt(Index);
-		}
+		CandidatePoints.RemoveComplementaryPoints(ComplementaryPointOffset);
 
 		// first segment is not optimal
 
diff --git a/Engine/Source/Runtime/Datasmith/CADKernel/Public/CADKernel/Geo/Sampling/Polyline.h b/Engine/Source/Runtime/Datasmith/CADKernel/Public/CADKernel/Geo/Sampling/Polyline.h
--- a/Engine/Source/Runtime/Datasmith/CADKernel/Public/CADKernel/Geo/Sampling/Polyline.h
+++ b/Engine/Source/Runtime/Datasmith/CADKernel/Public/CADKernel/Geo/Sampling/Polyline.h
@@ -154,10 +154,31 @@ public:
 		Points.EmplaceAt(Index, Polyline.Points[PointIndex]);
 	}
 
-	void RemoveAt(int32 Index)
+
+	void RemoveComplementaryPoints(int32 Offset)
 	{
-		Coordinates.RemoveAt(Index);
-		Points.RemoveAt(Index);
+		const int32 Count = Points.Num();
+
+		TBitArray<> Markers(true, Count);
+
+		for (int32 Index = Count - 1; Index >= 0; Index -= Offset)
+		{
+			Markers[Index] = false;
+		}
+
+		int32 CurrentIndex = 0;
+		for (int32 Index = 0; Index < Count; ++Index)
+		{
+			if (Markers[Index])
+			{
+				Coordinates[CurrentIndex] = Coordinates[Index];
+				Points[CurrentIndex] = Points[Index];
+				++CurrentIndex;
+			}
+		}
+
+		Coordinates.SetNum(CurrentIndex, EAllowShrinking::No);
+		Points.SetNum(CurrentIndex, EAllowShrinking::No);
 	}
 
 	void Pop()
diff --git a/Engine/Source/Runtime/Datasmith/CADKernel/Public/CADKernel/Geo/Sampling/SurfacicPolyline.h b/Engine/Source/Runtime/Datasmith/CADKernel/Public/CADKernel/Geo/Sampling/SurfacicPolyline.h
--- a/Engine/Source/Runtime/Datasmith/CADKernel/Public/CADKernel/Geo/Sampling/SurfacicPolyline.h
+++ b/Engine/Source/Runtime/Datasmith/CADKernel/Public/CADKernel/Geo/Sampling/SurfacicPolyline.h
@@ -394,18 +394,47 @@ public:
 		}
 	}
 
-	void RemoveAt(int32 Index)
+	void RemoveComplementaryPoints(int32 Offset)
 	{
-		Coordinates.RemoveAt(Index);
-		Points2D.RemoveAt(Index);
-		Points3D.RemoveAt(Index);
+		const int32 Count = Points2D.Num();
+
+		TBitArray<> Markers(true, Count);
+
+		for (int32 Index = Count - 1; Index >= 0; Index -= Offset)
+		{
+			Markers[Index] = false;
+		}
+
+		int32 CurrentIndex = 0;
+		for (int32 Index = 0; Index < Count; ++Index)
+		{
+			if (Markers[Index])
+			{
+				Coordinates[CurrentIndex] = Coordinates[Index];
+				Points2D[CurrentIndex] = Points2D[Index];
+				Points3D[CurrentIndex] = Points3D[Index];
+				if (bWithNormals)
+				{
+					Normals[CurrentIndex] = Normals[Index];
+				}
+				if (bWithTangent)
+				{
+					Tangents[CurrentIndex] = Tangents[Index];
+				}
+				++CurrentIndex;
+			}
+		}
+
+		Coordinates.SetNum(CurrentIndex, EAllowShrinking::No);
+		Points2D.SetNum(CurrentIndex, EAllowShrinking::No);
+		Points3D.SetNum(CurrentIndex, EAllowShrinking::No);
 		if (bWithNormals)
 		{
-			Normals.RemoveAt(Index);
+			Normals.SetNum(CurrentIndex, EAllowShrinking::No);
 		}
 		if (bWithTangent)
 		{
-			Tangents.RemoveAt(Index);
+			Tangents.SetNum(CurrentIndex, EAllowShrinking::No);
 		}
 	}
 
