Performance improvements: Part #1
  - Improved bezier curve implementation to support their use as curve on surfaces
  - Fixed a couple of rough edges in the logic

#rnx

diff --git a/Engine/Source/Runtime/Datasmith/CADKernel/Private/CADKernel/Geo/Curves/BezierCurve.cpp b/Engine/Source/Runtime/Datasmith/CADKernel/Private/CADKernel/Geo/Curves/BezierCurve.cpp
--- a/Engine/Source/Runtime/Datasmith/CADKernel/Private/CADKernel/Geo/Curves/BezierCurve.cpp
+++ b/Engine/Source/Runtime/Datasmith/CADKernel/Private/CADKernel/Geo/Curves/BezierCurve.cpp
@@ -8,35 +8,159 @@
 namespace UE::CADKernel
 {
 
-void FBezierCurve::EvaluatePoint(double Coordinate, FCurvePoint& OutPoint, int32 DerivativeOrder) const
+bool FBezierCurve::IsBezier(const FNurbsCurveData& NurbsCurveData)
 {
-	OutPoint.DerivativeOrder = DerivativeOrder;
-	OutPoint.Init();
+	bool bIsBezier = ((NurbsCurveData.Poles.Num() - 1) % NurbsCurveData.Degree) == 0;
+	bIsBezier &= ((NurbsCurveData.NodalVector.Num() - 2) % NurbsCurveData.Degree) == 0;
+
+	if (!bIsBezier)
+	{
+		return false;
+	}
+
+	const int32 NumSegments = (NurbsCurveData.Poles.Num() - 1) / NurbsCurveData.Degree;
+	for (int32 Index = 0, KnotIndex = 1; Index < NumSegments && bIsBezier; ++Index)
+	{
+		const double KValue = NurbsCurveData.NodalVector[KnotIndex++];
+		for (int32 Jndex = 1; Jndex < NurbsCurveData.Degree; ++Jndex, ++KnotIndex)
+		{
+			if (!FMath::IsNearlyEqual(KValue, NurbsCurveData.NodalVector[KnotIndex], UE_DOUBLE_SMALL_NUMBER))
+			{
+				bIsBezier = false;
+				break;
+			}
+		}
+	}
+
+	return bIsBezier;
+}
+
+FBezierCurve::FBezierCurve(const FNurbsCurveData& NurbsCurveData)
+{
+	ensure(IsBezier(NurbsCurveData));
 
-	TArray<double> Bernstein;
-	TArray<double> BernsteinD1;
-	TArray<double> BernsteinD2;
+	Degree = NurbsCurveData.Degree;
+	Dimension = (int8)NurbsCurveData.Dimension;
 
-	BSpline::Bernstein(Poles.Num()-1, Coordinate, Bernstein, BernsteinD1, BernsteinD2);
+	NumSegments = (NurbsCurveData.Poles.Num() - 1) / NurbsCurveData.Degree;
 
-	for (int32 Index = 0; Index < Poles.Num(); Index++)
+	NodalVector.SetNum(NumSegments+1);
+	for (int32 Index = 0, Offset = Degree; Index <= NumSegments; ++Index, Offset += Degree)
 	{
-		OutPoint.Point += Poles[Index] * Bernstein[Index];
+		NodalVector[Index] = NurbsCurveData.NodalVector[Offset];
 	}
 
-	if (DerivativeOrder > 0)
+	Poles = NurbsCurveData.Poles;
+	Weights = NurbsCurveData.Weights;
+
+	// Validate the curve is actually rational
+	if (NurbsCurveData.bIsRational)
 	{
-		for (int32 Index = 0; Index < Poles.Num(); Index++)
+		const double WeightRef = Weights[0];
+		bIsRational = false;
+		for (int32 Index = 1; Index < Weights.Num() && !bIsRational; ++Index)
+		{
+			if (!FMath::IsNearlyEqual(WeightRef, Weights[Index], UE_DOUBLE_SMALL_NUMBER))
+			{
+				bIsRational = true;
+				break;
+			}
+		}
+
+		if (!bIsRational && !FMath::IsNearlyEqual(WeightRef, 1., UE_DOUBLE_SMALL_NUMBER))
 		{
-			OutPoint.Gradient += Poles[Index] * BernsteinD1[Index];
+			const double InvWeight = 1. / WeightRef;
+			for (FPoint& Pole : Poles)
+			{
+				Pole *= InvWeight;
+			}
 		}
 	}
 
-	if (DerivativeOrder > 1)
+	Boundary.Set(NodalVector[0], NodalVector[NumSegments]);
+}
+
+void FBezierCurve::EvaluatePoint(double Coordinate, FCurvePoint& OutPoint, int32 DerivativeOrder) const
+{
+	// Adjust coordinate value to curve's boundaries
+	if (!ensure(Coordinate >= (NodalVector[0] - UE_DOUBLE_SMALL_NUMBER)))
 	{
-		for (int32 Index = 0; Index < Poles.Num(); Index++)
+		Coordinate = NodalVector[0];
+	}
+
+	if (!ensure(Coordinate <= (NodalVector[NumSegments] + UE_DOUBLE_SMALL_NUMBER)))
+	{
+		Coordinate = NodalVector[NumSegments];
+	}
+
+	int32 SegmentIndex = 0;
+	if (FMath::IsNearlyEqual(Coordinate, NodalVector[NumSegments], UE_DOUBLE_SMALL_NUMBER))
+	{
+		SegmentIndex = NumSegments - 1;
+	}
+	else
+	{
+		for (; SegmentIndex < NumSegments + 1; ++SegmentIndex)
+		{
+			if (Coordinate < NodalVector[SegmentIndex])
+			{
+				--SegmentIndex;
+				break;
+			}
+		}
+	}
+	ensure(SegmentIndex < NumSegments);
+
+	OutPoint.DerivativeOrder = DerivativeOrder;
+	OutPoint.Init();
+
+	// Normalize input coordinates as Bezier computation is not impacted by knots' values
+	double NormalizedValue = (Coordinate - NodalVector[SegmentIndex]) / (NodalVector[SegmentIndex + 1] - NodalVector[SegmentIndex]);
+	
+	if (Degree == 1)
+	{
+		// Simple linear interpolation...
+		OutPoint.Point = Poles[0] * (1. - NormalizedValue) + Poles[1] * NormalizedValue;
+	}
+	else
+	{
+		const int32 Order = Degree + 1;
+		
+		TArray<double> BernsteinCoeffs;
+		BernsteinCoeffs.SetNum(DerivativeOrder > 1 ? 3 * Order : (DerivativeOrder > 0 ? 2 * Order : Order));
+		double* Bernstein = BernsteinCoeffs.GetData();
+		double* BernsteinD1 = DerivativeOrder > 0 ? Bernstein + Order : nullptr;
+		double* BernsteinD2 = DerivativeOrder > 1 ? BernsteinD1 + Order : nullptr;
+
+		BSpline::Bernstein(Degree, NormalizedValue, Bernstein, BernsteinD1, BernsteinD2);
+
+		double Weight = 0.;
+		for (int32 Index = 0, PolesIndex = SegmentIndex * Degree; Index < Order; ++Index, ++PolesIndex)
+		{
+			OutPoint.Point += Poles[PolesIndex] * Bernstein[Index];
+			Weight += Weights[PolesIndex] * Bernstein[Index];
+		}
+
+		if (bIsRational)
+		{
+			const double InvWeight = 1. / Weight;
+			OutPoint.Point *= InvWeight;
+		}
+
+		if (BernsteinD1)
 		{
-			OutPoint.Laplacian += Poles[Index] * BernsteinD2[Index];
+		for (int32 Index = 0, PolesIndex = SegmentIndex * Degree; Index < Order; ++Index, ++PolesIndex)
+			{
+				OutPoint.Gradient += Poles[PolesIndex] * BernsteinD1[Index];
+			}
+
+			if (BernsteinD2)
+			{
+		for (int32 Index = 0, PolesIndex = SegmentIndex * Degree; Index < Order; ++Index, ++PolesIndex)
+				{
+					OutPoint.Laplacian += Poles[PolesIndex] * BernsteinD2[Index];
+				}
+			}
 		}
 	}
 }
diff --git a/Engine/Source/Runtime/Datasmith/CADKernel/Private/CADKernel/Geo/Curves/BezierCurve.h b/Engine/Source/Runtime/Datasmith/CADKernel/Private/CADKernel/Geo/Curves/BezierCurve.h
--- a/Engine/Source/Runtime/Datasmith/CADKernel/Private/CADKernel/Geo/Curves/BezierCurve.h
+++ b/Engine/Source/Runtime/Datasmith/CADKernel/Private/CADKernel/Geo/Curves/BezierCurve.h
@@ -2,29 +2,54 @@
 #pragma once
 
 #include "CADKernel/Geo/Curves/Curve.h"
+#include "NURBSCurve.h"
 
 namespace UE::CADKernel
 {
+
 class CADKERNEL_API FBezierCurve : public FCurve
 {
 	friend class FEntity;
 
 protected:
+	bool bIsRational = false;
+
+	int32 Degree = 0;
+	TArray<double> NodalVector;
+	int32 NumSegments = 0;
+
 	TArray<FPoint> Poles;
+	TArray<double> Weights;
 
 	FBezierCurve(const TArray<FPoint>& InPoles)
 		: Poles{ InPoles }
 	{
+		Degree = Poles.Num() - 1;
+		NumSegments = 1;
+		NodalVector.Add(0.);
+		NodalVector.Add(1.);
+		Weights.Add(1.);
+		Weights.Add(1.);
 	}
 
+	FBezierCurve(const FNurbsCurveData& NurbsCurveData);
+
 	FBezierCurve() = default;
 
 public:
 
+	static bool IsBezier(const FNurbsCurveData& NurbsCurveData);
+
 	virtual void Serialize(FCADKernelArchive& Ar) override
 	{
 		FCurve::Serialize(Ar);
+
+		Ar.Serialize(bIsRational);
+		Ar.Serialize(Degree);
+		Ar.Serialize(NumSegments);
+		Ar.Serialize(NodalVector);
 		Ar.Serialize(Poles);
+		Ar.Serialize(Weights);
 	}
 
 #ifdef CADKERNEL_DEV
@@ -50,6 +75,25 @@ public:
 	virtual void Offset(const FPoint& OffsetDirection) override;
 
 	virtual void EvaluatePoint(double Coordinate, FCurvePoint& OutPoint, int32 DerivativeOrder = 0) const override;
+	virtual void Evaluate2DPoint(double Coordinate, FCurvePoint2D& OutPoint, int32 DerivativeOrder = 0) const override
+	{
+		FCurvePoint CurvePoint;
+		EvaluatePoint(Coordinate, CurvePoint, DerivativeOrder);
+		OutPoint.DerivativeOrder = DerivativeOrder;
+		OutPoint.Point[0] = CurvePoint.Point[0];
+		OutPoint.Point[1] = CurvePoint.Point[1];
+		if (DerivativeOrder > 0)
+		{
+			OutPoint.Gradient[0] = CurvePoint.Gradient[0];
+			OutPoint.Gradient[1] = CurvePoint.Gradient[1];
+			
+			if (DerivativeOrder > 1)
+			{
+				OutPoint.Laplacian[0] = CurvePoint.Laplacian[0];
+				OutPoint.Laplacian[1] = CurvePoint.Laplacian[1];
+			}
+		}
+	}
 
 	virtual void ExtendTo(const FPoint& Point) override;
 
diff --git a/Engine/Source/Runtime/Datasmith/CADKernel/Private/CADKernel/Geo/Curves/CurveMaker.cpp b/Engine/Source/Runtime/Datasmith/CADKernel/Private/CADKernel/Geo/Curves/CurveMaker.cpp
--- a/Engine/Source/Runtime/Datasmith/CADKernel/Private/CADKernel/Geo/Curves/CurveMaker.cpp
+++ b/Engine/Source/Runtime/Datasmith/CADKernel/Private/CADKernel/Geo/Curves/CurveMaker.cpp
@@ -10,6 +10,10 @@ namespace UE::CADKernel
 
 TSharedPtr<FCurve> FCurve::MakeNurbsCurve(FNurbsCurveData& InNurbsData)
 {
+	if (FBezierCurve::IsBezier(InNurbsData))
+	{
+		return FEntity::MakeShared<UE::CADKernel::FBezierCurve>(InNurbsData);
+	}
 	return FEntity::MakeShared<UE::CADKernel::FNURBSCurve>(InNurbsData);
 }
 
diff --git a/Engine/Source/Runtime/Datasmith/CADKernel/Private/CADKernel/Geo/Curves/NURBSCurve.h b/Engine/Source/Runtime/Datasmith/CADKernel/Private/CADKernel/Geo/Curves/NURBSCurve.h
--- a/Engine/Source/Runtime/Datasmith/CADKernel/Private/CADKernel/Geo/Curves/NURBSCurve.h
+++ b/Engine/Source/Runtime/Datasmith/CADKernel/Private/CADKernel/Geo/Curves/NURBSCurve.h
@@ -8,7 +8,7 @@ namespace UE::CADKernel
 {
 struct FNurbsCurveData
 {
-	bool bIsRational;
+	bool bIsRational = false;
 	int32 Dimension = 0;
 
 	int32 Degree = 0;
@@ -90,7 +90,7 @@ public:
 		return Weights;
 	}
 
-	TArray<double> GetHPoles() const
+	const TArray<double>& GetHPoles() const
 	{
 		return HomogeneousPoles;
 	}
diff --git a/Engine/Source/Runtime/Datasmith/CADKernel/Private/CADKernel/Geo/Sampler/SamplerAbstract.h b/Engine/Source/Runtime/Datasmith/CADKernel/Private/CADKernel/Geo/Sampler/SamplerAbstract.h
--- a/Engine/Source/Runtime/Datasmith/CADKernel/Private/CADKernel/Geo/Sampler/SamplerAbstract.h
+++ b/Engine/Source/Runtime/Datasmith/CADKernel/Private/CADKernel/Geo/Sampler/SamplerAbstract.h
@@ -149,7 +149,6 @@ protected:
 	 */
 	void RunSampling()
 	{
-
 		// < 100000 : check to avoid that the process loops endlessly
 		int32 CandidatePointsCount = CandidatePoints.Size();
 		while (CandidatePointsCount && Sampling.Coordinates.Num() < 100000)
@@ -251,6 +250,10 @@ protected:
 	/** Adds coordinates of the next candidate points in NextCoordinates array */
 	void AddIntermediateCoordinates(double UMin, double UMax, int32 PointNum)
 	{
+		if (FMath::IsNearlyEqual(UMin, UMax, UE_DOUBLE_SMALL_NUMBER))
+		{
+			return;
+		}
 		double DeltaCoord = (UMax - UMin) / (PointNum + 1);
 		double UCoord = UMin;
 		for (int32 Index = 0; Index < PointNum; Index++)
diff --git a/Engine/Source/Runtime/Datasmith/CADKernel/Private/CADKernel/Geo/Sampler/SamplerOnParam.h b/Engine/Source/Runtime/Datasmith/CADKernel/Private/CADKernel/Geo/Sampler/SamplerOnParam.h
--- a/Engine/Source/Runtime/Datasmith/CADKernel/Private/CADKernel/Geo/Sampler/SamplerOnParam.h
+++ b/Engine/Source/Runtime/Datasmith/CADKernel/Private/CADKernel/Geo/Sampler/SamplerOnParam.h
@@ -194,6 +194,7 @@ protected:
 		else
 		{
 			Curve.Evaluate2DPoints(CandidatePoints.Coordinates, CandidatePoints.Points2D);
+			Surface.ValidateUVPoints(CandidatePoints.Points2D);
 			Surface.EvaluatePoints(CandidatePoints);
 		}
 	};
diff --git a/Engine/Source/Runtime/Datasmith/CADKernel/Private/CADKernel/Geo/Surfaces/NURBSSurface.h b/Engine/Source/Runtime/Datasmith/CADKernel/Private/CADKernel/Geo/Surfaces/NURBSSurface.h
--- a/Engine/Source/Runtime/Datasmith/CADKernel/Private/CADKernel/Geo/Surfaces/NURBSSurface.h
+++ b/Engine/Source/Runtime/Datasmith/CADKernel/Private/CADKernel/Geo/Surfaces/NURBSSurface.h
@@ -108,6 +108,40 @@ protected:
 
 public:
 
+
+	virtual void ValidateUVPoints(TArray<FPoint2D>& UVPoints) const
+	{
+		auto ValidateValue = [](double& Value, const double& Reference)
+			{
+				if (FMath::IsNearlyEqual(Value, Reference, UE_DOUBLE_KINDA_SMALL_NUMBER))
+				{
+					Value = Reference;
+					return true;
+				}
+
+				return false;
+			};
+
+		for (FPoint2D& UVPoint : UVPoints)
+		{
+			for (const double& NodalValue : UNodalVector)
+			{
+				if (ValidateValue(UVPoint.U, NodalValue))
+				{
+					break;
+				}
+			}
+
+			for (const double& NodalValue : VNodalVector)
+			{
+				if (ValidateValue(UVPoint.V, NodalValue))
+				{
+					break;
+				}
+			}
+		}
+	}
+
 	virtual void Serialize(FCADKernelArchive& Ar) override
 	{
 		FSurface::Serialize(Ar);
@@ -170,7 +204,7 @@ public:
 		return Weights;
 	}
 
-	TArray<double> GetHPoles() const
+	const TArray<double>& GetHPoles() const
 	{
 		return HomogeneousPoles;
 	}
diff --git a/Engine/Source/Runtime/Datasmith/CADKernel/Private/CADKernel/Math/BSpline.cpp b/Engine/Source/Runtime/Datasmith/CADKernel/Private/CADKernel/Math/BSpline.cpp
--- a/Engine/Source/Runtime/Datasmith/CADKernel/Private/CADKernel/Math/BSpline.cpp
+++ b/Engine/Source/Runtime/Datasmith/CADKernel/Private/CADKernel/Math/BSpline.cpp
@@ -69,22 +69,19 @@ int32 CoefficientBinomial(int32 n, int32 k)
 	return fact(n) / (fact(k) * fact(n - k));
 }
 
-void Bernstein(int32 Degree, double InCoordinateU, TArray<double>& BernsteinValuesAtU, TArray<double>& BernsteinGradientsAtU, TArray<double>& BernsteinLaplaciansAtU)
+// #cadkernel_check: Precompute Bernstein values.
+void Bernstein(int32 Degree, double InCoordinateU, double* BernsteinValuesAtU, double* BernsteinGradientsAtU, double* BernsteinLaplaciansAtU)
 {
-	TArray<double> ti;
-	TArray<double> si;
-	TArray<double> ci;
-
 	int32 i;
 	double t, s;
 
-	ti.SetNum(Degree + 1);
-	si.SetNum(Degree + 1);
-	ci.SetNum(Degree + 1);
+	const int32 Order = Degree + 1;
+	TArray<double> Buffer;
+	Buffer.SetNum(3 * Order);
 
-	BernsteinValuesAtU.SetNum(Degree + 1);
-	BernsteinGradientsAtU.SetNum(Degree + 1);
-	BernsteinLaplaciansAtU.SetNum(Degree + 1);
+	double* ti = Buffer.GetData();
+	double* si = ti + Order;
+	double* ci = si + Order;
 
 	// calculer les coefficient ci
 	ci[0] = 1;
@@ -92,15 +89,15 @@ void Bernstein(int32 Degree, double InCoordinateU, TArray<double>& BernsteinValu
 
 	for (i = 1; i < Degree; i++)
 	{
-		ci[i] = ci[i - 1] * (Degree - i + 1) / i;
+		ci[i] = ci[i - 1] * (Order - i) / (double)i;
 	}
 
-	ti[0] = 1.0f;
-	si[0] = 1.0f;
+	ti[0] = 1.;
+	si[0] = 1.;
 
 	// calculer les monomes t**i et s**i=(1-t)**i
 	t = InCoordinateU;
-	s = 1.0f - InCoordinateU;
+	s = 1. - InCoordinateU;
 	for (i = 1; i <= Degree; i++)
 	{
 		ti[i] = ti[i - 1] * t;
@@ -113,35 +110,41 @@ void Bernstein(int32 Degree, double InCoordinateU, TArray<double>& BernsteinValu
 		BernsteinValuesAtU[i] = ci[i] * ti[i] * si[Degree - i];
 	}
 
-	// calculer les derivees premieres
-	BernsteinGradientsAtU[0] = -Degree * si[Degree - 1];
-	BernsteinGradientsAtU[Degree] = Degree * ti[Degree - 1];
-	for (i = 1; i < Degree; i++)
-	{
-		BernsteinGradientsAtU[i] = ci[i] * (i * ti[i - 1] * si[Degree - i] - (Degree - i) * ti[i] * si[Degree - i - 1]);
-	}
-
-	// calculer les derivees secondes
-	if (Degree > 1)
+	if (BernsteinGradientsAtU)
 	{
-		BernsteinLaplaciansAtU[0] = Degree * (Degree - 1) * si[Degree - 2];
-		BernsteinLaplaciansAtU[1] = ci[1] * (-2) * (Degree - 1) * si[Degree - 2];
-		if (Degree > 2) BernsteinLaplaciansAtU[1] += ci[1] * (Degree - 1) * (Degree - 2) * ti[1] * si[Degree - 3];
-		BernsteinLaplaciansAtU[Degree] = Degree * (Degree - 1) * ti[Degree - 2];
-		BernsteinLaplaciansAtU[Degree - 1] = ci[Degree - 1] * (-2) * (Degree - 1) * ti[Degree - 2];
-		if (Degree > 2) BernsteinLaplaciansAtU[Degree - 1] += ci[Degree - 1] * (Degree - 1) * (Degree - 2) * ti[Degree - 3] * si[1];
+		// calculer les derivees premieres
+		BernsteinGradientsAtU[0] = -Degree * si[Degree - 1];
+		BernsteinGradientsAtU[Degree] = Degree * ti[Degree - 1];
+		for (i = 1; i < Degree; i++)
+		{
+			BernsteinGradientsAtU[i] = ci[i] * (i * ti[i - 1] * si[Degree - i] - (Degree - i) * ti[i] * si[Degree - i - 1]);
+		}
 
-		for (i = 2; i < Degree - 1; i++)
+		if (BernsteinLaplaciansAtU)
 		{
-			BernsteinLaplaciansAtU[i] = ci[i] * (i * (i - 1) * ti[i - 2] * si[Degree - i]
-				- 2 * i * (Degree - i) * ti[i - 1] * si[Degree - i - 1]
-				+ (Degree - i) * (Degree - i - 1) * ti[i] * si[Degree - i - 2]);
+			// calculer les derivees secondes
+			if (Degree > 1)
+			{
+				BernsteinLaplaciansAtU[0] = Degree * (Degree - 1) * si[Degree - 2];
+				BernsteinLaplaciansAtU[1] = ci[1] * (-2) * (Degree - 1) * si[Degree - 2];
+				if (Degree > 2) BernsteinLaplaciansAtU[1] += ci[1] * (Degree - 1) * (Degree - 2) * ti[1] * si[Degree - 3];
+				BernsteinLaplaciansAtU[Degree] = Degree * (Degree - 1) * ti[Degree - 2];
+				BernsteinLaplaciansAtU[Degree - 1] = ci[Degree - 1] * (-2) * (Degree - 1) * ti[Degree - 2];
+				if (Degree > 2) BernsteinLaplaciansAtU[Degree - 1] += ci[Degree - 1] * (Degree - 1) * (Degree - 2) * ti[Degree - 3] * si[1];
+
+				for (i = 2; i < Degree - 1; i++)
+				{
+					BernsteinLaplaciansAtU[i] = ci[i] * (i * (i - 1) * ti[i - 2] * si[Degree - i]
+						- 2 * i * (Degree - i) * ti[i - 1] * si[Degree - i - 1]
+						+ (Degree - i) * (Degree - i - 1) * ti[i] * si[Degree - i - 2]);
+				}
+			}
+			else
+			{
+				BernsteinLaplaciansAtU[0] = 0.0f;
+			}
 		}
 	}
-	else
-	{
-		BernsteinLaplaciansAtU[0] = 0.0f;
-	}
 }
 
 int32 Dichotomy(const double* const tab, double key, int32 first, int32 last)
diff --git a/Engine/Source/Runtime/Datasmith/CADKernel/Private/CADKernel/Mesh/Meshers/MesherTools.cpp b/Engine/Source/Runtime/Datasmith/CADKernel/Private/CADKernel/Mesh/Meshers/MesherTools.cpp
--- a/Engine/Source/Runtime/Datasmith/CADKernel/Private/CADKernel/Mesh/Meshers/MesherTools.cpp
+++ b/Engine/Source/Runtime/Datasmith/CADKernel/Private/CADKernel/Mesh/Meshers/MesherTools.cpp
@@ -221,7 +221,7 @@ void FMesherTools::ComputeFinalCuttingPointsWithImposedCuttingPoints(const TArra
 		double TravelTime = 0;
 
 		TravelTimeByStep.Empty();
-		for (;;)
+		while(CrossingIndex < CrossingUs.Num() - 1)
 		{
 			double Step = CrossingUs[CrossingIndex + 1];
 			Step = FMath::Min(Step, EndU);
diff --git a/Engine/Source/Runtime/Datasmith/CADKernel/Private/CADKernel/Mesh/Meshers/ParametricMesher.cpp b/Engine/Source/Runtime/Datasmith/CADKernel/Private/CADKernel/Mesh/Meshers/ParametricMesher.cpp
--- a/Engine/Source/Runtime/Datasmith/CADKernel/Private/CADKernel/Mesh/Meshers/ParametricMesher.cpp
+++ b/Engine/Source/Runtime/Datasmith/CADKernel/Private/CADKernel/Mesh/Meshers/ParametricMesher.cpp
@@ -132,7 +132,11 @@ void FParametricMesher::PreMeshingTasks()
 #endif
 		if (Face == nullptr || Face->IsNotMeshable())
 		{
+#ifdef CADKernelMultiThread
 			return;
+#else
+			continue;
+#endif
 		}
 		ApplyFaceCriteria(*Face, Criteria, MeshingTolerance, bThinZone);
 		if (!Face->IsDeletedOrDegenerated())
diff --git a/Engine/Source/Runtime/Datasmith/CADKernel/Private/CADKernel/Topo/TopologicalEdge.cpp b/Engine/Source/Runtime/Datasmith/CADKernel/Private/CADKernel/Topo/TopologicalEdge.cpp
--- a/Engine/Source/Runtime/Datasmith/CADKernel/Private/CADKernel/Topo/TopologicalEdge.cpp
+++ b/Engine/Source/Runtime/Datasmith/CADKernel/Private/CADKernel/Topo/TopologicalEdge.cpp
@@ -29,6 +29,7 @@ FTopologicalEdge::FTopologicalEdge(const TSharedRef<FRestrictionCurve>& InCurve,
 	, Length3D(-1)
 {
 	ensureCADKernel(Boundary.IsValid());
+	ensureCADKernel(Boundary.GetMin() < Boundary.GetMax());
 }
 
 FTopologicalEdge::FTopologicalEdge(const TSharedRef<FRestrictionCurve>& InCurve, const TSharedRef<FTopologicalVertex>& InVertex1, const TSharedRef<FTopologicalVertex>& InVertex2)
@@ -404,6 +405,11 @@ FTopologicalFace* FTopologicalEdge::GetFace() const
 
 void FTopologicalEdge::ComputeCrossingPointCoordinates()
 {
+	if (FMath::IsNearlyEqual(Boundary.GetMin(), Boundary.GetMax(), UE_DOUBLE_SMALL_NUMBER))
+	{
+		// #cadkernel_check: Why could this happen? Shouldn't it be detected way earlier?
+		return;
+	}
 	double Tolerance = GetTolerance3D();
 
 	FSurfacicPolyline Presampling;
diff --git a/Engine/Source/Runtime/Datasmith/CADKernel/Private/CADKernel/Topo/TopologicalFace.cpp b/Engine/Source/Runtime/Datasmith/CADKernel/Private/CADKernel/Topo/TopologicalFace.cpp
--- a/Engine/Source/Runtime/Datasmith/CADKernel/Private/CADKernel/Topo/TopologicalFace.cpp
+++ b/Engine/Source/Runtime/Datasmith/CADKernel/Private/CADKernel/Topo/TopologicalFace.cpp
@@ -25,20 +25,24 @@ namespace UE::CADKernel
 
 void FTopologicalFace::ComputeBoundary() const
 {
-	Boundary->Init();
+	Boundary = CarrierSurface->GetBoundary();
+
 	TArray<TArray<FPoint2D>> TmpLoops;
 	Get2DLoopSampling(TmpLoops);
 
-	for (const TArray<FPoint2D>& Loop : TmpLoops)
+	if (TmpLoops.Num() > 0)
 	{
-		for (const FPoint2D& Point : Loop)
+		for (const TArray<FPoint2D>& Loop : TmpLoops)
 		{
-			Boundary->ExtendTo(Point);
+			for (const FPoint2D& Point : Loop)
+			{
+				Boundary->ExtendTo(Point);
+			}
 		}
-	}
 
-	// Check with the carrier surface bounds
-	CarrierSurface->ExtendBoundaryTo(Boundary);
+		// Check with the carrier surface bounds
+		CarrierSurface->ExtendBoundaryTo(Boundary);
+	}
 
 	Boundary->WidenIfDegenerated();
 	Boundary.SetReady();
@@ -79,7 +83,8 @@ void FTopologicalFace::UpdateBBox(int32 IsoCount, const double ApproximationFact
 			TArray<double> Intersections;
 			FindLoopIntersectionsWithIso(IsoType, Coordinate, BoundaryApproximation, Intersections);
 			int32 IntersectionCount = Intersections.Num();
-			if (IntersectionCount == 0)
+			// #cadkernel_check: Why IntersectionCount can be less than 2 or Intersections has same extremities
+			if (IntersectionCount < 2 || FMath::IsNearlyEqual(Intersections[0], Intersections.Last(), UE_DOUBLE_SMALL_NUMBER))
 			{
 				continue;
 			}
diff --git a/Engine/Source/Runtime/Datasmith/CADKernel/Private/CADKernel/Topo/Topomaker.cpp b/Engine/Source/Runtime/Datasmith/CADKernel/Private/CADKernel/Topo/Topomaker.cpp
--- a/Engine/Source/Runtime/Datasmith/CADKernel/Private/CADKernel/Topo/Topomaker.cpp
+++ b/Engine/Source/Runtime/Datasmith/CADKernel/Private/CADKernel/Topo/Topomaker.cpp
@@ -177,7 +177,7 @@ FTopologicalVertex* SplitAndLink(FTopologicalVertex& StartVertex, FTopologicalEd
 	}
 
 	double SquareDistanceToOtherPoint = ProjectedPoint.SquareDistance(EndVertex->GetBarycenter());
-	if (SquareDistanceToOtherPoint < SquareMinEdgeLength)
+	if (FMath::IsNearlyEqual(UProjectedPoint, EdgeToSplit.GetBoundary().GetMax(), UE_DOUBLE_SMALL_NUMBER) || SquareDistanceToOtherPoint < SquareMinEdgeLength)
 	{
 		// the new point is closed to the extremity, a degenerated edge will be created, so the edges are joined
 		VertexToLink->Link(*EndVertex);
@@ -186,7 +186,7 @@ FTopologicalVertex* SplitAndLink(FTopologicalVertex& StartVertex, FTopologicalEd
 	}
 
 	SquareDistanceToProjectedPoint = ProjectedPoint.SquareDistance(StartVertex.GetBarycenter());
-	if (SquareDistanceToProjectedPoint < SquareMinEdgeLength)
+	if (FMath::IsNearlyEqual(UProjectedPoint, EdgeToSplit.GetBoundary().GetMin(), UE_DOUBLE_SMALL_NUMBER) || SquareDistanceToProjectedPoint < SquareMinEdgeLength)
 	{
 		VertexToLink->Link(StartVertex);
 		EdgeToLink.SetAsDegenerated();
diff --git a/Engine/Source/Runtime/Datasmith/CADKernel/Public/CADKernel/Geo/Sampling/Polyline.h b/Engine/Source/Runtime/Datasmith/CADKernel/Public/CADKernel/Geo/Sampling/Polyline.h
--- a/Engine/Source/Runtime/Datasmith/CADKernel/Public/CADKernel/Geo/Sampling/Polyline.h
+++ b/Engine/Source/Runtime/Datasmith/CADKernel/Public/CADKernel/Geo/Sampling/Polyline.h
@@ -109,6 +109,11 @@ public:
 	void SwapCoordinates(TArray<double>& NewCoordinates)
 	{
 		Swap(NewCoordinates, Coordinates);
+		Points.Empty(Coordinates.Num());
+		if (bWithTangent)
+		{
+			Tangent.Empty(Coordinates.Num());
+		}
 	}
 
 	void GetAt(int32 Index, double& Coordinate, PointType& Point)
diff --git a/Engine/Source/Runtime/Datasmith/CADKernel/Public/CADKernel/Geo/Sampling/SurfacicPolyline.h b/Engine/Source/Runtime/Datasmith/CADKernel/Public/CADKernel/Geo/Sampling/SurfacicPolyline.h
--- a/Engine/Source/Runtime/Datasmith/CADKernel/Public/CADKernel/Geo/Sampling/SurfacicPolyline.h
+++ b/Engine/Source/Runtime/Datasmith/CADKernel/Public/CADKernel/Geo/Sampling/SurfacicPolyline.h
@@ -303,6 +303,16 @@ public:
 	void SwapCoordinates(TArray<double>& NewCoordinates)
 	{
 		Swap(NewCoordinates, Coordinates);
+		Points2D.Empty(Coordinates.Num());
+		Points3D.Empty(Coordinates.Num());
+		if (bWithNormals)
+		{
+			Normals.Empty(Coordinates.Num());
+		}
+		if (bWithTangent)
+		{
+			Tangents.Empty(Coordinates.Num());
+		}
 	}
 
 	/**
diff --git a/Engine/Source/Runtime/Datasmith/CADKernel/Public/CADKernel/Geo/Surfaces/Surface.h b/Engine/Source/Runtime/Datasmith/CADKernel/Public/CADKernel/Geo/Surfaces/Surface.h
--- a/Engine/Source/Runtime/Datasmith/CADKernel/Public/CADKernel/Geo/Surfaces/Surface.h
+++ b/Engine/Source/Runtime/Datasmith/CADKernel/Public/CADKernel/Geo/Surfaces/Surface.h
@@ -90,6 +90,8 @@ public:
 	static TSharedPtr<FSurface> MakeSphericalSurface(const double InToleranceGeometric, const FMatrixH& InMatrix, double InRadius, const FSurfacicBoundary& InBoundary);
 	static TSharedPtr<FSurface> MakeTorusSurface(const double InToleranceGeometric, const FMatrixH& InMatrix, double InMajorRadius, double InMinorRadius, const FSurfacicBoundary& InBoundary);
 
+	virtual void ValidateUVPoints(TArray<FPoint2D>& UVPoints) const {}
+
 	virtual void Serialize(FCADKernelArchive& Ar) override
 	{
 		// Surface's type is serialize because it is used to instantiate the correct entity on de-serialization (@see Deserialize(FCADKernelArchive& Archive))
diff --git a/Engine/Source/Runtime/Datasmith/CADKernel/Public/CADKernel/Math/BSpline.h b/Engine/Source/Runtime/Datasmith/CADKernel/Public/CADKernel/Math/BSpline.h
--- a/Engine/Source/Runtime/Datasmith/CADKernel/Public/CADKernel/Math/BSpline.h
+++ b/Engine/Source/Runtime/Datasmith/CADKernel/Public/CADKernel/Math/BSpline.h
@@ -24,7 +24,7 @@ namespace BSpline
 /**
  * Compute the values of Bernstein polynomial
  */
-CADKERNEL_API void Bernstein(int32 Degree, double InCoordinateU, TArray<double>& BernsteinValuesAtU, TArray<double>& BernsteinGradientsAtU, TArray<double>& BernsteinLaplaciansAtU);
+CADKERNEL_API void Bernstein(int32 Degree, double InCoordinateU, double* BernsteinValuesAtU, double* BernsteinGradientsAtU, double* BernsteinLaplaciansAtU);
 
 CADKERNEL_API void FindNotDerivableParameters(const FNURBSCurve&, int32 InDerivativeOrder, const FLinearBoundary& Boundary, TArray<double>& OutNotDerivableParameters);
 CADKERNEL_API void FindNotDerivableParameters(const FNURBSSurface&, int32 InDerivativeOrder, const FSurfacicBoundary& Boundary, FCoordinateGrid& OutNotDerivableParameters);
