[CADKernel] renamed Entity files to CADEntity

diff --git a/Engine/Plugins/Enterprise/DatasmithCADImporter/Source/CADInterfaces/Private/TechSoftFileParserCADKernelTessellator.cpp b/Engine/Plugins/Enterprise/DatasmithCADImporter/Source/CADInterfaces/Private/TechSoftFileParserCADKernelTessellator.cpp
--- a/Engine/Plugins/Enterprise/DatasmithCADImporter/Source/CADInterfaces/Private/TechSoftFileParserCADKernelTessellator.cpp
+++ b/Engine/Plugins/Enterprise/DatasmithCADImporter/Source/CADInterfaces/Private/TechSoftFileParserCADKernelTessellator.cpp
@@ -10,7 +10,7 @@
 #include "TechSoftBridge.h"
 #include "TUniqueTechSoftObj.h"
 
-#include "CADKernel/Core/Entity.h"
+#include "CADKernel/Core/CADEntity.h"
 #include "CADKernel/Core/Session.h"
 #include "CADKernel/Core/Types.h"
 
diff --git a/Engine/Plugins/Enterprise/DatasmithCADImporter/Source/CADKernelSurface/Private/CADKernelSurfaceExtension.cpp b/Engine/Plugins/Enterprise/DatasmithCADImporter/Source/CADKernelSurface/Private/CADKernelSurfaceExtension.cpp
--- a/Engine/Plugins/Enterprise/DatasmithCADImporter/Source/CADKernelSurface/Private/CADKernelSurfaceExtension.cpp
+++ b/Engine/Plugins/Enterprise/DatasmithCADImporter/Source/CADKernelSurface/Private/CADKernelSurfaceExtension.cpp
@@ -14,7 +14,7 @@
 #include "UObject/EnterpriseObjectVersion.h"
 
 #include "CADKernel/Core/CADKernelArchive.h"
-#include "CADKernel/Core/Entity.h"
+#include "CADKernel/Core/CADEntity.h"
 #include "CADKernel/Core/Session.h"
 #include "CADKernel/Core/Types.h"
 #include "CADKernel/Mesh/Structure/ModelMesh.h"
diff --git a/Engine/Plugins/Enterprise/DatasmithCADImporter/Source/CADLibrary/Private/CADKernelTools.cpp b/Engine/Plugins/Enterprise/DatasmithCADImporter/Source/CADLibrary/Private/CADKernelTools.cpp
--- a/Engine/Plugins/Enterprise/DatasmithCADImporter/Source/CADLibrary/Private/CADKernelTools.cpp
+++ b/Engine/Plugins/Enterprise/DatasmithCADImporter/Source/CADLibrary/Private/CADKernelTools.cpp
@@ -9,7 +9,7 @@
 #include "StaticMeshAttributes.h"
 #include "StaticMeshOperations.h"
 
-#include "CADKernel/Core/Entity.h"
+#include "CADKernel/Core/CADEntity.h"
 #include "CADKernel/Core/MetadataDictionary.h"
 #include "CADKernel/Core/Session.h"
 #include "CADKernel/Core/Types.h"
diff --git a/Engine/Source/Runtime/Datasmith/CADKernel/Private/CADKernel/Core/CADEntity.cpp b/Engine/Source/Runtime/Datasmith/CADKernel/Private/CADKernel/Core/CADEntity.cpp
new file mode 100644
--- /dev/null
+++ b/Engine/Source/Runtime/Datasmith/CADKernel/Private/CADKernel/Core/CADEntity.cpp
@@ -0,0 +1,379 @@
+// Copyright Epic Games, Inc. All Rights Reserved.
+
+#include "CADKernel/Core/CADEntity.h"
+
+#include "CADKernel/Core/EntityGeom.h"
+
+#include "CADKernel/Core/Database.h"
+#include "CADKernel/Core/Session.h"
+#include "CADKernel/Core/System.h"
+#include "CADKernel/Geo/Curves/Curve.h"
+#include "CADKernel/Geo/Surfaces/Surface.h"
+#include "CADKernel/Mesh/Criteria/Criterion.h"
+#include "CADKernel/Topo/Body.h"
+#include "CADKernel/Topo/Model.h"
+#include "CADKernel/Topo/Shell.h"
+#include "CADKernel/Topo/TopologicalEdge.h"
+#include "CADKernel/Topo/TopologicalFace.h"
+#include "CADKernel/Topo/TopologicalLink.h"
+#include "CADKernel/Topo/TopologicalLoop.h"
+#include "CADKernel/Topo/TopologicalVertex.h"
+
+namespace UE::CADKernel
+{
+
+	const FPairOfIndex FPairOfIndex::Undefined(-1, -1);
+
+	const TCHAR* FEntity::TypesNames[] = 
+	{
+		TEXT("UndefinedEntity"),
+		TEXT("Curve"),
+		TEXT("Surface"),
+
+		TEXT("Edge Link"),
+		TEXT("Vertex Link"),
+		TEXT("Edge"),
+		TEXT("Face"),
+ 		TEXT("Link"),
+		TEXT("Loop"),
+		TEXT("Vertex"),
+		TEXT("Shell"),
+		TEXT("Body"),
+		TEXT("Model"),
+
+		TEXT("Mesh Model"),
+		TEXT("Mesh"),
+
+		TEXT("Group"),
+		TEXT("Criterion"),
+		TEXT("Property"),
+		nullptr
+	};
+
+	const TCHAR* FEntity::GetTypeName(EEntity Type)
+	{
+		if (Type > EEntity::None && Type < EEntity::EntityTypeEnd)
+		{
+			return TypesNames[(int32)Type];
+		}
+		return TypesNames[(int32)EEntity::None];
+	}
+
+	FEntity::~FEntity()
+	{
+		TRACE_CPUPROFILER_EVENT_SCOPE(FEntity::~FEntity);
+	}
+
+	TSharedPtr<FEntity> FEntity::Deserialize(FCADKernelArchive& Archive)
+	{
+		TRACE_CPUPROFILER_EVENT_SCOPE(FEntity::Deserialize);
+
+		ensureCADKernel(Archive.IsLoading());
+
+		EEntity Type = EEntity::None;
+		Archive << Type;
+
+		TSharedPtr<FEntity> Entity;
+		switch (Type)
+		{
+		case EEntity::Body:
+			Entity = FEntity::MakeShared<FBody>(Archive);
+			break;
+		case EEntity::Curve:
+			Entity = FCurve::Deserialize(Archive);
+			break;
+		case EEntity::Criterion:
+			Entity = FCriterion::Deserialize(Archive);
+			break;
+		case EEntity::EdgeLink:
+			Entity = FEntity::MakeShared<FEdgeLink>(Archive);
+			break;
+		case EEntity::Model:
+			Entity = FEntity::MakeShared<FModel>(Archive);
+			break;
+		case EEntity::Shell:
+			Entity = FEntity::MakeShared<FShell>(Archive);
+			break;
+		case EEntity::Surface:
+			Entity = FSurface::Deserialize(Archive);
+			break;
+		case EEntity::TopologicalEdge:
+			Entity = FEntity::MakeShared<FTopologicalEdge>(Archive);
+			break;
+		case EEntity::TopologicalFace:
+			Entity = FEntity::MakeShared<FTopologicalFace>(Archive);
+			break;
+		case EEntity::TopologicalLoop:
+			Entity = FEntity::MakeShared<FTopologicalLoop>(Archive);
+			break;
+		case EEntity::TopologicalVertex:
+			Entity = FEntity::MakeShared<FTopologicalVertex>(Archive);
+			break;
+		case EEntity::VertexLink:
+			Entity = FEntity::MakeShared<FVertexLink>(Archive);
+			break;
+		default:
+			break;
+		}
+
+		if (Entity.IsValid())
+		{
+			Entity->Serialize(Archive);
+			Archive.AddEntityFromArchive(Entity);
+		}
+		return Entity;
+	}
+
+#ifdef CADKERNEL_DEV
+
+	void FEntity::InfoEntity() const
+	{
+		FInfoEntity Info;
+		GetInfo(Info);
+		Info.Display();
+	}
+
+	FInfoEntity& FEntity::GetInfo(FInfoEntity& Info) const
+	{
+		Info.Set(*this);
+		Info.Add(TEXT("Id"), GetId())
+			.Add(TEXT("Type"), TypesNames[(uint32)GetEntityType()]);
+		return Info;
+	}
+
+	FInfoEntity& FEntityGeom::GetInfo(FInfoEntity& Info) const
+	{
+		return FEntity::GetInfo(Info)
+			.Add(TEXT("Kio"), CtKioId);
+	}
+
+	void FEntity::AddEntityInDatabase(TSharedRef<FEntity> Entity)
+	{
+		FSession::Session.GetDatabase().AddEntity(Entity);
+	}
+
+#endif
+
+	void FEntity::SerializeIdents(FCADKernelArchive& Ar, TArray<TOrientedEntity<FEntity>>& Array)
+	{
+		if (Ar.IsLoading())
+		{
+			int32 ArrayNum = 0;
+			Ar << ArrayNum;
+			Array.Init(TOrientedEntity<FEntity>(), ArrayNum);
+			for (TOrientedEntity<FEntity>& OrientedEntity : Array)
+			{
+				FIdent OldId = 0;
+				Ar << OldId;
+				Ar.SetReferencedEntityOrAddToWaitingList(OldId, OrientedEntity.Entity);
+				Ar.Serialize(&OrientedEntity.Direction, sizeof(EOrientation));
+			}
+		}
+		else
+		{
+			int32 ArrayNum = Array.Num();
+			Ar << ArrayNum;
+			for (TOrientedEntity<FEntity>& OrientedEntity : Array)
+			{
+				FIdent Id = OrientedEntity.Entity->GetId();
+				Ar << Id;
+				Ar.Serialize(&OrientedEntity.Direction, sizeof(EOrientation));
+				Ar.AddEntityToSave(Id);
+			}
+		}
+	}
+
+	void FEntity::SerializeIdents(FCADKernelArchive& Ar, TArray<FEntity*>& EntityArray, bool bSaveSelection)
+	{
+		if (Ar.IsLoading())
+		{
+			int32 ArrayNum = 0;
+			Ar << ArrayNum;
+			EntityArray.Init(nullptr, ArrayNum);
+			for (FEntity*& Entity : EntityArray)
+			{
+				FIdent OldId = 0;
+				Ar << OldId;
+				Ar.SetReferencedEntityOrAddToWaitingList(OldId, &Entity);
+			}
+		}
+		else
+		{
+			int32 ArrayNum = EntityArray.Num();
+			Ar << ArrayNum;
+			for (FEntity* Entity : EntityArray)
+			{
+				if (Entity)
+				{
+					FIdent Id = Entity->GetId();
+					Ar << Id;
+				}
+			}
+
+			if (bSaveSelection)
+			{
+				for (FEntity* Entity : EntityArray)
+				{
+					if (Entity)
+					{
+						Ar.AddEntityToSave(Entity->GetId());
+					}
+				}
+			}
+		}
+	}
+
+	void FEntity::SerializeIdents(FCADKernelArchive& Ar, TArray<TWeakPtr<FEntity>>& EntityArray, bool bSaveSelection)
+	{
+		if (Ar.IsLoading())
+		{
+			int32 ArrayNum = 0;
+			Ar << ArrayNum;
+			EntityArray.Init(TSharedPtr<FEntity>(), ArrayNum);
+			for (TWeakPtr<FEntity>& Entity : EntityArray)
+			{
+				FIdent OldId = 0;
+				Ar << OldId;
+				Ar.SetReferencedEntityOrAddToWaitingList(OldId, Entity);
+			}
+		}
+		else
+		{
+			int32 ArrayNum = EntityArray.Num();
+			Ar << ArrayNum;
+			for (TWeakPtr<FEntity>& Entity : EntityArray)
+			{
+				if (Entity.IsValid())
+				{
+					FIdent Id = Entity.Pin()->GetId();
+					Ar << Id;
+				}
+			}
+
+			if (bSaveSelection)
+			{
+				for (TWeakPtr<FEntity>& Entity : EntityArray)
+				{
+					if (Entity.IsValid())
+					{
+						Ar.AddEntityToSave(Entity.Pin()->GetId());
+					}
+				}
+			}
+		}
+	}
+
+	void FEntity::SerializeIdents(FCADKernelArchive& Ar, TArray<TSharedPtr<FEntity>>& EntityArray, bool bSaveSelection)
+	{
+		if (Ar.IsLoading())
+		{
+			int32 ArrayNum = 0;
+			Ar << ArrayNum;
+			EntityArray.Init(TSharedPtr<FEntity>(), ArrayNum);
+			for (TSharedPtr<FEntity>& Entity : EntityArray)
+			{
+				FIdent OldId = 0;
+				Ar << OldId;
+				Ar.SetReferencedEntityOrAddToWaitingList(OldId, Entity);
+			}
+		}
+		else
+		{
+			int32 ArrayNum = EntityArray.Num();
+			Ar << ArrayNum;
+			FIdent Id = 0;
+			for (TSharedPtr<FEntity>& Entity : EntityArray)
+			{
+				if (Entity.IsValid() && !Entity->IsDeleted())
+				{
+					Id = Entity->GetId();
+				}
+				else
+				{
+					Id = 0;
+				}
+				Ar << Id;
+			}
+
+			if (bSaveSelection)
+			{
+				for (TSharedPtr<FEntity>& Entity : EntityArray)
+				{
+					if(Entity.IsValid() && !Entity->IsDeleted())
+					{ 
+						Ar.AddEntityToSave(Entity->GetId());
+					}
+				}
+			}
+		}
+	}
+
+	void FEntity::SerializeIdent(FCADKernelArchive& Ar, TSharedPtr<FEntity>& Entity, bool bSaveSelection)
+	{
+		if (Ar.IsLoading())
+		{
+			FIdent OldId = 0;
+			Ar << OldId;
+			Ar.SetReferencedEntityOrAddToWaitingList(OldId, Entity);
+		}
+		else
+		{
+			FIdent Id = Entity.IsValid() ? Entity->GetId() : 0;
+			Ar << Id;
+			if (bSaveSelection && Id)
+			{
+				Ar.AddEntityToSave(Id);
+			}
+		}
+	}
+
+	void FEntity::SerializeIdent(FCADKernelArchive& Ar, TWeakPtr<FEntity>& Entity, bool bSaveSelection)
+	{
+		if (Ar.Archive.IsLoading())
+		{
+			FIdent OldId;
+			Ar.Archive << OldId;
+			Ar.SetReferencedEntityOrAddToWaitingList(OldId, Entity);
+		}
+		else
+		{
+			FIdent Id = Entity.IsValid() ? Entity.Pin()->GetId() : 0;
+			Ar.Archive << Id;
+			if (bSaveSelection && Id)
+			{
+				Ar.AddEntityToSave(Id);
+			}
+		}
+	}
+
+	void FEntity::SerializeIdent(FCADKernelArchive& Ar, FEntity** Entity, bool bSaveSelection)
+	{
+		if (Ar.Archive.IsLoading())
+		{
+			FIdent OldId;
+			Ar.Archive << OldId;
+			Ar.SetReferencedEntityOrAddToWaitingList(OldId, Entity);
+		}
+		else
+		{
+			FIdent Id = *Entity ? (*Entity)->GetId() : 0;
+			Ar.Archive << Id;
+			if (bSaveSelection && Id)
+			{
+				Ar.AddEntityToSave(Id);
+			}
+		}
+	}
+	
+	bool FEntity::SetId(FDatabase& Database)
+	{
+		if (Id < 1)
+		{
+			Database.AddEntity(AsShared());
+			++Database.EntityCount;
+			return true;
+		}
+		return Database.bForceSpawning || false;
+	}
+
+} // namespace UE::CADKernel
diff --git a/Engine/Source/Runtime/Datasmith/CADKernel/Private/CADKernel/Core/Entity.cpp b/Engine/Source/Runtime/Datasmith/CADKernel/Private/CADKernel/Core/Entity.cpp
deleted file mode 100644
diff --git a/Engine/Source/Runtime/Datasmith/CADKernel/Private/CADKernel/Core/Group.h b/Engine/Source/Runtime/Datasmith/CADKernel/Private/CADKernel/Core/Group.h
--- a/Engine/Source/Runtime/Datasmith/CADKernel/Private/CADKernel/Core/Group.h
+++ b/Engine/Source/Runtime/Datasmith/CADKernel/Private/CADKernel/Core/Group.h
@@ -1,7 +1,7 @@
 // Copyright Epic Games, Inc. All Rights Reserved.
 #pragma once
 
-#include "CADKernel/Core/Entity.h"
+#include "CADKernel/Core/CADEntity.h"
 #include "CADKernel/Core/CADKernelArchive.h"
 
 namespace UE::CADKernel
diff --git a/Engine/Source/Runtime/Datasmith/CADKernel/Private/CADKernel/Core/System.cpp b/Engine/Source/Runtime/Datasmith/CADKernel/Private/CADKernel/Core/System.cpp
--- a/Engine/Source/Runtime/Datasmith/CADKernel/Private/CADKernel/Core/System.cpp
+++ b/Engine/Source/Runtime/Datasmith/CADKernel/Private/CADKernel/Core/System.cpp
@@ -2,7 +2,7 @@
 #include "CADKernel/Core/System.h"
 
 #include "CADKernel/Core/Database.h"
-#include "CADKernel/Core/Entity.h"
+#include "CADKernel/Core/CADEntity.h"
 #include "CADKernel/UI/Visu.h"
 #include "CADKernel/Utils/Util.h"
 
diff --git a/Engine/Source/Runtime/Datasmith/CADKernel/Public/CADKernel/Core/CADEntity.h b/Engine/Source/Runtime/Datasmith/CADKernel/Public/CADKernel/Core/CADEntity.h
new file mode 100644
--- /dev/null
+++ b/Engine/Source/Runtime/Datasmith/CADKernel/Public/CADKernel/Core/CADEntity.h
@@ -0,0 +1,305 @@
+// Copyright Epic Games, Inc. All Rights Reserved.
+#pragma once
+
+#include "CADKernel/Core/CADKernelArchive.h"
+#include "CADKernel/Core/HaveStates.h"
+#include "CADKernel/Core/OrientedEntity.h"
+#include "CADKernel/Core/Types.h"
+#include "CADKernel/Geo/GeoEnum.h"
+#include "Serialization/Archive.h"
+
+#ifdef CADKERNEL_DEV
+#include "Toolkit/Core/InfoEntity.h"
+#endif
+
+namespace UE::CADKernel
+{
+enum class EEntity : uint8
+{
+	None = 0,
+
+	Curve,
+	Surface,
+
+	EdgeLink,
+	VertexLink,
+
+	TopologicalEdge,
+	TopologicalFace,
+
+	TopologicalLink,
+	TopologicalLoop,
+
+	TopologicalVertex,
+
+	Shell,
+	Body,
+	Model,
+
+	MeshModel,
+	Mesh,
+
+	Group,
+	Criterion,
+	Property,
+
+	EntityTypeEnd,
+};
+
+class FDatabase;
+class FEntity;
+class FInfoEntity;
+class FSession;
+class FSystem;
+
+class CADKERNEL_API FEntity : public TSharedFromThis<FEntity>, public FHaveStates
+{
+	friend class FDatabase;
+
+protected:
+	static const TCHAR* TypesNames[];
+	FIdent Id = 0;
+
+public:
+	virtual ~FEntity();
+
+	virtual void Delete()
+	{
+		Empty();
+		SetDeletedMarker();
+	}
+
+	virtual void Empty()
+	{
+
+	}
+
+	template<typename OtherEntity, typename... InArgTypes>
+	static TSharedRef<OtherEntity> MakeShared(InArgTypes&&... Args)
+	{
+		OtherEntity* Entity = new OtherEntity(Forward<InArgTypes>(Args)...);
+		TSharedRef<OtherEntity> NewShared = MakeShareable<OtherEntity>(Entity);
+
+#ifdef CADKERNEL_DEV
+		AddEntityInDatabase(NewShared);
+#endif
+		return NewShared;
+	}
+
+	template<typename OtherEntity>
+	static TSharedRef<OtherEntity> MakeShared(FCADKernelArchive& Archive)
+	{
+		OtherEntity* Entity = new OtherEntity();
+		TSharedRef<OtherEntity> NewShared = MakeShareable<OtherEntity>(Entity);
+		return NewShared;
+	}
+
+	/**
+	 * To Serialize a TSharedPtr<FEntity>, the Id of the FEntity is saved in the archive (ArchiveId).
+	 * At the deserialization of the archive, a map is build linking ArchiveId to the TSharedPtr of the deserialized entities.
+	 * At the deserialization of a TSharedPtr<FEntity>, the ArchiveId of the TSharedPtr<FEntity> is get. The map return the associated TSharedPtr.
+	 * If the Entity is not yet deserialized, the TSharedPtr is add to a waiting list to be set as soon as the associated Entity is deserialized
+	 * See FDatabase for more details
+	 */
+	static void SerializeIdent(FCADKernelArchive& Ar, TSharedPtr<FEntity>& Entity, bool bSaveSelection = true);
+	static void SerializeIdent(FCADKernelArchive& Ar, TWeakPtr<FEntity>& Entity, bool bSaveSelection = true);
+	static void SerializeIdent(FCADKernelArchive& Ar, FEntity** Entity, bool bSaveSelection = true);
+
+	template<typename EntityType>
+	static void SerializeIdent(FCADKernelArchive& Ar, EntityType** Entity, bool bSaveSelection = true)
+	{
+		SerializeIdent(Ar, (FEntity**)Entity, bSaveSelection);
+	}
+
+	template<typename EntityType>
+	static void SerializeIdent(FCADKernelArchive& Ar, TSharedPtr<EntityType>& Entity, bool bSaveSelection = true)
+	{
+		SerializeIdent(Ar, (TSharedPtr<FEntity>&) Entity, bSaveSelection);
+	}
+
+	template<typename EntityType>
+	static void SerializeIdent(FCADKernelArchive& Ar, TWeakPtr<EntityType>& Entity, bool bSaveSelection = true)
+	{
+		SerializeIdent(Ar, (TWeakPtr<FEntity>&) Entity, bSaveSelection);
+	}
+
+	/**
+	 * SerializeIdent of each TSharedPtr<FEntity> of the array
+	 */
+	static void SerializeIdents(FCADKernelArchive& Ar, TArray<FEntity*>& Array, bool bSaveSelection = true);
+	static void SerializeIdents(FCADKernelArchive& Ar, TArray<TWeakPtr<FEntity>>& Array, bool bSaveSelection = true);
+	static void SerializeIdents(FCADKernelArchive& Ar, TArray<TSharedPtr<FEntity>>& Array, bool bSaveSelection = true);
+	static void SerializeIdents(FCADKernelArchive& Ar, TArray<TOrientedEntity<FEntity>>& Array);
+
+	template<typename EntityType>
+	static void SerializeIdents(FCADKernelArchive& Ar, TArray<EntityType*>& Array, bool bSaveSelection = true)
+	{
+		SerializeIdents(Ar, (TArray<FEntity*>&) Array, bSaveSelection);
+	}
+
+	template<typename EntityType>
+	static void SerializeIdents(FCADKernelArchive& Ar, TArray<TWeakPtr<EntityType>>& Array, bool bSaveSelection = true)
+	{
+		SerializeIdents(Ar, (TArray<TWeakPtr<FEntity>>&) Array, bSaveSelection);
+	}
+
+	template<typename EntityType>
+	static void SerializeIdents(FCADKernelArchive& Ar, TArray<TSharedPtr<EntityType>>& Array, bool bSaveSelection = true)
+	{
+		SerializeIdents(Ar, (TArray<TSharedPtr<FEntity>>&) Array, bSaveSelection);
+	}
+
+	/**
+	 * Spawn Ident on each FEntity* of the array
+	 */
+	static void SpawnIdentOnEntities(TArray<FEntity*>& Array, FDatabase& Database)
+	{
+		for (FEntity* Entity : Array)
+		{
+			Entity->SpawnIdent(Database);
+		}
+	}
+
+	/**
+	 * Spawn Ident on each TSharedPtr<FEntity> of the array
+	 */
+	static void SpawnIdentOnEntities(TArray<TSharedPtr<FEntity>>& Array, FDatabase& Database)
+	{
+		for (TSharedPtr<FEntity>& Entity : Array)
+		{
+			Entity->SpawnIdent(Database);
+		}
+	}
+
+	static void SpawnIdentOnEntities(TArray<TOrientedEntity<FEntity>>& Array, FDatabase& Database)
+	{
+		for (TOrientedEntity<FEntity>& Entity : Array)
+		{
+			Entity.Entity->SpawnIdent(Database);
+		}
+	}
+
+	template<typename EntityType>
+	static void SpawnIdentOnEntities(TArray<TSharedPtr<EntityType>>& Array, FDatabase& Database)
+	{
+		SpawnIdentOnEntities((TArray<TSharedPtr<FEntity>>&) Array, Database);
+	}
+
+	/**
+	 * Reset Processed flags Recursively on each TSharedPtr<FEntity> of the array
+	 */
+	static void ResetMarkersRecursivelyOnEntities(const TArray<FEntity*>& Array)
+	{
+		for (const FEntity* Entity : Array)
+		{
+			Entity->ResetMarkersRecursively();
+		}
+	}
+
+	static void ResetMarkersRecursivelyOnEntities(const TArray<TWeakPtr<FEntity>>& Array)
+	{
+		for (const TWeakPtr<FEntity>& Entity : Array)
+		{
+			Entity.Pin()->ResetMarkersRecursively();
+		}
+	}
+
+	static void ResetMarkersRecursivelyOnEntities(const TArray<TSharedPtr<FEntity>>& Array)
+	{
+		for (const TSharedPtr<FEntity>& Entity : Array)
+		{
+			Entity->ResetMarkersRecursively();
+		}
+	}
+
+	static void ResetMarkersRecursivelyOnEntities(const TArray<TOrientedEntity<FEntity>>& Array)
+	{
+		for (const TOrientedEntity<FEntity>& OrientedEntity : Array)
+		{
+			OrientedEntity.Entity->ResetMarkersRecursively();
+		}
+	}
+
+	template<typename EntityType>
+	static void ResetMarkersRecursivelyOnEntities(const TArray<TSharedPtr<EntityType>>& Array)
+	{
+		ResetMarkersRecursivelyOnEntities((const TArray<TSharedPtr<FEntity>>&) Array);
+	}
+
+	/**
+	 * Serialization of a FEntity. Each class derived from FEntity has to override this method (and call the direct base class override method first)
+	 * E.g.
+	 * 	class FEntityXXX : public FEntityXX
+	 *  {
+	 *		virtual void Serialize(FCADKernelArchive& Ar) override
+	 * 		{
+	 *			FEntityXX::Serialize(Ar);
+	 *			...
+	 * 		}
+	 *	}
+	 */
+	virtual void Serialize(FCADKernelArchive& Ar)
+	{
+		Ar << Id;
+		FHaveStates::Serialize(Ar);
+	}
+
+	/**
+	 * Deserialize the next entity in the archive.
+	 */
+	static TSharedPtr<FEntity> Deserialize(FCADKernelArchive& Ar);
+
+	static const TCHAR* GetTypeName(EEntity Type);
+
+	virtual EEntity GetEntityType() const = 0;
+
+	bool IsTopologicalEntity() const
+	{
+		return (GetEntityType() >= EEntity::EdgeLink && GetEntityType() <= EEntity::Model);
+	}
+
+	bool IsTopologicalShapeEntity() const
+	{
+		return (GetEntityType() == EEntity::TopologicalFace) || (GetEntityType() >= EEntity::Shell && GetEntityType() <= EEntity::Model);
+	}
+
+	bool IsGeometricalEntity()
+	{
+		return (GetEntityType() == EEntity::Curve || GetEntityType() == EEntity::Surface);
+	}
+
+	const TCHAR* GetTypeName() const
+	{
+		return GetTypeName(GetEntityType());
+	}
+
+	const FIdent& GetId() const
+	{
+		return Id;
+	}
+
+#ifdef CADKERNEL_DEV
+	void InfoEntity() const;
+	virtual FInfoEntity& GetInfo(FInfoEntity& EntityInfo) const;
+	static void AddEntityInDatabase(TSharedRef<FEntity> Entity);
+#endif
+
+	virtual void SpawnIdent(FDatabase& Database)
+	{
+		if (!FEntity::SetId(Database))
+		{
+			return;
+		}
+	}
+
+	virtual void ResetMarkersRecursively() const
+	{
+		ResetMarkers();
+	}
+
+protected:
+
+	bool SetId(FDatabase& Database);
+};
+}
+
diff --git a/Engine/Source/Runtime/Datasmith/CADKernel/Public/CADKernel/Core/Database.h b/Engine/Source/Runtime/Datasmith/CADKernel/Public/CADKernel/Core/Database.h
--- a/Engine/Source/Runtime/Datasmith/CADKernel/Public/CADKernel/Core/Database.h
+++ b/Engine/Source/Runtime/Datasmith/CADKernel/Public/CADKernel/Core/Database.h
@@ -2,7 +2,7 @@
 #pragma once
 
 #include "CADKernel/Core/CADKernelArchive.h"
-#include "CADKernel/Core/Entity.h"
+#include "CADKernel/Core/CADEntity.h"
 #include "CADKernel/Core/Types.h"
 #include "Containers/Queue.h"
 
diff --git a/Engine/Source/Runtime/Datasmith/CADKernel/Public/CADKernel/Core/Entity.h b/Engine/Source/Runtime/Datasmith/CADKernel/Public/CADKernel/Core/Entity.h
deleted file mode 100644
diff --git a/Engine/Source/Runtime/Datasmith/CADKernel/Public/CADKernel/Core/EntityGeom.h b/Engine/Source/Runtime/Datasmith/CADKernel/Public/CADKernel/Core/EntityGeom.h
--- a/Engine/Source/Runtime/Datasmith/CADKernel/Public/CADKernel/Core/EntityGeom.h
+++ b/Engine/Source/Runtime/Datasmith/CADKernel/Public/CADKernel/Core/EntityGeom.h
@@ -2,7 +2,7 @@
 #pragma once
 
 #include "CADKernel/Core/CADKernelArchive.h"
-#include "CADKernel/Core/Entity.h"
+#include "CADKernel/Core/CADEntity.h"
 
 #include "CADKernel/Math/MatrixH.h"
 
diff --git a/Engine/Source/Runtime/Datasmith/CADKernel/Public/CADKernel/Core/Session.h b/Engine/Source/Runtime/Datasmith/CADKernel/Public/CADKernel/Core/Session.h
--- a/Engine/Source/Runtime/Datasmith/CADKernel/Public/CADKernel/Core/Session.h
+++ b/Engine/Source/Runtime/Datasmith/CADKernel/Public/CADKernel/Core/Session.h
@@ -3,7 +3,7 @@
 
 #include "CADKernel/Core/CADKernelArchive.h"
 #include "CADKernel/Core/Database.h"
-#include "CADKernel/Core/Entity.h"
+#include "CADKernel/Core/CADEntity.h"
 #include "CADKernel/Core/Types.h"
 #include "CADKernel/Math/Geometry.h"
 
diff --git a/Engine/Source/Runtime/Datasmith/CADKernel/Public/CADKernel/Geo/Curves/Curve.h b/Engine/Source/Runtime/Datasmith/CADKernel/Public/CADKernel/Geo/Curves/Curve.h
--- a/Engine/Source/Runtime/Datasmith/CADKernel/Public/CADKernel/Geo/Curves/Curve.h
+++ b/Engine/Source/Runtime/Datasmith/CADKernel/Public/CADKernel/Geo/Curves/Curve.h
@@ -2,7 +2,7 @@
 #pragma once
 
 #include "CADKernel/Core/CADKernelArchive.h"
-#include "CADKernel/Core/Entity.h"
+#include "CADKernel/Core/CADEntity.h"
 #include "CADKernel/Core/EntityGeom.h"
 #include "CADKernel/Core/Types.h"
 #include "CADKernel/Geo/GeoEnum.h"
diff --git a/Engine/Source/Runtime/Datasmith/CADKernel/Public/CADKernel/Geo/Curves/SurfacicCurve.h b/Engine/Source/Runtime/Datasmith/CADKernel/Public/CADKernel/Geo/Curves/SurfacicCurve.h
--- a/Engine/Source/Runtime/Datasmith/CADKernel/Public/CADKernel/Geo/Curves/SurfacicCurve.h
+++ b/Engine/Source/Runtime/Datasmith/CADKernel/Public/CADKernel/Geo/Curves/SurfacicCurve.h
@@ -1,7 +1,7 @@
 // Copyright Epic Games, Inc. All Rights Reserved.
 #pragma once
 
-#include "CADKernel/Core/Entity.h"
+#include "CADKernel/Core/CADEntity.h"
 #include "CADKernel/Core/Types.h"
 #include "CADKernel/Geo/Curves/Curve.h"
 #include "CADKernel/Geo/GeoEnum.h"
diff --git a/Engine/Source/Runtime/Datasmith/CADKernel/Public/CADKernel/Geo/Surfaces/Surface.h b/Engine/Source/Runtime/Datasmith/CADKernel/Public/CADKernel/Geo/Surfaces/Surface.h
--- a/Engine/Source/Runtime/Datasmith/CADKernel/Public/CADKernel/Geo/Surfaces/Surface.h
+++ b/Engine/Source/Runtime/Datasmith/CADKernel/Public/CADKernel/Geo/Surfaces/Surface.h
@@ -2,7 +2,7 @@
 #pragma once
 
 #include "CADKernel/Core/CADKernelArchive.h"
-#include "CADKernel/Core/Entity.h"
+#include "CADKernel/Core/CADEntity.h"
 #include "CADKernel/Core/EntityGeom.h"
 #include "CADKernel/Core/Types.h"
 #include "CADKernel/Geo/GeoEnum.h"
diff --git a/Engine/Source/Runtime/Datasmith/CADKernel/Public/CADKernel/Mesh/Criteria/Criterion.h b/Engine/Source/Runtime/Datasmith/CADKernel/Public/CADKernel/Mesh/Criteria/Criterion.h
--- a/Engine/Source/Runtime/Datasmith/CADKernel/Public/CADKernel/Mesh/Criteria/Criterion.h
+++ b/Engine/Source/Runtime/Datasmith/CADKernel/Public/CADKernel/Mesh/Criteria/Criterion.h
@@ -2,7 +2,7 @@
 #pragma once
 
 #include "CADKernel/Core/CADKernelArchive.h"
-#include "CADKernel/Core/Entity.h"
+#include "CADKernel/Core/CADEntity.h"
 #include "CADKernel/Core/Types.h"
 #include "CADKernel/Geo/GeoPoint.h"
 #include "CADKernel/Math/Point.h"
diff --git a/Engine/Source/Runtime/Datasmith/CADKernel/Public/CADKernel/Mesh/Structure/Mesh.h b/Engine/Source/Runtime/Datasmith/CADKernel/Public/CADKernel/Mesh/Structure/Mesh.h
--- a/Engine/Source/Runtime/Datasmith/CADKernel/Public/CADKernel/Mesh/Structure/Mesh.h
+++ b/Engine/Source/Runtime/Datasmith/CADKernel/Public/CADKernel/Mesh/Structure/Mesh.h
@@ -1,7 +1,7 @@
 // Copyright Epic Games, Inc. All Rights Reserved.
 #pragma once
 
-#include "CADKernel/Core/Entity.h"
+#include "CADKernel/Core/CADEntity.h"
 #include "CADKernel/Core/EntityGeom.h"
 #include "CADKernel/Core/Types.h"
 #include "CADKernel/Math/Point.h"
diff --git a/Engine/Source/Runtime/Datasmith/CADKernel/Public/CADKernel/Mesh/Structure/ModelMesh.h b/Engine/Source/Runtime/Datasmith/CADKernel/Public/CADKernel/Mesh/Structure/ModelMesh.h
--- a/Engine/Source/Runtime/Datasmith/CADKernel/Public/CADKernel/Mesh/Structure/ModelMesh.h
+++ b/Engine/Source/Runtime/Datasmith/CADKernel/Public/CADKernel/Mesh/Structure/ModelMesh.h
@@ -1,7 +1,7 @@
 // Copyright Epic Games, Inc. All Rights Reserved.
 #pragma once
 
-#include "CADKernel/Core/Entity.h"
+#include "CADKernel/Core/CADEntity.h"
 #include "CADKernel/Core/EntityGeom.h"
 #include "CADKernel/Core/Session.h"
 #include "CADKernel/Core/Types.h"
diff --git a/Engine/Source/Runtime/Datasmith/CADKernel/Public/CADKernel/Topo/Body.h b/Engine/Source/Runtime/Datasmith/CADKernel/Public/CADKernel/Topo/Body.h
--- a/Engine/Source/Runtime/Datasmith/CADKernel/Public/CADKernel/Topo/Body.h
+++ b/Engine/Source/Runtime/Datasmith/CADKernel/Public/CADKernel/Topo/Body.h
@@ -1,7 +1,7 @@
 // Copyright Epic Games, Inc. All Rights Reserved.
 #pragma once
 
-#include "CADKernel/Core/Entity.h"
+#include "CADKernel/Core/CADEntity.h"
 #include "CADKernel/Core/MetadataDictionary.h"
 #include "CADKernel/Geo/GeoEnum.h"
 #include "CADKernel/Topo/Shell.h"
diff --git a/Engine/Source/Runtime/Datasmith/CADKernel/Public/CADKernel/Topo/Model.h b/Engine/Source/Runtime/Datasmith/CADKernel/Public/CADKernel/Topo/Model.h
--- a/Engine/Source/Runtime/Datasmith/CADKernel/Public/CADKernel/Topo/Model.h
+++ b/Engine/Source/Runtime/Datasmith/CADKernel/Public/CADKernel/Topo/Model.h
@@ -2,7 +2,7 @@
 #pragma once
 
 #include "CADKernel/Core/CADKernelArchive.h"
-#include "CADKernel/Core/Entity.h"
+#include "CADKernel/Core/CADEntity.h"
 #include "CADKernel/Core/EntityGeom.h"
 #include "CADKernel/Topo/TopologicalShapeEntity.h"
 #include "CoreTypes.h"
diff --git a/Engine/Source/Runtime/Datasmith/CADKernel/Public/CADKernel/Topo/Shell.h b/Engine/Source/Runtime/Datasmith/CADKernel/Public/CADKernel/Topo/Shell.h
--- a/Engine/Source/Runtime/Datasmith/CADKernel/Public/CADKernel/Topo/Shell.h
+++ b/Engine/Source/Runtime/Datasmith/CADKernel/Public/CADKernel/Topo/Shell.h
@@ -2,7 +2,7 @@
 #pragma once
 
 #include "CADKernel/Core/CADKernelArchive.h"
-#include "CADKernel/Core/Entity.h"
+#include "CADKernel/Core/CADEntity.h"
 #include "CADKernel/Core/HaveStates.h"
 #include "CADKernel/Core/MetadataDictionary.h"
 #include "CADKernel/Core/OrientedEntity.h"
diff --git a/Engine/Source/Runtime/Datasmith/CADKernel/Public/CADKernel/Topo/TopologicalEntity.h b/Engine/Source/Runtime/Datasmith/CADKernel/Public/CADKernel/Topo/TopologicalEntity.h
--- a/Engine/Source/Runtime/Datasmith/CADKernel/Public/CADKernel/Topo/TopologicalEntity.h
+++ b/Engine/Source/Runtime/Datasmith/CADKernel/Public/CADKernel/Topo/TopologicalEntity.h
@@ -2,7 +2,7 @@
 #pragma once
 
 #include "CADKernel/Core/CADKernelArchive.h"
-#include "CADKernel/Core/Entity.h"
+#include "CADKernel/Core/CADEntity.h"
 #include "CADKernel/Core/HaveStates.h"
 #include "CADKernel/Core/Types.h"
 
diff --git a/Engine/Source/Runtime/Datasmith/CADKernel/Public/CADKernel/Topo/TopologicalFace.h b/Engine/Source/Runtime/Datasmith/CADKernel/Public/CADKernel/Topo/TopologicalFace.h
--- a/Engine/Source/Runtime/Datasmith/CADKernel/Public/CADKernel/Topo/TopologicalFace.h
+++ b/Engine/Source/Runtime/Datasmith/CADKernel/Public/CADKernel/Topo/TopologicalFace.h
@@ -1,7 +1,7 @@
 // Copyright Epic Games, Inc. All Rights Reserved.
 #pragma once
 
-#include "CADKernel/Core/Entity.h"
+#include "CADKernel/Core/CADEntity.h"
 #include "CADKernel/Core/HaveStates.h"
 #include "CADKernel/Core/Types.h"
 #include "CADKernel/Geo/GeoEnum.h"
diff --git a/Engine/Source/Runtime/Datasmith/CADKernel/Public/CADKernel/Topo/TopologicalLink.h b/Engine/Source/Runtime/Datasmith/CADKernel/Public/CADKernel/Topo/TopologicalLink.h
--- a/Engine/Source/Runtime/Datasmith/CADKernel/Public/CADKernel/Topo/TopologicalLink.h
+++ b/Engine/Source/Runtime/Datasmith/CADKernel/Public/CADKernel/Topo/TopologicalLink.h
@@ -1,7 +1,7 @@
 // Copyright Epic Games, Inc. All Rights Reserved.
 #pragma once
 
-#include "CADKernel/Core/Entity.h" 
+#include "CADKernel/Core/CADEntity.h" 
 #include "CADKernel/Math/Point.h"
 #include "CADKernel/Topo/TopologicalEntity.h"
 #include "CADKernel/UI/Message.h"
diff --git a/Engine/Source/Runtime/Datasmith/CADKernel/Public/CADKernel/Topo/TopologicalLoop.h b/Engine/Source/Runtime/Datasmith/CADKernel/Public/CADKernel/Topo/TopologicalLoop.h
--- a/Engine/Source/Runtime/Datasmith/CADKernel/Public/CADKernel/Topo/TopologicalLoop.h
+++ b/Engine/Source/Runtime/Datasmith/CADKernel/Public/CADKernel/Topo/TopologicalLoop.h
@@ -2,7 +2,7 @@
 #pragma once
 
 #include "CADKernel/Core/CADKernelArchive.h"
-#include "CADKernel/Core/Entity.h"
+#include "CADKernel/Core/CADEntity.h"
 #include "CADKernel/Core/OrientedEntity.h"
 #include "CADKernel/Geo/GeoEnum.h"
 #include "CADKernel/Math/Boundary.h"
diff --git a/Engine/Source/Runtime/Datasmith/CADKernel/Public/CADKernel/Topo/TopologicalVertex.h b/Engine/Source/Runtime/Datasmith/CADKernel/Public/CADKernel/Topo/TopologicalVertex.h
--- a/Engine/Source/Runtime/Datasmith/CADKernel/Public/CADKernel/Topo/TopologicalVertex.h
+++ b/Engine/Source/Runtime/Datasmith/CADKernel/Public/CADKernel/Topo/TopologicalVertex.h
@@ -2,7 +2,7 @@
 #pragma once
 
 #include "CADKernel/Core/CADKernelArchive.h"
-#include "CADKernel/Core/Entity.h"
+#include "CADKernel/Core/CADEntity.h"
 #include "CADKernel/Core/HaveStates.h"
 #include "CADKernel/Core/Types.h"
 #include "CADKernel/Math/Point.h"
