Datasmith - Remove the experimental Datasmith Clo json importer plugin, and deprecate unused Datasmith cloth code.

diff --git a/Engine/Plugins/Enterprise/DatasmithContent/Source/DatasmithContent/Private/DatasmithScene.cpp b/Engine/Plugins/Enterprise/DatasmithContent/Source/DatasmithContent/Private/DatasmithScene.cpp
--- a/Engine/Plugins/Enterprise/DatasmithContent/Source/DatasmithContent/Private/DatasmithScene.cpp
+++ b/Engine/Plugins/Enterprise/DatasmithContent/Source/DatasmithContent/Private/DatasmithScene.cpp
@@ -29,6 +29,7 @@ UDatasmithScene::UDatasmithScene()
 #endif
 }
 
+PRAGMA_DISABLE_DEPRECATION_WARNINGS // For Clothes destruction: UE_DEPRECATED(5.5, "The experimental Cloth importer is no longer supported.")
 UDatasmithScene::~UDatasmithScene()
 {
 #if WITH_EDITOR
@@ -38,6 +39,7 @@ UDatasmithScene::~UDatasmithScene()
 	}
 #endif
 }
+PRAGMA_ENABLE_DEPRECATION_WARNINGS
 
 void UDatasmithScene::GetAssetRegistryTags(TArray<FAssetRegistryTag>& OutTags) const
 {
diff --git a/Engine/Plugins/Enterprise/DatasmithContent/Source/DatasmithContent/Public/DatasmithScene.h b/Engine/Plugins/Enterprise/DatasmithContent/Source/DatasmithContent/Public/DatasmithScene.h
--- a/Engine/Plugins/Enterprise/DatasmithContent/Source/DatasmithContent/Public/DatasmithScene.h
+++ b/Engine/Plugins/Enterprise/DatasmithContent/Source/DatasmithContent/Public/DatasmithScene.h
@@ -42,6 +42,7 @@ public:
 	TMap< FName, TSoftObjectPtr< UStaticMesh > > StaticMeshes;
 
 	/** Map of all the cloth related to this Datasmith Scene */
+	UE_DEPRECATED(5.5, "The experimental Cloth importer is no longer supported.")
 	UPROPERTY(VisibleAnywhere, Category="Datasmith", AdvancedDisplay)
 	TMap< FName, TSoftObjectPtr< UObject > > Clothes; // UChaosClothAsset
 
diff --git a/Engine/Plugins/Enterprise/DatasmithImporter/Source/DatasmithImporter/Private/DatasmithActorImporter.cpp b/Engine/Plugins/Enterprise/DatasmithImporter/Source/DatasmithImporter/Private/DatasmithActorImporter.cpp
--- a/Engine/Plugins/Enterprise/DatasmithImporter/Source/DatasmithImporter/Private/DatasmithActorImporter.cpp
+++ b/Engine/Plugins/Enterprise/DatasmithImporter/Source/DatasmithImporter/Private/DatasmithActorImporter.cpp
@@ -3,7 +3,7 @@
 #include "DatasmithActorImporter.h"
 
 #include "DatasmithCameraImporter.h"
-#include "DatasmithCloth.h"
+#include "DatasmithCloth.h"  // UE_DEPRECATED(5.5, "The experimental Cloth importer is no longer supported.")
 #include "DatasmithImportContext.h"
 #include "DatasmithImporterModule.h"
 #include "DatasmithImportOptions.h"
@@ -337,7 +337,8 @@ UStaticMeshComponent* FDatasmithActorImporter::ImportStaticMeshComponent( FDatas
 	return StaticMeshComponent;
 }
 
-AActor* FDatasmithActorImporter::ImportClothActor(FDatasmithImportContext& ImportContext, const TSharedRef<IDatasmithClothActorElement>& ClothActorElement)
+PRAGMA_DISABLE_DEPRECATION_WARNINGS
+AActor* FDatasmithActorImporter::ImportClothActor(FDatasmithImportContext& ImportContext, const TSharedRef<IDatasmithClothActorElement>& ClothActorElement)  // UE_DEPRECATED(5.5, "The experimental Cloth importer is no longer supported.")
 {
 	AActor* ImportedActor = ImportActor(AActor::StaticClass(), ClothActorElement, ImportContext, ImportContext.Options->StaticMeshActorImportPolicy);
 
@@ -394,6 +395,7 @@ AActor* FDatasmithActorImporter::ImportClothActor(FDatasmithImportContext& Impor
 	}
 	return ImportedActor;
 }
+PRAGMA_ENABLE_DEPRECATION_WARNINGS
 
 void FDatasmithActorImporter::SetupStaticMeshComponent( FDatasmithImportContext& ImportContext, UStaticMeshComponent* StaticMeshComponent, const TSharedRef< IDatasmithMeshActorElement >& MeshActorElement )
 {
diff --git a/Engine/Plugins/Enterprise/DatasmithImporter/Source/DatasmithImporter/Private/DatasmithImportContext.cpp b/Engine/Plugins/Enterprise/DatasmithImporter/Source/DatasmithImporter/Private/DatasmithImportContext.cpp
--- a/Engine/Plugins/Enterprise/DatasmithImporter/Source/DatasmithImporter/Private/DatasmithImportContext.cpp
+++ b/Engine/Plugins/Enterprise/DatasmithImporter/Source/DatasmithImporter/Private/DatasmithImportContext.cpp
@@ -555,10 +555,12 @@ void FDatasmithImportContext::FInternalReferenceCollector::AddReferencedObjects(
 		Collector.AddReferencedObject( It.Value );
 	}
 
-	for (auto& Pair : ImportContext->ImportedClothes)
+PRAGMA_DISABLE_DEPRECATION_WARNINGS
+	for (auto& Pair : ImportContext->ImportedClothes)  // UE_DEPRECATED(5.5, "The experimental Cloth importer is no longer supported.")
 	{
 		Collector.AddReferencedObject(Pair.Value);
 	}
+PRAGMA_ENABLE_DEPRECATION_WARNINGS
 
 	Collector.AddReferencedObjects(ImportContext->ImportedMaterials);
 	Collector.AddReferencedObjects(ImportContext->ImportedParentMaterials);
diff --git a/Engine/Plugins/Enterprise/DatasmithImporter/Source/DatasmithImporter/Private/DatasmithImportFactory.cpp b/Engine/Plugins/Enterprise/DatasmithImporter/Source/DatasmithImporter/Private/DatasmithImportFactory.cpp
--- a/Engine/Plugins/Enterprise/DatasmithImporter/Source/DatasmithImporter/Private/DatasmithImportFactory.cpp
+++ b/Engine/Plugins/Enterprise/DatasmithImporter/Source/DatasmithImporter/Private/DatasmithImportFactory.cpp
@@ -263,14 +263,16 @@ namespace DatasmithImportFactoryImpl
 			return false;
 		}
 
+PRAGMA_DISABLE_DEPRECATION_WARNINGS
 		// CLOTH
-		FDatasmithImporter::ImportClothes( InContext );
+		FDatasmithImporter::ImportClothes( InContext );  // UE_DEPRECATED(5.5, "The experimental Cloth importer is no longer supported.")
 
 		if ( InContext.bUserCancelled )
 		{
 			bOutOperationCancelled = true;
 			return false;
 		}
+PRAGMA_ENABLE_DEPRECATION_WARNINGS
 
 		// ACTORS
 		if( InContext.ShouldImportActors() )
diff --git a/Engine/Plugins/Enterprise/DatasmithImporter/Source/DatasmithImporter/Private/DatasmithImporter.cpp b/Engine/Plugins/Enterprise/DatasmithImporter/Source/DatasmithImporter/Private/DatasmithImporter.cpp
--- a/Engine/Plugins/Enterprise/DatasmithImporter/Source/DatasmithImporter/Private/DatasmithImporter.cpp
+++ b/Engine/Plugins/Enterprise/DatasmithImporter/Source/DatasmithImporter/Private/DatasmithImporter.cpp
@@ -8,7 +8,7 @@
 #include "DatasmithAssetImportData.h"
 #include "DatasmithAssetUserData.h"
 #include "DatasmithCameraImporter.h"
-#include "DatasmithCloth.h"
+#include "DatasmithCloth.h"  // UE_DEPRECATED(5.5, "The experimental Cloth importer is no longer supported.")
 #include "DatasmithImportContext.h"
 #include "DatasmithImporterModule.h"
 #include "DatasmithLevelSequenceImporter.h"
@@ -371,7 +371,8 @@ void FDatasmithImporter::ImportStaticMeshes( FDatasmithImportContext& ImportCont
 	}
 }
 
-void FDatasmithImporter::ImportClothes(FDatasmithImportContext& ImportContext)
+PRAGMA_DISABLE_DEPRECATION_WARNINGS
+void FDatasmithImporter::ImportClothes(FDatasmithImportContext& ImportContext)  // UE_DEPRECATED(5.5, "The experimental Cloth importer is no longer supported.")
 {
 	const int32 ClothesCount = ImportContext.FilteredScene->GetClothesCount();
 
@@ -517,6 +518,7 @@ void FDatasmithImporter::ImportClothes(FDatasmithImportContext& ImportContext)
 		//}
 	}
 }
+PRAGMA_ENABLE_DEPRECATION_WARNINGS
 
 UStaticMesh* FDatasmithImporter::ImportStaticMesh( FDatasmithImportContext& ImportContext, TSharedRef< IDatasmithMeshElement > MeshElement, UStaticMesh* ExistingStaticMesh, FDatasmithMeshElementPayload* MeshPayload)
 {
@@ -589,7 +591,7 @@ UStaticMesh* FDatasmithImporter::FinalizeStaticMesh( UStaticMesh* SourceStaticMe
 	return DestinationStaticMesh;
 }
 
-UObject* FDatasmithImporter::FinalizeCloth(UObject* SourceCloth, const TCHAR* FolderPath, UObject* ExistingCloth, TMap<UObject*, UObject*>* ReferencesToRemap)
+UObject* FDatasmithImporter::FinalizeCloth(UObject* SourceCloth, const TCHAR* FolderPath, UObject* ExistingCloth, TMap<UObject*, UObject*>* ReferencesToRemap)  // UE_DEPRECATED(5.5, "The experimental Cloth importer is no longer supported.")
 {
 	if (Cast<USkinnedAsset>(ExistingCloth))
 	{
@@ -1115,10 +1117,12 @@ AActor* FDatasmithImporter::ImportActor( FDatasmithImportContext& ImportContext,
 	{
 		ImportedActor = FDatasmithActorImporter::ImportStaticMeshActor( ImportContext, StaticCastSharedRef< IDatasmithMeshActorElement >( ActorElement ) );
 	}
-	else if (ActorElement->IsA(EDatasmithElementType::ClothActor))
+PRAGMA_DISABLE_DEPRECATION_WARNINGS
+	else if (ActorElement->IsA(EDatasmithElementType::ClothActor))  // UE_DEPRECATED(5.5, "The experimental Cloth importer is no longer supported.")
 	{
 		ImportedActor = FDatasmithActorImporter::ImportClothActor(ImportContext, StaticCastSharedRef<IDatasmithClothActorElement>( ActorElement ) );
 	}
+PRAGMA_ENABLE_DEPRECATION_WARNINGS
 	else if (ActorElement->IsA(EDatasmithElementType::EnvironmentLight))
 	{
 		ImportedActor = FDatasmithActorImporter::ImportEnvironment( ImportContext, StaticCastSharedRef< IDatasmithEnvironmentElement >( ActorElement ) );
@@ -1894,8 +1898,10 @@ void FDatasmithImporter::FinalizeImport(FDatasmithImportContext& ImportContext,
 	TRACE_CPUPROFILER_EVENT_SCOPE(FDatasmithImporter::FinalizeImport);
 
 	const int32 NumImportedObjects = ImportContext.ImportedStaticMeshes.Num() +
-									 ImportContext.ImportedClothes.Num() +
-									 ImportContext.ImportedClothPresets.Num() +
+PRAGMA_DISABLE_DEPRECATION_WARNINGS
+									 ImportContext.ImportedClothes.Num() +       // UE_DEPRECATED(5.5, "The experimental Cloth importer is no longer supported.")
+									 ImportContext.ImportedClothPresets.Num() +  //
+PRAGMA_ENABLE_DEPRECATION_WARNINGS
 									 ImportContext.ImportedTextures.Num() +
 									 ImportContext.ImportedMaterialFunctions.Num() +
 									 ImportContext.ImportedMaterials.Num() +
@@ -2097,7 +2103,8 @@ void FDatasmithImporter::FinalizeImport(FDatasmithImportContext& ImportContext,
 
 	FDatasmithStaticMeshImporter::BuildStaticMeshes(StaticMeshes.Array(), ProgressFunction);
 
-	for (auto& SourcePresetObj : ImportContext.ImportedClothPresets)
+PRAGMA_DISABLE_DEPRECATION_WARNINGS
+	for (auto& SourcePresetObj : ImportContext.ImportedClothPresets)  // UE_DEPRECATED(5.5, "The experimental Cloth importer is no longer supported.")
 	{
 		if (!SourcePresetObj || (ValidAssets.Num() > 0 && !ValidAssets.Contains(SourcePresetObj)))
 		{
@@ -2118,7 +2125,7 @@ void FDatasmithImporter::FinalizeImport(FDatasmithImportContext& ImportContext,
 		SourcePresetObj = FinalizedPreset;
 	}
 
-	for (auto& ImportedPair : ImportContext.ImportedClothes)
+	for (auto& ImportedPair : ImportContext.ImportedClothes)  // UE_DEPRECATED(5.5, "The experimental Cloth importer is no longer supported.")
 	{
 		UObject* SourceClothObj = ImportedPair.Value;
 
@@ -2144,6 +2151,7 @@ void FDatasmithImporter::FinalizeImport(FDatasmithImportContext& ImportContext,
 
 		ImportedPair.Value = FinalizedCloth;
 	}
+PRAGMA_ENABLE_DEPRECATION_WARNINGS
 
 	for (const auto& ImportedLevelSequencePair : ImportContext.ImportedLevelSequences)
 	{
diff --git a/Engine/Plugins/Enterprise/DatasmithImporter/Source/DatasmithImporter/Public/DatasmithActorImporter.h b/Engine/Plugins/Enterprise/DatasmithImporter/Source/DatasmithImporter/Public/DatasmithActorImporter.h
--- a/Engine/Plugins/Enterprise/DatasmithImporter/Source/DatasmithImporter/Public/DatasmithActorImporter.h
+++ b/Engine/Plugins/Enterprise/DatasmithImporter/Source/DatasmithImporter/Public/DatasmithActorImporter.h
@@ -11,7 +11,7 @@ class AStaticMeshActor;
 class FDatasmithActorUniqueLabelProvider;
 class IDatasmithActorElement;
 class IDatasmithCameraActorElement;
-class IDatasmithClothActorElement;
+class UE_DEPRECATED(5.5, "The experimental Cloth importer is no longer supported.") IDatasmithClothActorElement;
 class IDatasmithCustomActorElement;
 class IDatasmithDecalActorElement;
 class IDatasmithEnvironmentElement;
@@ -72,7 +72,10 @@ public:
 	static AStaticMeshActor* ImportStaticMeshActor( FDatasmithImportContext& ImportContext, const TSharedRef< IDatasmithMeshActorElement >& MeshActorElement );
 	static UStaticMeshComponent* ImportStaticMeshComponent( FDatasmithImportContext& ImportContext, const TSharedRef< IDatasmithMeshActorElement >& MeshActorElement, UObject* Outer, FDatasmithActorUniqueLabelProvider& UniqueNameProvider );
 
+PRAGMA_DISABLE_DEPRECATION_WARNINGS
+	UE_DEPRECATED(5.5, "The experimental Cloth importer is no longer supported.")
 	static AActor* ImportClothActor( FDatasmithImportContext& ImportContext, const TSharedRef< IDatasmithClothActorElement >& ClothActorElement );
+PRAGMA_ENABLE_DEPRECATION_WARNINGS
 
 	/**
 	 * Spawns a cine camera actor
diff --git a/Engine/Plugins/Enterprise/DatasmithImporter/Source/DatasmithImporter/Public/DatasmithImportContext.h b/Engine/Plugins/Enterprise/DatasmithImporter/Source/DatasmithImporter/Public/DatasmithImportContext.h
--- a/Engine/Plugins/Enterprise/DatasmithImporter/Source/DatasmithImporter/Public/DatasmithImportContext.h
+++ b/Engine/Plugins/Enterprise/DatasmithImporter/Source/DatasmithImporter/Public/DatasmithImportContext.h
@@ -21,7 +21,7 @@ struct FDatasmithImportContext;
 class FJsonObject;
 class FMessageLogModule;
 class IDatasmithBaseMaterialElement;
-class IDatasmithClothElement;
+class UE_DEPRECATED(5.5, "The experimental Cloth importer is no longer supported.") IDatasmithClothElement;
 class IDatasmithElement;
 class IDatasmithLevelSequenceElement;
 class IDatasmithLevelVariantSetsElement;
@@ -257,9 +257,13 @@ struct DATASMITHIMPORTER_API FDatasmithImportContext
 	/** Map of imported mesh for each mesh element */
 	TMap< TSharedRef< IDatasmithMeshElement >, TObjectPtr<UStaticMesh> > ImportedStaticMeshes;
 
+PRAGMA_DISABLE_DEPRECATION_WARNINGS
+	UE_DEPRECATED(5.5, "The experimental Cloth importer is no longer supported.")
 	TMap< TSharedRef< IDatasmithClothElement >, TObjectPtr<UObject> > ImportedClothes;
 
+	UE_DEPRECATED(5.5, "The experimental Cloth importer is no longer supported.")
 	TArray< UObject* > ImportedClothPresets;
+PRAGMA_ENABLE_DEPRECATION_WARNINGS
 
 	/** Register IDatasmithMeshElement by their name so they can be searched faster */
 	TMap< FString, TSharedRef < IDatasmithMeshElement > > ImportedStaticMeshesByName;
diff --git a/Engine/Plugins/Enterprise/DatasmithImporter/Source/DatasmithImporter/Public/DatasmithImporter.h b/Engine/Plugins/Enterprise/DatasmithImporter/Source/DatasmithImporter/Public/DatasmithImporter.h
--- a/Engine/Plugins/Enterprise/DatasmithImporter/Source/DatasmithImporter/Public/DatasmithImporter.h
+++ b/Engine/Plugins/Enterprise/DatasmithImporter/Source/DatasmithImporter/Public/DatasmithImporter.h
@@ -38,6 +38,7 @@ struct DATASMITHIMPORTER_API FDatasmithImporter
 	 */
 	static void ImportStaticMeshes( FDatasmithImportContext& ImportContext );
 
+	UE_DEPRECATED(5.5, "The experimental Cloth importer is no longer supported.")
 	static void ImportClothes( FDatasmithImportContext& ImportContext );
 
 	/**
@@ -57,6 +58,8 @@ struct DATASMITHIMPORTER_API FDatasmithImporter
 	 * @return The resulting static mesh
 	 */
 	static UStaticMesh* FinalizeStaticMesh( UStaticMesh* SourceStaticMesh, const TCHAR* StaticMeshesFolderPath, UStaticMesh* ExistingStaticMesh, TMap< UObject*, UObject* >* ReferencesToRemap = nullptr, bool bBuild = true );
+
+	UE_DEPRECATED(5.5, "The experimental Cloth importer is no longer supported.")
 	static UObject* FinalizeCloth( UObject* SourceCloth, const TCHAR* FolderPath, UObject* ExistingCloth, TMap<UObject*, UObject*>* ReferencesToRemap = nullptr);
 
 	/**
diff --git a/Engine/Plugins/Enterprise/DatasmithImporter/Source/DatasmithNativeTranslator/Private/DatasmithNativeTranslator.cpp b/Engine/Plugins/Enterprise/DatasmithImporter/Source/DatasmithNativeTranslator/Private/DatasmithNativeTranslator.cpp
--- a/Engine/Plugins/Enterprise/DatasmithImporter/Source/DatasmithNativeTranslator/Private/DatasmithNativeTranslator.cpp
+++ b/Engine/Plugins/Enterprise/DatasmithImporter/Source/DatasmithNativeTranslator/Private/DatasmithNativeTranslator.cpp
@@ -43,12 +43,14 @@ void FDatasmithNativeTranslator::ResolveSceneFilePaths(TSharedRef<IDatasmithScen
 		Mesh->SetFile(*ResolveFilePath(Path, ResourcePaths));
 	}
 
-	for (int32 Index = 0; Index < Scene->GetClothesCount(); ++Index)
+PRAGMA_DISABLE_DEPRECATION_WARNINGS
+	for (int32 Index = 0; Index < Scene->GetClothesCount(); ++Index)  // UE_DEPRECATED(5.5, "The experimental Cloth importer is no longer supported.")
 	{
 		const TSharedPtr<IDatasmithClothElement>& Cloth = Scene->GetCloth(Index);
 		const TCHAR* Path = Cloth->GetFile();
 		Cloth->SetFile(*ResolveFilePath(Path, ResourcePaths));
 	}
+PRAGMA_ENABLE_DEPRECATION_WARNINGS
 
 	for (int32 Index = 0; Index < Scene->GetTexturesCount(); ++Index)
 	{
@@ -159,7 +161,8 @@ bool FDatasmithNativeTranslator::LoadLevelSequence(const TSharedRef<IDatasmithLe
 	return false;
 }
 
-bool FDatasmithNativeTranslator::LoadCloth(const TSharedRef<IDatasmithClothElement> ClothElement, FDatasmithClothElementPayload& OutClothPayload)
+PRAGMA_DISABLE_DEPRECATION_WARNINGS
+bool FDatasmithNativeTranslator::LoadCloth(const TSharedRef<IDatasmithClothElement> ClothElement, FDatasmithClothElementPayload& OutClothPayload)  // UE_DEPRECATED(5.5, "The experimental Cloth importer is no longer supported.")
 {
 	TRACE_CPUPROFILER_EVENT_SCOPE(FDatasmithNativeTranslator::LoadCloth);
 
@@ -174,4 +177,5 @@ bool FDatasmithNativeTranslator::LoadCloth(const TSharedRef<IDatasmithClothEleme
 
 	return false;
 }
+PRAGMA_ENABLE_DEPRECATION_WARNINGS
 
diff --git a/Engine/Plugins/Enterprise/DatasmithImporter/Source/DatasmithNativeTranslator/Public/DatasmithNativeTranslator.h b/Engine/Plugins/Enterprise/DatasmithImporter/Source/DatasmithNativeTranslator/Public/DatasmithNativeTranslator.h
--- a/Engine/Plugins/Enterprise/DatasmithImporter/Source/DatasmithNativeTranslator/Public/DatasmithNativeTranslator.h
+++ b/Engine/Plugins/Enterprise/DatasmithImporter/Source/DatasmithNativeTranslator/Public/DatasmithNativeTranslator.h
@@ -19,7 +19,10 @@ public:
 
 	virtual bool LoadLevelSequence(const TSharedRef<IDatasmithLevelSequenceElement> LevelSequenceElement, FDatasmithLevelSequencePayload& OutLevelSequencePayload) override;
 
+PRAGMA_DISABLE_DEPRECATION_WARNINGS
+	UE_DEPRECATED(5.5, "The experimental Cloth importer is no longer supported.")
 	virtual bool LoadCloth(const TSharedRef<IDatasmithClothElement> ClothElement, FDatasmithClothElementPayload& OutClothPayload) override;
+PRAGMA_ENABLE_DEPRECATION_WARNINGS
 
 protected:
 	static FString ResolveFilePath(const FString& FilePath, const TArray<FString>& ResourcePaths);
diff --git a/Engine/Plugins/Enterprise/DatasmithImporter/Source/DatasmithTranslator/Public/DatasmithPayload.h b/Engine/Plugins/Enterprise/DatasmithImporter/Source/DatasmithTranslator/Public/DatasmithPayload.h
--- a/Engine/Plugins/Enterprise/DatasmithImporter/Source/DatasmithTranslator/Public/DatasmithPayload.h
+++ b/Engine/Plugins/Enterprise/DatasmithImporter/Source/DatasmithTranslator/Public/DatasmithPayload.h
@@ -2,7 +2,7 @@
 
 #pragma once
 
-#include "DatasmithCloth.h"
+#include "DatasmithCloth.h"  // UE_DEPRECATED(5.5, "The experimental Cloth importer is no longer supported.")
 #include "MeshDescription.h"
 
 
@@ -23,9 +23,11 @@ struct FDatasmithMeshElementPayload : public FDatasmithPayload
 /**
  * Describes a Cloth element payload, which is the actual data to be imported.
  */
-struct FDatasmithClothElementPayload : public FDatasmithPayload
+struct UE_DEPRECATED(5.5, "The experimental Cloth importer is no longer supported.") FDatasmithClothElementPayload : public FDatasmithPayload
 {
+PRAGMA_DISABLE_DEPRECATION_WARNINGS
 	FDatasmithCloth Cloth;
+PRAGMA_ENABLE_DEPRECATION_WARNINGS
 };
 
 
diff --git a/Engine/Plugins/Enterprise/DatasmithImporter/Source/DatasmithTranslator/Public/DatasmithTranslator.h b/Engine/Plugins/Enterprise/DatasmithImporter/Source/DatasmithTranslator/Public/DatasmithTranslator.h
--- a/Engine/Plugins/Enterprise/DatasmithImporter/Source/DatasmithTranslator/Public/DatasmithTranslator.h
+++ b/Engine/Plugins/Enterprise/DatasmithImporter/Source/DatasmithTranslator/Public/DatasmithTranslator.h
@@ -13,7 +13,7 @@
 
 class FDatasmithSceneSource;
 class IDatasmithScene;
-class IDatasmithClothElement;
+class UE_DEPRECATED(5.5, "The experimental Cloth importer is no longer supported.") IDatasmithClothElement;
 class IDatasmithMeshElement;
 class IDatasmithLevelSequenceElement;
 
@@ -104,7 +104,10 @@ public:
 	 * @param OutClothPayload    Actual cloth data from the source
 	 * @returns                  Operation succeed
 	 */
+PRAGMA_DISABLE_DEPRECATION_WARNINGS
+	UE_DEPRECATED(5.5, "The experimental Cloth importer is no longer supported.")
 	virtual bool LoadCloth(const TSharedRef<IDatasmithClothElement> ClothElement, FDatasmithClothElementPayload& OutClothPayload) { return false; }
+PRAGMA_ENABLE_DEPRECATION_WARNINGS
 
 	/**
 	 * Get payload related to the given Element
diff --git a/Engine/Plugins/Experimental/DatasmithCloTranslator/DatasmithCloTranslator.uplugin b/Engine/Plugins/Experimental/DatasmithCloTranslator/DatasmithCloTranslator.uplugin
deleted file mode 100644
diff --git a/Engine/Plugins/Experimental/DatasmithCloTranslator/Source/DatasmithCloTranslator/DatasmithCloTranslator.Build.cs b/Engine/Plugins/Experimental/DatasmithCloTranslator/Source/DatasmithCloTranslator/DatasmithCloTranslator.Build.cs
deleted file mode 100644
diff --git a/Engine/Plugins/Experimental/DatasmithCloTranslator/Source/DatasmithCloTranslator/Private/DatasmithCloTranslator.cpp b/Engine/Plugins/Experimental/DatasmithCloTranslator/Source/DatasmithCloTranslator/Private/DatasmithCloTranslator.cpp
deleted file mode 100644
diff --git a/Engine/Plugins/Experimental/DatasmithCloTranslator/Source/DatasmithCloTranslator/Private/DatasmithCloTranslator.h b/Engine/Plugins/Experimental/DatasmithCloTranslator/Source/DatasmithCloTranslator/Private/DatasmithCloTranslator.h
deleted file mode 100644
diff --git a/Engine/Plugins/Experimental/DatasmithCloTranslator/Source/DatasmithCloTranslator/Private/DatasmithCloTranslatorModule.cpp b/Engine/Plugins/Experimental/DatasmithCloTranslator/Source/DatasmithCloTranslator/Private/DatasmithCloTranslatorModule.cpp
deleted file mode 100644
diff --git a/Engine/Plugins/Experimental/DatasmithCloTranslator/Source/DatasmithCloTranslator/Public/DatasmithCloImportOptions.h b/Engine/Plugins/Experimental/DatasmithCloTranslator/Source/DatasmithCloTranslator/Public/DatasmithCloImportOptions.h
deleted file mode 100644
diff --git a/Engine/Plugins/Experimental/DatasmithCloTranslator/Source/DatasmithCloTranslator/Public/DatasmithCloTranslatorModule.h b/Engine/Plugins/Experimental/DatasmithCloTranslator/Source/DatasmithCloTranslator/Public/DatasmithCloTranslatorModule.h
deleted file mode 100644
diff --git a/Engine/Source/Developer/Datasmith/DatasmithExporter/Private/DatasmithMeshExporter.cpp b/Engine/Source/Developer/Datasmith/DatasmithExporter/Private/DatasmithMeshExporter.cpp
--- a/Engine/Source/Developer/Datasmith/DatasmithExporter/Private/DatasmithMeshExporter.cpp
+++ b/Engine/Source/Developer/Datasmith/DatasmithExporter/Private/DatasmithMeshExporter.cpp
@@ -3,7 +3,7 @@
 #include "DatasmithMeshExporter.h"
 
 #include "DatasmithCore.h"
-#include "DatasmithCloth.h"
+#include "DatasmithCloth.h"  // UE_DEPRECATED(5.5, "The experimental Cloth importer is no longer supported.")
 #include "DatasmithExporterManager.h"
 #include "DatasmithMesh.h"
 #include "DatasmithMeshUObject.h"
@@ -173,7 +173,8 @@ bool FDatasmithMeshExporter::ExportToUObject(TSharedPtr<IDatasmithMeshElement>&
 }
 
 
-bool FDatasmithMeshExporter::ExportCloth(FDatasmithCloth& Cloth, TSharedPtr<IDatasmithClothElement>& ClothElement, const TCHAR* FilePath, const TCHAR* AssetsOutputPath) const
+PRAGMA_DISABLE_DEPRECATION_WARNINGS
+bool FDatasmithMeshExporter::ExportCloth(FDatasmithCloth& Cloth, TSharedPtr<IDatasmithClothElement>& ClothElement, const TCHAR* FilePath, const TCHAR* AssetsOutputPath) const  // UE_DEPRECATED(5.5, "The experimental Cloth importer is no longer supported.")
 {
 	FString Path = FPaths::SetExtension(FilePath, TEXT(".udscloth"));
 	FPaths::NormalizeFilename(Path);
@@ -200,6 +201,7 @@ bool FDatasmithMeshExporter::ExportCloth(FDatasmithCloth& Cloth, TSharedPtr<IDat
 
 	return true;
 }
+PRAGMA_ENABLE_DEPRECATION_WARNINGS
 
 FString FDatasmithMeshExporter::GetLastError() const
 {
diff --git a/Engine/Source/Developer/Datasmith/DatasmithExporter/Public/DatasmithMeshExporter.h b/Engine/Source/Developer/Datasmith/DatasmithExporter/Public/DatasmithMeshExporter.h
--- a/Engine/Source/Developer/Datasmith/DatasmithExporter/Public/DatasmithMeshExporter.h
+++ b/Engine/Source/Developer/Datasmith/DatasmithExporter/Public/DatasmithMeshExporter.h
@@ -42,7 +42,10 @@ public:
 	 */
 	bool ExportToUObject(TSharedPtr<IDatasmithMeshElement>& MeshElement, const TCHAR* Filepath, FDatasmithMesh& Mesh, FDatasmithMesh* CollisionMesh, EDSExportLightmapUV LightmapUV);
 
+PRAGMA_DISABLE_DEPRECATION_WARNINGS
+	UE_DEPRECATED(5.5, "The experimental Cloth importer is no longer supported.")
 	bool ExportCloth(class FDatasmithCloth& Cloth, TSharedPtr<class IDatasmithClothElement>& ClothElement, const TCHAR* FilePath, const TCHAR* AssetsOutputPath) const;
+PRAGMA_ENABLE_DEPRECATION_WARNINGS
 
 	/**
 	 * @return The error that happened during the last export, if any
diff --git a/Engine/Source/Runtime/Datasmith/DatasmithCore/Private/DatasmithCloth.cpp b/Engine/Source/Runtime/Datasmith/DatasmithCore/Private/DatasmithCloth.cpp
--- a/Engine/Source/Runtime/Datasmith/DatasmithCore/Private/DatasmithCloth.cpp
+++ b/Engine/Source/Runtime/Datasmith/DatasmithCore/Private/DatasmithCloth.cpp
@@ -5,8 +5,7 @@
 #include "Serialization/CustomVersion.h"
 
 
-
-enum EDatasmithClothSerializationVersion
+enum UE_DEPRECATED(5.5, "The experimental Cloth importer is no longer supported.") EDatasmithClothSerializationVersion
 {
     EDCSV_Base = 0,
     EDCSV_WithPatternParameters = 1,
@@ -17,7 +16,7 @@ enum EDatasmithClothSerializationVersion
     _EDCSV_Last = _EDCSV_Count - 1
 };
 
-struct FDatasmithClothSerializationVersion
+struct UE_DEPRECATED(5.5, "The experimental Cloth importer is no longer supported.") FDatasmithClothSerializationVersion
 {
     const static FGuid GUID;
 
@@ -27,11 +26,14 @@ private:
 
 const FGuid FDatasmithClothSerializationVersion::GUID(0x28B01036, 0x66B4498F, 0x99425ACA, 0xDB78A9B5);
 
+PRAGMA_DISABLE_DEPRECATION_WARNINGS
+
 // Register the custom version with core
+UE_DEPRECATED(5.5, "The experimental Cloth importer is no longer supported.")
 FCustomVersionRegistration GRegisterDatasmithClothCustomVersion(FDatasmithClothSerializationVersion::GUID, _EDCSV_Last, TEXT("DatasmithCloth"));
 
 
-FArchive& operator<<(FArchive& Ar, FParameterData& ParameterData)
+FArchive& operator<<(FArchive& Ar, FParameterData& ParameterData)  // UE_DEPRECATED(5.5, "The experimental Cloth importer is no longer supported.")
 {
 	Ar.UsingCustomVersion(FDatasmithClothSerializationVersion::GUID);
 
@@ -46,7 +48,7 @@ FArchive& operator<<(FArchive& Ar, FParameterData& ParameterData)
 }
 
 
-FArchive& operator<<(FArchive& Ar, FDatasmithClothPattern& Pattern)
+FArchive& operator<<(FArchive& Ar, FDatasmithClothPattern& Pattern)  // UE_DEPRECATED(5.5, "The experimental Cloth importer is no longer supported.")
 {
 	Ar.UsingCustomVersion(FDatasmithClothSerializationVersion::GUID);
 	const int32 ClothSerialVersion = Ar.CustomVer(FDatasmithClothSerializationVersion::GUID);
@@ -64,7 +66,7 @@ FArchive& operator<<(FArchive& Ar, FDatasmithClothPattern& Pattern)
 }
 
 
-FArchive& operator<<(FArchive& Ar, FDatasmithClothSewingInfo& Sewing)
+FArchive& operator<<(FArchive& Ar, FDatasmithClothSewingInfo& Sewing)  // UE_DEPRECATED(5.5, "The experimental Cloth importer is no longer supported.")
 {
 	Ar.UsingCustomVersion(FDatasmithClothSerializationVersion::GUID);
 
@@ -77,7 +79,7 @@ FArchive& operator<<(FArchive& Ar, FDatasmithClothSewingInfo& Sewing)
 }
 
 
-FArchive& operator<<(FArchive& Ar, FDatasmithClothPresetProperty& Property)
+FArchive& operator<<(FArchive& Ar, FDatasmithClothPresetProperty& Property)  // UE_DEPRECATED(5.5, "The experimental Cloth importer is no longer supported.")
 {
 	Ar.UsingCustomVersion(FDatasmithClothSerializationVersion::GUID);
 
@@ -88,7 +90,7 @@ FArchive& operator<<(FArchive& Ar, FDatasmithClothPresetProperty& Property)
 }
 
 
-FArchive& operator<<(FArchive& Ar, FDatasmithClothPresetPropertySet& PropertySet)
+FArchive& operator<<(FArchive& Ar, FDatasmithClothPresetPropertySet& PropertySet)  // UE_DEPRECATED(5.5, "The experimental Cloth importer is no longer supported.")
 {
 	Ar.UsingCustomVersion(FDatasmithClothSerializationVersion::GUID);
 
@@ -99,7 +101,7 @@ FArchive& operator<<(FArchive& Ar, FDatasmithClothPresetPropertySet& PropertySet
 }
 
 
-FArchive& operator<<(FArchive& Ar, FDatasmithCloth& Cloth)
+FArchive& operator<<(FArchive& Ar, FDatasmithCloth& Cloth)  // UE_DEPRECATED(5.5, "The experimental Cloth importer is no longer supported.")
 {
 	Ar.UsingCustomVersion(FDatasmithClothSerializationVersion::GUID);
 	const int32 ClothSerialVersion = Ar.CustomVer(FDatasmithClothSerializationVersion::GUID);
@@ -114,3 +116,4 @@ FArchive& operator<<(FArchive& Ar, FDatasmithCloth& Cloth)
 
 	return Ar;
 }
+PRAGMA_ENABLE_DEPRECATION_WARNINGS
diff --git a/Engine/Source/Runtime/Datasmith/DatasmithCore/Private/DatasmithMeshSerialization.cpp b/Engine/Source/Runtime/Datasmith/DatasmithCore/Private/DatasmithMeshSerialization.cpp
--- a/Engine/Source/Runtime/Datasmith/DatasmithCore/Private/DatasmithMeshSerialization.cpp
+++ b/Engine/Source/Runtime/Datasmith/DatasmithCore/Private/DatasmithMeshSerialization.cpp
@@ -22,7 +22,8 @@ FArchive& operator<<(FArchive& Ar, FDatasmithMeshModels& Models)
 }
 
 
-FArchive& operator<<(FArchive& Ar, FDatasmithClothInfo& Info)
+PRAGMA_DISABLE_DEPRECATION_WARNINGS
+FArchive& operator<<(FArchive& Ar, FDatasmithClothInfo& Info)  // UE_DEPRECATED(5.5, "The experimental Cloth importer is no longer supported.")
 {
 	uint8 Version = 0;
 	Ar << Version;
@@ -31,6 +32,7 @@ FArchive& operator<<(FArchive& Ar, FDatasmithClothInfo& Info)
 
 	return Ar;
 }
+PRAGMA_ENABLE_DEPRECATION_WARNINGS
 
 
 namespace DatasmithMeshSerializationImpl
@@ -367,8 +369,7 @@ FDatasmithPackedMeshes GetDatasmithMeshFromFile(const FString& MeshPath)
 	return Result;
 }
 
-
-FMD5Hash FDatasmithPackedCloths::Serialize(FArchive& Ar, bool bSaveCompressed)
+FMD5Hash FDatasmithPackedCloths::Serialize(FArchive& Ar, bool bSaveCompressed)  // UE_DEPRECATED(5.5, "The experimental Cloth importer is no longer supported.")
 {
 	using namespace DatasmithMeshSerializationImpl;
 	// structure          from
@@ -455,7 +456,8 @@ FMD5Hash FDatasmithPackedCloths::Serialize(FArchive& Ar, bool bSaveCompressed)
 }
 
 
-FDatasmithPackedCloths GetDatasmithClothFromFile(const FString& Path)
+PRAGMA_DISABLE_DEPRECATION_WARNINGS
+FDatasmithPackedCloths GetDatasmithClothFromFile(const FString& Path)  // UE_DEPRECATED(5.5, "The experimental Cloth importer is no longer supported.")
 {
 	TUniquePtr<FArchive> Archive(IFileManager::Get().CreateFileReader(*Path));
 	if ( !Archive.IsValid() )
@@ -475,4 +477,5 @@ FDatasmithPackedCloths GetDatasmithClothFromFile(const FString& Path)
 
 	return Result;
 }
+PRAGMA_ENABLE_DEPRECATION_WARNINGS
 
diff --git a/Engine/Source/Runtime/Datasmith/DatasmithCore/Private/DatasmithSceneElementsImpl.cpp b/Engine/Source/Runtime/Datasmith/DatasmithCore/Private/DatasmithSceneElementsImpl.cpp
--- a/Engine/Source/Runtime/Datasmith/DatasmithCore/Private/DatasmithSceneElementsImpl.cpp
+++ b/Engine/Source/Runtime/Datasmith/DatasmithCore/Private/DatasmithSceneElementsImpl.cpp
@@ -107,7 +107,9 @@ TSharedPtr<IDatasmithMaterialIDElement> FDatasmithMeshElementImpl::GetMaterialSl
 }
 
 FDatasmithClothElementImpl::FDatasmithClothElementImpl(const TCHAR* InName)
-	: FDatasmithElementImpl(InName, EDatasmithElementType::Cloth)
+PRAGMA_DISABLE_DEPRECATION_WARNINGS
+	: FDatasmithElementImpl(InName, EDatasmithElementType::Cloth)  // UE_DEPRECATED(5.5, "The experimental Cloth importer is no longer supported.")
+PRAGMA_ENABLE_DEPRECATION_WARNINGS
 {
 	Store.RegisterParameter(File, "File");
 }
@@ -1138,7 +1140,8 @@ void FDatasmithSceneImpl::RemoveMeshAt(int32 InIndex)
 	}
 }
 
-void FDatasmithSceneImpl::AddCloth(const TSharedPtr< IDatasmithClothElement >& InElement)
+PRAGMA_DISABLE_DEPRECATION_WARNINGS
+void FDatasmithSceneImpl::AddCloth(const TSharedPtr< IDatasmithClothElement >& InElement)  // UE_DEPRECATED(5.5, "The experimental Cloth importer is no longer supported.")
 {
 	Clothes.Add(InElement);
 }
@@ -1148,7 +1151,7 @@ int32 FDatasmithSceneImpl::GetClothesCount() const
 	return Clothes.Num();
 }
 
-TSharedPtr< IDatasmithClothElement > FDatasmithSceneImpl::GetCloth(int32 InIndex)
+TSharedPtr< IDatasmithClothElement > FDatasmithSceneImpl::GetCloth(int32 InIndex)  // UE_DEPRECATED(5.5, "The experimental Cloth importer is no longer supported.")
 {
 	if ( Clothes.IsValidIndex( InIndex ) )
 	{
@@ -1157,7 +1160,7 @@ TSharedPtr< IDatasmithClothElement > FDatasmithSceneImpl::GetCloth(int32 InIndex
 	return {};
 }
 
-const TSharedPtr< IDatasmithClothElement >& FDatasmithSceneImpl::GetCloth(int32 InIndex) const
+const TSharedPtr< IDatasmithClothElement >& FDatasmithSceneImpl::GetCloth(int32 InIndex) const  // UE_DEPRECATED(5.5, "The experimental Cloth importer is no longer supported.")
 {
 	if ( Clothes.IsValidIndex( InIndex ) )
 	{
@@ -1167,12 +1170,12 @@ const TSharedPtr< IDatasmithClothElement >& FDatasmithSceneImpl::GetCloth(int32
 	return NullCloth;
 }
 
-void FDatasmithSceneImpl::RemoveCloth(const TSharedPtr< IDatasmithClothElement >& InElement)
+void FDatasmithSceneImpl::RemoveCloth(const TSharedPtr< IDatasmithClothElement >& InElement)  // UE_DEPRECATED(5.5, "The experimental Cloth importer is no longer supported.")
 {
 	Clothes.Remove(InElement);
 }
 
-void FDatasmithSceneImpl::RemoveClothAt(int32 InIndex)
+void FDatasmithSceneImpl::RemoveClothAt(int32 InIndex)  // UE_DEPRECATED(5.5, "The experimental Cloth importer is no longer supported.")
 {
 	if (Clothes.IsValidIndex(InIndex))
 	{
@@ -1180,10 +1183,11 @@ void FDatasmithSceneImpl::RemoveClothAt(int32 InIndex)
 	}
 }
 
-void FDatasmithSceneImpl::EmptyClothes()
+void FDatasmithSceneImpl::EmptyClothes()  // UE_DEPRECATED(5.5, "The experimental Cloth importer is no longer supported.")
 {
 	Clothes.Empty();
 }
+PRAGMA_ENABLE_DEPRECATION_WARNINGS
 
 static const TSharedPtr< IDatasmithMetaDataElement > InvalidMetaData;
 
@@ -1512,8 +1516,10 @@ void FDatasmithSceneImpl::AttachActorToSceneRoot(const TSharedPtr< IDatasmithAct
 	}
 }
 
-FDatasmithClothActorElementImpl::FDatasmithClothActorElementImpl(const TCHAR* InName)
+FDatasmithClothActorElementImpl::FDatasmithClothActorElementImpl(const TCHAR* InName)  // UE_DEPRECATED(5.5, "The experimental Cloth importer is no longer supported.")
+PRAGMA_DISABLE_DEPRECATION_WARNINGS
 	: FDatasmithActorElementImpl<IDatasmithClothActorElement>(InName, EDatasmithElementType::ClothActor)
+PRAGMA_ENABLE_DEPRECATION_WARNINGS
 {
 	Store.RegisterParameter(Cloth, "Cloth");
 }
diff --git a/Engine/Source/Runtime/Datasmith/DatasmithCore/Private/DatasmithSceneElementsImpl.h b/Engine/Source/Runtime/Datasmith/DatasmithCore/Private/DatasmithSceneElementsImpl.h
--- a/Engine/Source/Runtime/Datasmith/DatasmithCore/Private/DatasmithSceneElementsImpl.h
+++ b/Engine/Source/Runtime/Datasmith/DatasmithCore/Private/DatasmithSceneElementsImpl.h
@@ -390,7 +390,10 @@ private:
 /*
  * Experimental Element that describes a cloth asset.
  */
-class FDatasmithClothElementImpl : public FDatasmithElementImpl< IDatasmithClothElement >
+class UE_DEPRECATED(5.5, "The experimental Cloth importer is no longer supported.") FDatasmithClothElementImpl
+PRAGMA_DISABLE_DEPRECATION_WARNINGS
+ : public FDatasmithElementImpl< IDatasmithClothElement >
+PRAGMA_ENABLE_DEPRECATION_WARNINGS
 {
 public:
 	explicit FDatasmithClothElementImpl(const TCHAR* InName);
@@ -1023,7 +1026,10 @@ private:
 };
 
 
-class FDatasmithClothActorElementImpl : public FDatasmithActorElementImpl<IDatasmithClothActorElement>
+PRAGMA_DISABLE_DEPRECATION_WARNINGS
+class UE_DEPRECATED(5.5, "The experimental Cloth importer is no longer supported.") FDatasmithClothActorElementImpl
+	: public FDatasmithActorElementImpl<IDatasmithClothActorElement>
+PRAGMA_ENABLE_DEPRECATION_WARNINGS
 {
 public:
 	FDatasmithClothActorElementImpl(const TCHAR* InName);
@@ -1625,13 +1631,22 @@ public:
 	virtual void RemoveMeshAt(int32 InIndex) override;
 	virtual void EmptyMeshes() override { Meshes.Empty(); }
 
+PRAGMA_DISABLE_DEPRECATION_WARNINGS
+	UE_DEPRECATED(5.5, "The experimental Cloth importer is no longer supported.")
 	virtual void AddCloth(const TSharedPtr< IDatasmithClothElement >& InElement) override;
+	UE_DEPRECATED(5.5, "The experimental Cloth importer is no longer supported.")
 	virtual int32 GetClothesCount() const override;
+	UE_DEPRECATED(5.5, "The experimental Cloth importer is no longer supported.")
 	virtual TSharedPtr< IDatasmithClothElement > GetCloth(int32 InIndex) override;
+	UE_DEPRECATED(5.5, "The experimental Cloth importer is no longer supported.")
 	virtual const TSharedPtr< IDatasmithClothElement >& GetCloth(int32 InIndex) const override;
+	UE_DEPRECATED(5.5, "The experimental Cloth importer is no longer supported.")
 	virtual void RemoveCloth(const TSharedPtr< IDatasmithClothElement >& InElement) override;
+	UE_DEPRECATED(5.5, "The experimental Cloth importer is no longer supported.")
 	virtual void RemoveClothAt(int32 InIndex) override;
+	UE_DEPRECATED(5.5, "The experimental Cloth importer is no longer supported.")
 	virtual void EmptyClothes() override;
+PRAGMA_ENABLE_DEPRECATION_WARNINGS
 
 	virtual void AddActor(const TSharedPtr< IDatasmithActorElement >& InActor) override { Actors.Add(InActor);  }
 	virtual int32 GetActorsCount() const override { return Actors.Num(); }
@@ -1693,7 +1708,10 @@ private:
 
 	TDatasmithReferenceArrayProxy<IDatasmithActorElement>            Actors;
 	TDatasmithReferenceArrayProxy<IDatasmithMeshElement>             Meshes;
+PRAGMA_DISABLE_DEPRECATION_WARNINGS
+	UE_DEPRECATED(5.5, "The experimental Cloth importer is no longer supported.")
 	TDatasmithReferenceArrayProxy<IDatasmithClothElement>            Clothes;
+PRAGMA_ENABLE_DEPRECATION_WARNINGS
 	TDatasmithReferenceArrayProxy<IDatasmithBaseMaterialElement>     Materials;
 	TDatasmithReferenceArrayProxy<IDatasmithTextureElement>          Textures;
 	TDatasmithReferenceArrayProxy<IDatasmithMetaDataElement>         MetaData;
diff --git a/Engine/Source/Runtime/Datasmith/DatasmithCore/Private/DatasmithSceneFactory.cpp b/Engine/Source/Runtime/Datasmith/DatasmithCore/Private/DatasmithSceneFactory.cpp
--- a/Engine/Source/Runtime/Datasmith/DatasmithCore/Private/DatasmithSceneFactory.cpp
+++ b/Engine/Source/Runtime/Datasmith/DatasmithCore/Private/DatasmithSceneFactory.cpp
@@ -45,10 +45,12 @@ TSharedPtr< IDatasmithElement > FDatasmithSceneFactory::CreateElement( EDatasmit
 		}
 	case EDatasmithElementType::StaticMesh:
 		return CreateMesh( InName );
-	case EDatasmithElementType::Cloth:
+PRAGMA_DISABLE_DEPRECATION_WARNINGS
+	case EDatasmithElementType::Cloth:  // UE_DEPRECATED(5.5, "The experimental Cloth importer is no longer supported.")
 		return CreateCloth( InName );
-	case EDatasmithElementType::ClothActor:
+	case EDatasmithElementType::ClothActor:  // UE_DEPRECATED(5.5, "The experimental Cloth importer is no longer supported.")
 		return CreateClothActor( InName );
+PRAGMA_ENABLE_DEPRECATION_WARNINGS
 	case EDatasmithElementType::StaticMeshActor:
 		return CreateMeshActor( InName );
 	case EDatasmithElementType::PointLight:
@@ -204,15 +206,17 @@ TSharedRef< IDatasmithMeshActorElement > FDatasmithSceneFactory::CreateMeshActor
 	return MakeShared< FDatasmithMeshActorElementImpl<> >( InName );
 }
 
-TSharedRef< IDatasmithClothElement > FDatasmithSceneFactory::CreateCloth(const TCHAR* InName)
+PRAGMA_DISABLE_DEPRECATION_WARNINGS
+TSharedRef< IDatasmithClothElement > FDatasmithSceneFactory::CreateCloth(const TCHAR* InName)  // UE_DEPRECATED(5.5, "The experimental Cloth importer is no longer supported.")
 {
 	return MakeShared< FDatasmithClothElementImpl >( InName );
 }
 
-TSharedRef< IDatasmithClothActorElement > FDatasmithSceneFactory::CreateClothActor( const TCHAR* InName )
+TSharedRef< IDatasmithClothActorElement > FDatasmithSceneFactory::CreateClothActor( const TCHAR* InName )  // UE_DEPRECATED(5.5, "The experimental Cloth importer is no longer supported.")
 {
 	return MakeShared< FDatasmithClothActorElementImpl >( InName );
 }
+PRAGMA_ENABLE_DEPRECATION_WARNINGS
 
 TSharedRef< IDatasmithHierarchicalInstancedStaticMeshActorElement > FDatasmithSceneFactory::CreateHierarchicalInstanceStaticMeshActor(const TCHAR* InName)
 {
diff --git a/Engine/Source/Runtime/Datasmith/DatasmithCore/Private/DatasmithSceneXmlReader.cpp b/Engine/Source/Runtime/Datasmith/DatasmithCore/Private/DatasmithSceneXmlReader.cpp
--- a/Engine/Source/Runtime/Datasmith/DatasmithCore/Private/DatasmithSceneXmlReader.cpp
+++ b/Engine/Source/Runtime/Datasmith/DatasmithCore/Private/DatasmithSceneXmlReader.cpp
@@ -26,7 +26,7 @@ const TCHAR* ActorTagsView[] = {
 	DATASMITH_ACTORHIERARCHICALINSTANCEDMESHNAME,
 	DATASMITH_ACTORMESHNAME,
 	DATASMITH_CAMERANAME,
-	DATASMITH_CLOTHACTORNAME,
+	DATASMITH_CLOTHACTORNAME,  // UE_DEPRECATED(5.5, "The experimental Cloth importer is no longer supported.")
 	DATASMITH_CUSTOMACTORNAME,
 	DATASMITH_DECALACTORNAME,
 	DATASMITH_LANDSCAPENAME,
@@ -436,7 +436,8 @@ void FDatasmithSceneXmlReader::ParseMesh(FXmlNode* InNode, TSharedPtr<IDatasmith
 	}
 }
 
-void FDatasmithSceneXmlReader::ParseCloth(FXmlNode* InNode, TSharedPtr<IDatasmithClothElement>& OutElement) const
+PRAGMA_DISABLE_DEPRECATION_WARNINGS
+void FDatasmithSceneXmlReader::ParseCloth(FXmlNode* InNode, TSharedPtr<IDatasmithClothElement>& OutElement) const  // UE_DEPRECATED(5.5, "The experimental Cloth importer is no longer supported.")
 {
 	ParseElement( InNode, OutElement.ToSharedRef() );
 
@@ -448,6 +449,7 @@ void FDatasmithSceneXmlReader::ParseCloth(FXmlNode* InNode, TSharedPtr<IDatasmit
 		}
 	}
 }
+PRAGMA_ENABLE_DEPRECATION_WARNINGS
 
 void FDatasmithSceneXmlReader::ParseTextureElement(FXmlNode* InNode, TSharedPtr<IDatasmithTextureElement>& OutElement) const
 {
@@ -663,12 +665,14 @@ void FDatasmithSceneXmlReader::ParseActor(FXmlNode* InNode, TSharedPtr<IDatasmit
 		ParseMeshActor(InNode, MeshElement, Scene);
 		InOutElement = MeshElement;
 	}
-	else if (InNode->GetTag() == DATASMITH_CLOTHACTORNAME)
+PRAGMA_DISABLE_DEPRECATION_WARNINGS
+	else if (InNode->GetTag() == DATASMITH_CLOTHACTORNAME)  // UE_DEPRECATED(5.5, "The experimental Cloth importer is no longer supported.")
 	{
 		TSharedPtr< IDatasmithClothActorElement> Element = FDatasmithSceneFactory::CreateClothActor(*InNode->GetAttribute(TEXT("name")));
 		ParseClothActor(InNode, Element, Scene);
 		InOutElement = Element;
 	}
+PRAGMA_ENABLE_DEPRECATION_WARNINGS
 	else if (InNode->GetTag() == DATASMITH_LIGHTNAME)
 	{
 		TSharedPtr< IDatasmithLightActorElement > LightElement;
@@ -826,7 +830,8 @@ void FDatasmithSceneXmlReader::ParseMeshActor(FXmlNode* InNode, TSharedPtr<IData
 	}
 }
 
-void FDatasmithSceneXmlReader::ParseClothActor(FXmlNode* InNode, TSharedPtr<IDatasmithClothActorElement>& OutElement, TSharedRef< IDatasmithScene > Scene) const
+PRAGMA_DISABLE_DEPRECATION_WARNINGS
+void FDatasmithSceneXmlReader::ParseClothActor(FXmlNode* InNode, TSharedPtr<IDatasmithClothActorElement>& OutElement, TSharedRef< IDatasmithScene > Scene) const  // UE_DEPRECATED(5.5, "The experimental Cloth importer is no longer supported.")
 {
 	ParseElement(InNode, OutElement.ToSharedRef());
 
@@ -838,6 +843,7 @@ void FDatasmithSceneXmlReader::ParseClothActor(FXmlNode* InNode, TSharedPtr<IDat
 		}
 	}
 }
+PRAGMA_ENABLE_DEPRECATION_WARNINGS
 
 void FDatasmithSceneXmlReader::ParseHierarchicalInstancedStaticMeshActor(FXmlNode* InNode, TSharedPtr<IDatasmithHierarchicalInstancedStaticMeshActorElement>& OutElement, TSharedRef< IDatasmithScene > Scene) const
 {
@@ -1232,8 +1238,9 @@ bool FDatasmithSceneXmlReader::ParseXmlFile(TSharedRef< IDatasmithScene >& OutSc
 
 			Objects.Add( Element->GetName(), Element );
 		}
+PRAGMA_DISABLE_DEPRECATION_WARNINGS
 		// CLOTHES
-		else if (Nodes[i]->GetTag() == DATASMITH_CLOTH)
+		else if (Nodes[i]->GetTag() == DATASMITH_CLOTH)  // UE_DEPRECATED(5.5, "The experimental Cloth importer is no longer supported.")
 		{
 			FString ElementName = Nodes[i]->GetAttribute(TEXT("name"));
 			TSharedPtr< IDatasmithClothElement > Element = FDatasmithSceneFactory::CreateCloth(*ElementName);
@@ -1244,6 +1251,7 @@ bool FDatasmithSceneXmlReader::ParseXmlFile(TSharedRef< IDatasmithScene >& OutSc
 
 // 			Objects.Add( Element->GetName(), Element ); // #ue_ds_cloth_todo cloths referencable by other elements
 		}
+PRAGMA_ENABLE_DEPRECATION_WARNINGS
 		// LEVEL SEQUENCES
 		else if (Nodes[i]->GetTag() == DATASMITH_LEVELSEQUENCENAME)
 		{
diff --git a/Engine/Source/Runtime/Datasmith/DatasmithCore/Private/DatasmithSceneXmlWriter.cpp b/Engine/Source/Runtime/Datasmith/DatasmithCore/Private/DatasmithSceneXmlWriter.cpp
--- a/Engine/Source/Runtime/Datasmith/DatasmithCore/Private/DatasmithSceneXmlWriter.cpp
+++ b/Engine/Source/Runtime/Datasmith/DatasmithCore/Private/DatasmithSceneXmlWriter.cpp
@@ -32,7 +32,10 @@ public:
 
 	static void WriteTexture(FArchive& Archive, int32 Indent, const TCHAR* Prefix, const TCHAR* Name, FDatasmithTextureSampler UV);
 	static void WriteMeshElement(const TSharedPtr< IDatasmithMeshElement >& MeshElement, FArchive& Archive, int32 Indent);
+PRAGMA_DISABLE_DEPRECATION_WARNINGS
+	UE_DEPRECATED(5.5, "The experimental Cloth importer is no longer supported.")
 	static void WriteClothElement(const TSharedPtr< IDatasmithClothElement >& ClothElement, FArchive& Archive, int32 Indent);
+PRAGMA_ENABLE_DEPRECATION_WARNINGS
 	static void WriteLevelSequenceElement( const TSharedRef< IDatasmithLevelSequenceElement>& SequenceElement, FArchive& Archive, int32 Indent );
 
 	static void WriteLevelVariantSetsElement( const TSharedRef< IDatasmithLevelVariantSetsElement >& LevelVariantSetsElement, FArchive& Archive, int32 Indent );
@@ -47,7 +50,10 @@ public:
 	static void WriteActorTags(const TSharedPtr< IDatasmithActorElement >& ActorElement, FArchive& Archive, int32 Indent);
 	static void WriteActorChildren(const TSharedPtr< IDatasmithActorElement >& ActorElement, FArchive& Archive, int32 Indent);
 	static void WriteMeshActorElement(const TSharedPtr< IDatasmithMeshActorElement >& MeshActorElement, FArchive& Archive, int32 Indent);
+PRAGMA_DISABLE_DEPRECATION_WARNINGS
+	UE_DEPRECATED(5.5, "The experimental Cloth importer is no longer supported.")
 	static void WriteClothActorElement(const TSharedPtr< IDatasmithClothActorElement >& ClothActorElement, FArchive& Archive, int32 Indent);
+PRAGMA_ENABLE_DEPRECATION_WARNINGS
 
 	// Write the start of the actor element (Open the xml element and add the essential child elements for the mesh actor)
 	static void WriteBeginOfMeshActorElement(const TSharedPtr<IDatasmithMeshActorElement>& MeshActorElement, const FString& ElementTypeString, FArchive& Archive, int32 Indent);
@@ -360,7 +366,8 @@ void FDatasmithSceneXmlWriterImpl::WriteMeshElement(const TSharedPtr< IDatasmith
 	SerializeToArchive(Archive, XmlString);
 }
 
-void FDatasmithSceneXmlWriterImpl::WriteClothElement(const TSharedPtr<IDatasmithClothElement>& Element, FArchive& Archive, int32 Indent)
+PRAGMA_DISABLE_DEPRECATION_WARNINGS
+void FDatasmithSceneXmlWriterImpl::WriteClothElement(const TSharedPtr<IDatasmithClothElement>& Element, FArchive& Archive, int32 Indent)  // UE_DEPRECATED(5.5, "The experimental Cloth importer is no longer supported.")
 {
 	WriteIndent(Archive, Indent);
 	FString XmlString = TEXT("<") DATASMITH_CLOTH;
@@ -375,6 +382,7 @@ void FDatasmithSceneXmlWriterImpl::WriteClothElement(const TSharedPtr<IDatasmith
 	WriteIndent(Archive, Indent);
 	SerializeToArchive(Archive, TEXT("</") DATASMITH_CLOTH TEXT(">") LINE_TERMINATOR);
 }
+PRAGMA_ENABLE_DEPRECATION_WARNINGS
 
 void FDatasmithSceneXmlWriterImpl::WriteLevelSequenceElement(const TSharedRef< IDatasmithLevelSequenceElement>& SequenceElement, FArchive& Archive, int32 Indent)
 {
@@ -570,10 +578,12 @@ void FDatasmithSceneXmlWriterImpl::WriteActorElement(const TSharedPtr< IDatasmit
 	{
 		WriteMeshActorElement(StaticCastSharedPtr< IDatasmithMeshActorElement >(ActorElement), Archive, Indent);
 	}
-	else if ( ActorElement->IsA( EDatasmithElementType::ClothActor ) )
+PRAGMA_DISABLE_DEPRECATION_WARNINGS
+	else if ( ActorElement->IsA( EDatasmithElementType::ClothActor ) )  // UE_DEPRECATED(5.5, "The experimental Cloth importer is no longer supported.")
 	{
 		WriteClothActorElement(StaticCastSharedPtr< IDatasmithClothActorElement >(ActorElement), Archive, Indent);
 	}
+PRAGMA_ENABLE_DEPRECATION_WARNINGS
 	else if ( ActorElement->IsA( EDatasmithElementType::Camera ) )
 	{
 		WriteCameraActorElement(StaticCastSharedPtr< IDatasmithCameraActorElement >(ActorElement), Archive, Indent);
@@ -681,7 +691,8 @@ void FDatasmithSceneXmlWriterImpl::WriteMeshActorElement(const TSharedPtr< IData
 	WriteEndOfMeshActorElement(ElementTypeString, Archive, Indent);
 }
 
-void FDatasmithSceneXmlWriterImpl::WriteClothActorElement(const TSharedPtr< IDatasmithClothActorElement >& Element, FArchive& Archive, int32 Indent)
+PRAGMA_DISABLE_DEPRECATION_WARNINGS
+void FDatasmithSceneXmlWriterImpl::WriteClothActorElement(const TSharedPtr< IDatasmithClothActorElement >& Element, FArchive& Archive, int32 Indent)  // UE_DEPRECATED(5.5, "The experimental Cloth importer is no longer supported.")
 {
 	WriteIndent(Archive, Indent);
 	SerializeToArchive(Archive, FString(TEXT("<") DATASMITH_CLOTHACTORNAME TEXT(" name=\"")) + SanitizeXMLText(Element->GetName()) + TEXT("\""));
@@ -693,6 +704,7 @@ void FDatasmithSceneXmlWriterImpl::WriteClothActorElement(const TSharedPtr< IDat
 	WriteIndent(Archive, Indent);
 	SerializeToArchive(Archive, TEXT("</") DATASMITH_CLOTHACTORNAME TEXT(">") LINE_TERMINATOR);
 }
+PRAGMA_ENABLE_DEPRECATION_WARNINGS
 
 void FDatasmithSceneXmlWriterImpl::WriteBeginOfMeshActorElement(const TSharedPtr<IDatasmithMeshActorElement>& MeshActorElement, const FString& ElementTypeString, FArchive& Archive, int32 Indent)
 {
@@ -1987,11 +1999,13 @@ void FDatasmithSceneXmlWriter::Serialize( TSharedRef< IDatasmithScene > Datasmit
 		FDatasmithSceneXmlWriterImpl::WriteMeshElement( Mesh, Archive, 1 );
 	}
 
-	for ( int32 Index = 0; Index < DatasmithScene->GetClothesCount(); ++Index )
+PRAGMA_DISABLE_DEPRECATION_WARNINGS
+	for ( int32 Index = 0; Index < DatasmithScene->GetClothesCount(); ++Index )  // UE_DEPRECATED(5.5, "The experimental Cloth importer is no longer supported.")
 	{
 		const TSharedPtr< IDatasmithClothElement >& Cloth = DatasmithScene->GetCloth( Index );
 		FDatasmithSceneXmlWriterImpl::WriteClothElement( Cloth, Archive, 1 );
 	}
+PRAGMA_ENABLE_DEPRECATION_WARNINGS
 
 	for ( int32 ActorIndex = 0; ActorIndex < DatasmithScene->GetActorsCount(); ++ActorIndex )
 	{
diff --git a/Engine/Source/Runtime/Datasmith/DatasmithCore/Public/DatasmithCloth.h b/Engine/Source/Runtime/Datasmith/DatasmithCore/Public/DatasmithCloth.h
--- a/Engine/Source/Runtime/Datasmith/DatasmithCore/Public/DatasmithCloth.h
+++ b/Engine/Source/Runtime/Datasmith/DatasmithCore/Public/DatasmithCloth.h
@@ -15,7 +15,7 @@ class FDatasmithMesh;
 class UActorComponent;
 
 
-struct FParameterData
+struct UE_DEPRECATED(5.5, "The experimental Cloth importer is no longer supported.") FParameterData
 {
 	FString Name;
 
@@ -29,14 +29,16 @@ public:
 };
 
 
-class DATASMITHCORE_API FDatasmithClothPattern
+class UE_DEPRECATED(5.5, "The experimental Cloth importer is no longer supported.") DATASMITHCORE_API FDatasmithClothPattern
 {
 public:
 	TArray<FVector2f> SimPosition;
 	TArray<FVector3f> SimRestPosition;
 	TArray<uint32> SimTriangleIndices;
 
+PRAGMA_DISABLE_DEPRECATION_WARNINGS
 	TArray<FParameterData> Parameters;
+PRAGMA_ENABLE_DEPRECATION_WARNINGS
 
 public:
 	bool IsValid() const { return SimRestPosition.Num() == SimPosition.Num() && SimTriangleIndices.Num() % 3 == 0 && SimTriangleIndices.Num(); }
@@ -44,7 +46,7 @@ public:
 };
 
 
-struct FDatasmithClothSewingInfo
+struct UE_DEPRECATED(5.5, "The experimental Cloth importer is no longer supported.") FDatasmithClothSewingInfo
 {
 	uint32 Seam0PanelIndex = 0;
 	uint32 Seam1PanelIndex = 0;
@@ -55,7 +57,7 @@ struct FDatasmithClothSewingInfo
 };
 
 
-class FDatasmithClothPresetProperty
+class UE_DEPRECATED(5.5, "The experimental Cloth importer is no longer supported.") FDatasmithClothPresetProperty
 {
 public:
 	FName Name;
@@ -66,23 +68,27 @@ public:
 };
 
 
-class DATASMITHCORE_API FDatasmithClothPresetPropertySet
+class UE_DEPRECATED(5.5, "The experimental Cloth importer is no longer supported.") DATASMITHCORE_API FDatasmithClothPresetPropertySet
 {
 public:
 	FString SetName;
+PRAGMA_DISABLE_DEPRECATION_WARNINGS
 	TArray<FDatasmithClothPresetProperty> Properties;
+PRAGMA_ENABLE_DEPRECATION_WARNINGS
 
 public:
 	friend FArchive& operator<<(FArchive& Ar, FDatasmithClothPresetPropertySet& PropertySet);
 };
 
 
-class DATASMITHCORE_API FDatasmithCloth
+class UE_DEPRECATED(5.5, "The experimental Cloth importer is no longer supported.") DATASMITHCORE_API FDatasmithCloth
 {
 public:
+PRAGMA_DISABLE_DEPRECATION_WARNINGS
 	TArray<FDatasmithClothPattern> Patterns;
 	TArray<FDatasmithClothSewingInfo> Sewing;
 	TArray<FDatasmithClothPresetPropertySet> PropertySets;
+PRAGMA_ENABLE_DEPRECATION_WARNINGS
 
 public:
 	friend FArchive& operator<<(FArchive& Ar, FDatasmithCloth& Cloth);
@@ -96,17 +102,24 @@ class DATASMITHCORE_API UDatasmithClothAssetFactory : public UObject
 	GENERATED_BODY()
 
 public:
+	UE_DEPRECATED(5.5, "The experimental Cloth importer is no longer supported.")
 	UDatasmithClothAssetFactory() = default;
+	UE_DEPRECATED(5.5, "The experimental Cloth importer is no longer supported.")
 	virtual ~UDatasmithClothAssetFactory() = default;
 
+	UE_DEPRECATED(5.5, "The experimental Cloth importer is no longer supported.")
 	virtual UObject* CreateClothAsset(UObject* Outer, const FName& Name, EObjectFlags Flags) const
 	PURE_VIRTUAL(UDatasmithClothAssetFactory::CreateClothAsset, return nullptr;);
 
+	UE_DEPRECATED(5.5, "The experimental Cloth importer is no longer supported.")
 	virtual UObject* DuplicateClothAsset(UObject* ClothAsset, UObject* Outer, const FName& Name) const
 	PURE_VIRTUAL(UDatasmithClothAssetFactory::DuplicateClothAsset, return nullptr;);
 
+PRAGMA_DISABLE_DEPRECATION_WARNINGS
+	UE_DEPRECATED(5.5, "The experimental Cloth importer is no longer supported.")
 	virtual void InitializeClothAsset(UObject* ClothAsset, const FDatasmithCloth& DatasmithCloth) const
 	PURE_VIRTUAL(UDatasmithClothAssetFactory::InitializeClothAsset, );
+PRAGMA_ENABLE_DEPRECATION_WARNINGS
 };
 
 
@@ -117,19 +130,23 @@ class DATASMITHCORE_API UDatasmithClothComponentFactory : public UObject
 	GENERATED_BODY()
 
 public:
+	UE_DEPRECATED(5.5, "The experimental Cloth importer is no longer supported.")
 	UDatasmithClothComponentFactory() = default;
+	UE_DEPRECATED(5.5, "The experimental Cloth importer is no longer supported.")
 	virtual ~UDatasmithClothComponentFactory() = default;
 
+	UE_DEPRECATED(5.5, "The experimental Cloth importer is no longer supported.")
 	virtual USceneComponent* CreateClothComponent(UObject* Outer) const
 	PURE_VIRTUAL(UDatasmithClothComponentFactory::CreateClothComponent, return nullptr;);
 
+	UE_DEPRECATED(5.5, "The experimental Cloth importer is no longer supported.")
 	virtual void InitializeClothComponent(USceneComponent* ClothComponent, UObject* ClothAsset, USceneComponent* RootComponent) const
 	PURE_VIRTUAL(UDatasmithClothComponentFactory::InitializeClothComponent, );
 };
 
 
 /** A modular interface to provide factory classes to initialize cloth assets and components. */
-class IDatasmithClothFactoryClassesProvider : public IModularFeature
+class UE_DEPRECATED(5.5, "The experimental Cloth importer is no longer supported.") IDatasmithClothFactoryClassesProvider : public IModularFeature
 {
 public:
 	inline static const FName FeatureName = TEXT("IDatasmithClothFactoryClassesProvider");
diff --git a/Engine/Source/Runtime/Datasmith/DatasmithCore/Public/DatasmithDefinitions.h b/Engine/Source/Runtime/Datasmith/DatasmithCore/Public/DatasmithDefinitions.h
--- a/Engine/Source/Runtime/Datasmith/DatasmithCore/Public/DatasmithDefinitions.h
+++ b/Engine/Source/Runtime/Datasmith/DatasmithCore/Public/DatasmithDefinitions.h
@@ -54,8 +54,8 @@ enum class EDatasmithElementType : uint64
 	MaterialExpression             = 1ull << 31,
 	MaterialExpressionInput        = 1ull << 32,
 	MaterialExpressionOutput       = 1ull << 33,
-	Cloth                          = 1ull << 34,
-	ClothActor                     = 1ull << 35,
+	Cloth      UE_DEPRECATED(5.5, "The experimental Cloth importer is no longer supported.") = 1ull << 34,
+	ClothActor UE_DEPRECATED(5.5, "The experimental Cloth importer is no longer supported.") = 1ull << 35,
 };
 
 ENUM_CLASS_FLAGS( EDatasmithElementType ); // Define bitwise operators for EDatasmithElementType
@@ -469,8 +469,8 @@ static const TCHAR* KeyValuePropertyTypeStrings[] = { TEXT("String"), TEXT("Colo
 
 #define DATASMITH_ACTORNAME						TEXT("Actor")
 
-#define DATASMITH_CLOTH							TEXT("Cloth")
-#define DATASMITH_CLOTHACTORNAME				TEXT("ClothActor")
+#define DATASMITH_CLOTH							TEXT("Cloth")  // UE_DEPRECATED(5.5, "The experimental Cloth importer is no longer supported.")
+#define DATASMITH_CLOTHACTORNAME				TEXT("ClothActor")  // UE_DEPRECATED(5.5, "The experimental Cloth importer is no longer supported.")
 
 //ACTOR MESHES
 #define DATASMITH_ACTORMESHNAME					TEXT("ActorMesh")
diff --git a/Engine/Source/Runtime/Datasmith/DatasmithCore/Public/DatasmithMeshSerialization.h b/Engine/Source/Runtime/Datasmith/DatasmithCore/Public/DatasmithMeshSerialization.h
--- a/Engine/Source/Runtime/Datasmith/DatasmithCore/Public/DatasmithMeshSerialization.h
+++ b/Engine/Source/Runtime/Datasmith/DatasmithCore/Public/DatasmithMeshSerialization.h
@@ -8,7 +8,7 @@
 #include "Containers/UnrealString.h"
 #include "MeshDescription.h"
 #include "Misc/SecureHash.h"
-#include "DatasmithCloth.h"
+#include "DatasmithCloth.h"  // UE_DEPRECATED(5.5, "The experimental Cloth importer is no longer supported.")
 
 class FArchive;
 
@@ -33,19 +33,21 @@ DATASMITHCORE_API FDatasmithPackedMeshes GetDatasmithMeshFromFile(const FString&
 
 
 
-struct DATASMITHCORE_API FDatasmithClothInfo
+PRAGMA_DISABLE_DEPRECATION_WARNINGS
+struct UE_DEPRECATED(5.5, "The experimental Cloth importer is no longer supported.") DATASMITHCORE_API FDatasmithClothInfo
 {
 	FDatasmithCloth Cloth;
 	friend FArchive& operator<<(FArchive& Ar, FDatasmithClothInfo& Info);
 };
 
-struct DATASMITHCORE_API FDatasmithPackedCloths
+struct UE_DEPRECATED(5.5, "The experimental Cloth importer is no longer supported.") DATASMITHCORE_API FDatasmithPackedCloths
 {
 	TArray<FDatasmithClothInfo> ClothInfos;
 
 	FMD5Hash Serialize(FArchive& Ar, bool bSaveCompressed=true);
 };
 
-
+UE_DEPRECATED(5.5, "The experimental Cloth importer is no longer supported.")
 DATASMITHCORE_API FDatasmithPackedCloths GetDatasmithClothFromFile(const FString& Path);
+PRAGMA_ENABLE_DEPRECATION_WARNINGS
 
diff --git a/Engine/Source/Runtime/Datasmith/DatasmithCore/Public/DatasmithSceneFactory.h b/Engine/Source/Runtime/Datasmith/DatasmithCore/Public/DatasmithSceneFactory.h
--- a/Engine/Source/Runtime/Datasmith/DatasmithCore/Public/DatasmithSceneFactory.h
+++ b/Engine/Source/Runtime/Datasmith/DatasmithCore/Public/DatasmithSceneFactory.h
@@ -48,9 +48,13 @@ public:
 
 	static TSharedRef< IDatasmithMeshActorElement > CreateMeshActor(const TCHAR* InName);
 
+PRAGMA_DISABLE_DEPRECATION_WARNINGS
+	UE_DEPRECATED(5.5, "The experimental Cloth importer is no longer supported.")
 	static TSharedRef< IDatasmithClothElement > CreateCloth( const TCHAR* InName );
 
+	UE_DEPRECATED(5.5, "The experimental Cloth importer is no longer supported.")
 	static TSharedRef< IDatasmithClothActorElement > CreateClothActor(const TCHAR* InName);
+PRAGMA_ENABLE_DEPRECATION_WARNINGS
 
 	static TSharedRef< IDatasmithHierarchicalInstancedStaticMeshActorElement > CreateHierarchicalInstanceStaticMeshActor( const TCHAR* InName );
 
diff --git a/Engine/Source/Runtime/Datasmith/DatasmithCore/Public/DatasmithSceneXmlReader.h b/Engine/Source/Runtime/Datasmith/DatasmithCore/Public/DatasmithSceneXmlReader.h
--- a/Engine/Source/Runtime/Datasmith/DatasmithCore/Public/DatasmithSceneXmlReader.h
+++ b/Engine/Source/Runtime/Datasmith/DatasmithCore/Public/DatasmithSceneXmlReader.h
@@ -51,12 +51,18 @@ private:
 	void ParsePropertyCapture( FXmlNode* InNode, const TSharedRef<IDatasmithPropertyCaptureElement>& OutElement ) const;
 	void ParseObjectPropertyCapture( FXmlNode* InNode, const TSharedRef<IDatasmithObjectPropertyCaptureElement>& OutElement, const TMap< FString, TSharedPtr<IDatasmithElement> >& Objects ) const;
 	void ParseMesh(FXmlNode* InNode, TSharedPtr<IDatasmithMeshElement>& OutElement) const;
+PRAGMA_DISABLE_DEPRECATION_WARNINGS
+	UE_DEPRECATED(5.5, "The experimental Cloth importer is no longer supported.")
 	void ParseCloth(FXmlNode* InNode, TSharedPtr<IDatasmithClothElement>& OutElement) const;
+PRAGMA_ENABLE_DEPRECATION_WARNINGS
 	void ParseTextureElement(FXmlNode* InNode, TSharedPtr<IDatasmithTextureElement>& OutElement) const;
 	void ParseTexture(FXmlNode* InNode, FString& OutTextureFilename, FDatasmithTextureSampler& OutTextureUV) const;
 	void ParseActor(FXmlNode* InNode, TSharedPtr<IDatasmithActorElement>& InOutElement, TSharedRef< IDatasmithScene > Scene, TMap< FString, TSharedPtr<IDatasmithActorElement> >& Actors) const;
 	void ParseMeshActor(FXmlNode* InNode, TSharedPtr<IDatasmithMeshActorElement>& OutElement, TSharedRef< IDatasmithScene > Scene) const;
+PRAGMA_DISABLE_DEPRECATION_WARNINGS
+	UE_DEPRECATED(5.5, "The experimental Cloth importer is no longer supported.")
 	void ParseClothActor(FXmlNode* InNode, TSharedPtr<IDatasmithClothActorElement>& OutElement, TSharedRef< IDatasmithScene > Scene) const;
+PRAGMA_ENABLE_DEPRECATION_WARNINGS
 	void ParseHierarchicalInstancedStaticMeshActor(FXmlNode* InNode, TSharedPtr<IDatasmithHierarchicalInstancedStaticMeshActorElement>& OutElement, TSharedRef< IDatasmithScene > Scene) const;
 	void ParseLight(FXmlNode* InNode, TSharedPtr<IDatasmithLightActorElement>& OutElement) const;
 	void ParseCamera(FXmlNode* InNode, TSharedPtr<IDatasmithCameraActorElement>& OutElement) const;
diff --git a/Engine/Source/Runtime/Datasmith/DatasmithCore/Public/IDatasmithSceneElements.h b/Engine/Source/Runtime/Datasmith/DatasmithCore/Public/IDatasmithSceneElements.h
--- a/Engine/Source/Runtime/Datasmith/DatasmithCore/Public/IDatasmithSceneElements.h
+++ b/Engine/Source/Runtime/Datasmith/DatasmithCore/Public/IDatasmithSceneElements.h
@@ -240,7 +240,7 @@ protected:
 /**
 * IDatasmithClothElement class: experimental class that describes a cloth asset
 */
-class DATASMITHCORE_API IDatasmithClothElement : public IDatasmithElement
+class UE_DEPRECATED(5.5, "The experimental Cloth importer is no longer supported.") DATASMITHCORE_API IDatasmithClothElement : public IDatasmithElement
 {
 public:
 	/** Get the FDatasmithCloth resource filename */
@@ -298,7 +298,7 @@ public:
 	virtual void SetStaticMeshPathName(const TCHAR* InStaticMeshPathName) = 0;
 };
 
-class DATASMITHCORE_API IDatasmithClothActorElement : public IDatasmithActorElement
+class UE_DEPRECATED(5.5, "The experimental Cloth importer is no longer supported.") DATASMITHCORE_API IDatasmithClothActorElement : public IDatasmithActorElement
 {
 public:
 	virtual void SetCloth(const TCHAR* Cloth) = 0;
@@ -1619,14 +1619,22 @@ public:
 	*/
 	virtual void EmptyMeshes() = 0;
 
-	// #ue_ds_todo cloth api doc
+PRAGMA_DISABLE_DEPRECATION_WARNINGS
+	UE_DEPRECATED(5.5, "The experimental Cloth importer is no longer supported.")
 	virtual void AddCloth(const TSharedPtr< IDatasmithClothElement >& InElement) = 0;
+	UE_DEPRECATED(5.5, "The experimental Cloth importer is no longer supported.")
 	virtual int32 GetClothesCount() const = 0;
+	UE_DEPRECATED(5.5, "The experimental Cloth importer is no longer supported.")
 	virtual TSharedPtr< IDatasmithClothElement > GetCloth(int32 InIndex) = 0;
+	UE_DEPRECATED(5.5, "The experimental Cloth importer is no longer supported.")
 	virtual const TSharedPtr< IDatasmithClothElement >& GetCloth(int32 InIndex) const = 0;
+	UE_DEPRECATED(5.5, "The experimental Cloth importer is no longer supported.")
 	virtual void RemoveCloth(const TSharedPtr< IDatasmithClothElement >& InElement) = 0;
+	UE_DEPRECATED(5.5, "The experimental Cloth importer is no longer supported.")
 	virtual void RemoveClothAt(int32 InIndex) = 0;
+	UE_DEPRECATED(5.5, "The experimental Cloth importer is no longer supported.")
 	virtual void EmptyClothes() = 0;
+PRAGMA_ENABLE_DEPRECATION_WARNINGS
 
 	/**
 	 * Adds an Actor to the scene.
