CIS - Fixed warning: 'Clothes' is deprecated on Linux builds.

#rnx

diff --git a/Engine/Source/Runtime/Datasmith/DatasmithCore/Private/DatasmithSceneElementsImpl.h b/Engine/Source/Runtime/Datasmith/DatasmithCore/Private/DatasmithSceneElementsImpl.h
--- a/Engine/Source/Runtime/Datasmith/DatasmithCore/Private/DatasmithSceneElementsImpl.h
+++ b/Engine/Source/Runtime/Datasmith/DatasmithCore/Private/DatasmithSceneElementsImpl.h
@@ -1581,7 +1581,9 @@ private:
 	int32 MaterialPropertyIndex;
 };
 
+PRAGMA_DISABLE_DEPRECATION_WARNINGS  // For Clothes - UE_DEPRECATED(5.5, "The experimental Cloth importer is no longer supported.")
 class DATASMITHCORE_API FDatasmithSceneImpl : public FDatasmithElementImpl< IDatasmithScene >
+PRAGMA_ENABLE_DEPRECATION_WARNINGS
 {
 public:
 	explicit FDatasmithSceneImpl(const TCHAR* InName);
