Temporarily disabled all processing related to thin faces or zones.

#rnx

diff --git a/Engine/Config/ConsoleVariables.ini b/Engine/Config/ConsoleVariables.ini
--- a/Engine/Config/ConsoleVariables.ini
+++ b/Engine/Config/ConsoleVariables.ini
@@ -199,4 +199,9 @@ fx.SkipVectorVMBackendOptimizations=1
 ; JT Import Option
 ; Force to import embedded tessellation instead of meshing BRep model.
 ; Default is false
-;ds.CADTranslator.PreferJtFileEmbeddedTessellation=false
\ No newline at end of file
+;ds.CADTranslator.PreferJtFileEmbeddedTessellation=false
+
+; Temporarily disable all thin face/zone processing
+ds.CADTranslator.Meshing.ActivateThinZoneMeshing=false
+ds.CADTranslator.Stitching.RemoveThinFaces=false
+
diff --git a/Engine/Plugins/Enterprise/DatasmithCADImporter/Source/CADLibrary/Private/CADKernelTools.cpp b/Engine/Plugins/Enterprise/DatasmithCADImporter/Source/CADLibrary/Private/CADKernelTools.cpp
--- a/Engine/Plugins/Enterprise/DatasmithCADImporter/Source/CADLibrary/Private/CADKernelTools.cpp
+++ b/Engine/Plugins/Enterprise/DatasmithCADImporter/Source/CADLibrary/Private/CADKernelTools.cpp
@@ -271,7 +271,7 @@ bool FCADKernelTools::Tessellate(UE::CADKernel::FTopologicalShapeEntity& CADTopo
 	TSharedRef<FModelMesh> CADKernelModelMesh = FEntity::MakeShared<FModelMesh>();
 
 	const double GeometricTolerance = FImportParameters::GStitchingTolerance * 10; // cm to mm
-	const bool bActivateThinZoneMeshing = TessellationContext.bSkipThinZoneMeshing ? false : FImportParameters::bGActivateThinZoneMeshing;
+	const bool bActivateThinZoneMeshing = FImportParameters::bGActivateThinZoneMeshing;
 	FMesher Mesher(*CADKernelModelMesh, GeometricTolerance, bActivateThinZoneMeshing);
 
 	DefineMeshCriteria(*CADKernelModelMesh, TessellationContext.ImportParameters, GeometricTolerance);
diff --git a/Engine/Plugins/Enterprise/DatasmithCADImporter/Source/CADLibrary/Public/MeshDescriptionHelper.h b/Engine/Plugins/Enterprise/DatasmithCADImporter/Source/CADLibrary/Public/MeshDescriptionHelper.h
--- a/Engine/Plugins/Enterprise/DatasmithCADImporter/Source/CADLibrary/Public/MeshDescriptionHelper.h
+++ b/Engine/Plugins/Enterprise/DatasmithCADImporter/Source/CADLibrary/Public/MeshDescriptionHelper.h
@@ -28,7 +28,6 @@ struct FMeshConversionContext
 	TArray<int32> VertexIds;
 	TArray<int32> SymmetricVertexIds;
 	TSet<int32> PatchesToMesh;
-	bool bSkipThinZoneMeshing = false;
 
 	FMeshConversionContext(const FImportParameters& InImportParams, const FMeshParameters& InMeshParameters)
 		: ImportParameters(InImportParams)
diff --git a/Engine/Plugins/Enterprise/DatasmithCADImporter/Source/DatasmithWireTranslator/Private/DatasmithWireTranslator.cpp b/Engine/Plugins/Enterprise/DatasmithCADImporter/Source/DatasmithWireTranslator/Private/DatasmithWireTranslator.cpp
--- a/Engine/Plugins/Enterprise/DatasmithCADImporter/Source/DatasmithWireTranslator/Private/DatasmithWireTranslator.cpp
+++ b/Engine/Plugins/Enterprise/DatasmithCADImporter/Source/DatasmithWireTranslator/Private/DatasmithWireTranslator.cpp
@@ -28,12 +28,6 @@ static TAutoConsoleVariable<bool> CVarAliasUseNative(
 	TEXT("If true, the AlaisStudio tessellator will be used. Default false.\n"),
 	ECVF_Default);
 
-static TAutoConsoleVariable<bool> CVarAliasSkipThinZoneMeshing(
-	TEXT("ds.WireTranslator.SkipThinZoneMeshing"),
-	true,
-	TEXT("If true, the global thin zone cvar will be honored. Default true.\n"),
-	ECVF_Default);
-
 TAutoConsoleVariable<bool> GAliasLayerAsActor(
 	TEXT("ds.CADTranslator.Alias.LayersAsActors"),
 	false,
diff --git a/Engine/Plugins/Enterprise/DatasmithCADImporter/Source/WireInterface/Private/AliasModelToCADKernelConverter.cpp b/Engine/Plugins/Enterprise/DatasmithCADImporter/Source/WireInterface/Private/AliasModelToCADKernelConverter.cpp
--- a/Engine/Plugins/Enterprise/DatasmithCADImporter/Source/WireInterface/Private/AliasModelToCADKernelConverter.cpp
+++ b/Engine/Plugins/Enterprise/DatasmithCADImporter/Source/WireInterface/Private/AliasModelToCADKernelConverter.cpp
@@ -369,9 +369,6 @@ bool FAliasModelToCADKernelConverter::Tessellate(const CADLibrary::FMeshParamete
 
 	CADLibrary::FMeshConversionContext Context(ImportParameters, InMeshParameters);
 
-	const IConsoleVariable* ConsoleVariable = IConsoleManager::Get().FindConsoleVariable(TEXT("ds.WireTranslator.SkipThinZoneMeshing"));
-	Context.bSkipThinZoneMeshing = ConsoleVariable ? ConsoleVariable->GetBool() : false;
-
 	return CADLibrary::FCADKernelTools::Tessellate(Model, Context, OutMeshDescription);
 }
 
diff --git a/Engine/Source/Runtime/Datasmith/CADKernel/CADKernel.Build.cs b/Engine/Source/Runtime/Datasmith/CADKernel/CADKernel.Build.cs
--- a/Engine/Source/Runtime/Datasmith/CADKernel/CADKernel.Build.cs
+++ b/Engine/Source/Runtime/Datasmith/CADKernel/CADKernel.Build.cs
@@ -8,7 +8,10 @@ namespace UnrealBuildTool.Rules
 		public CADKernel(ReadOnlyTargetRules Target)
 			: base(Target)
 		{
+			//OptimizeCode = CodeOptimization.Never;
 			DeterministicWarningLevel = WarningLevel.Off; // __DATE__ in Private/CADKernel/Core/System.cpp
+			
+			PublicDefinitions.Add("CADKERNEL_THINZONE=0");
 
 			PublicDependencyModuleNames.AddRange(
 				new string[]
diff --git a/Engine/Source/Runtime/Datasmith/CADKernel/Private/CADKernel/Mesh/Meshers/ParametricFaceMesher.cpp b/Engine/Source/Runtime/Datasmith/CADKernel/Private/CADKernel/Mesh/Meshers/ParametricFaceMesher.cpp
--- a/Engine/Source/Runtime/Datasmith/CADKernel/Private/CADKernel/Mesh/Meshers/ParametricFaceMesher.cpp
+++ b/Engine/Source/Runtime/Datasmith/CADKernel/Private/CADKernel/Mesh/Meshers/ParametricFaceMesher.cpp
@@ -27,7 +27,9 @@ FParametricFaceMesher::FParametricFaceMesher(FTopologicalFace& InFace, FModelMes
 	: Face(InFace)
 	, MeshModel(InMeshModel)
 	, Tolerances(InTolerances)
+#if CADKERNEL_THINZONE
 	, bThinZoneMeshing(bActivateThinZoneMeshing)
+#endif
 	, Grid(Face, MeshModel)
 {
 #ifdef DEBUG_ONLY_SURFACE_TO_DEBUG
diff --git a/Engine/Source/Runtime/Datasmith/CADKernel/Private/CADKernel/Mesh/Meshers/ParametricMesher.cpp b/Engine/Source/Runtime/Datasmith/CADKernel/Private/CADKernel/Mesh/Meshers/ParametricMesher.cpp
--- a/Engine/Source/Runtime/Datasmith/CADKernel/Private/CADKernel/Mesh/Meshers/ParametricMesher.cpp
+++ b/Engine/Source/Runtime/Datasmith/CADKernel/Private/CADKernel/Mesh/Meshers/ParametricMesher.cpp
@@ -26,7 +26,9 @@ TUniquePtr<FMesherReport> FMesherReport::MesherReport;
 
 FMesher::FMesher(FModelMesh& InMeshModel, double InGeometricTolerance, bool bActivateThinZoneMeshing)
 	: GeometricTolerance(InGeometricTolerance)
+#if CADKERNEL_THINZONE
 	, bThinZoneMeshing(bActivateThinZoneMeshing)
+#endif
 	, MeshModel(InMeshModel)
 {
 }
@@ -40,7 +42,9 @@ void FMesher::MeshEntities(TArray<FTopologicalShapeEntity*>& InEntities)
 
 FParametricMesher::FParametricMesher(FModelMesh& InMeshModel, double GeometricTolerance, bool bActivateThinZoneMeshing)
 	: Tolerances(GeometricTolerance)
+#if CADKERNEL_THINZONE
 	, bThinZoneMeshing(bActivateThinZoneMeshing)
+#endif
 	, MeshModel(InMeshModel)
 {
 }
