Missing file from previous submission

#rnx

diff --git a/Engine/Source/Runtime/Datasmith/CADKernel/Public/CADKernel/Mesh/Structure/ModelMesh.h b/Engine/Source/Runtime/Datasmith/CADKernel/Public/CADKernel/Mesh/Structure/ModelMesh.h
--- a/Engine/Source/Runtime/Datasmith/CADKernel/Public/CADKernel/Mesh/Structure/ModelMesh.h
+++ b/Engine/Source/Runtime/Datasmith/CADKernel/Public/CADKernel/Mesh/Structure/ModelMesh.h
@@ -92,7 +92,7 @@ public:
 
 	double GetGeometricTolerance() const
 	{
-		return 0.02; //Session->GetGeometricTolerance();
+		return MinSize * 0.5; // MinSize is twice GeometricTolerance
 	}
 
 	double GetMinSize() const
