Fixed crash when loading specific file

#rnx

diff --git a/Engine/Source/Runtime/Datasmith/CADKernel/Private/CADKernel/Mesh/Meshers/MesherTools.cpp b/Engine/Source/Runtime/Datasmith/CADKernel/Private/CADKernel/Mesh/Meshers/MesherTools.cpp
--- a/Engine/Source/Runtime/Datasmith/CADKernel/Private/CADKernel/Mesh/Meshers/MesherTools.cpp
+++ b/Engine/Source/Runtime/Datasmith/CADKernel/Private/CADKernel/Mesh/Meshers/MesherTools.cpp
@@ -333,39 +333,41 @@ void FMesherTools::FillImposedIsoCuttingPoints(TArray<double>& UEdgeSetOfInterse
 		PreviousU = InterU;
 	}
 
-	int32 Index;
-	int32 NewCoordinateCount = OutImposedIsoVertexSet.Num() - StartIndex;
-	switch (NewCoordinateCount)
-	{
-	case 0:
-		return;
 
-	case 1:
+	const int32 NewCoordinateCount = OutImposedIsoVertexSet.Num() - StartIndex;
+	if (NewCoordinateCount > 0)
 	{
-		int32 CuttingPointIndex = 0;
-		while (CuttingPointIndex < Edge.GetCrossingPointUs().Num() && Edge.GetCrossingPointUs()[CuttingPointIndex] + DOUBLE_SMALL_NUMBER <= OutImposedIsoVertexSet[StartIndex].Coordinate)
-		{
-			++CuttingPointIndex;
-		};
-		if (CuttingPointIndex > 0)
+		if (NewCoordinateCount == 1)
 		{
-			--CuttingPointIndex;
+			int32 CuttingPointIndex = 0;
+			while (CuttingPointIndex < Edge.GetCrossingPointUs().Num() && Edge.GetCrossingPointUs()[CuttingPointIndex] + DOUBLE_SMALL_NUMBER <= OutImposedIsoVertexSet[StartIndex].Coordinate)
+			{
+				++CuttingPointIndex;
+			};
+			if (CuttingPointIndex > 0)
+			{
+				--CuttingPointIndex;
+			}
+			// #cadkernel_check: To investigate further - This assert should not happen anymore.
+			if (Edge.GetDeltaUMaxs().IsValidIndex(CuttingPointIndex))
+			{
+				OutImposedIsoVertexSet[StartIndex].IsoDeltaU = Edge.GetDeltaUMaxs()[CuttingPointIndex] * AQuarter;
+			}
+			else
+			{
+				OutImposedIsoVertexSet.SetNum(StartIndex, EAllowShrinking::No);
+			}
 		}
-		OutImposedIsoVertexSet[StartIndex].IsoDeltaU = Edge.GetDeltaUMaxs()[CuttingPointIndex] * AQuarter;
-		break;
-	}
-
-	default:
-	{
-		OutImposedIsoVertexSet[StartIndex].IsoDeltaU = (OutImposedIsoVertexSet[StartIndex + 1].Coordinate - OutImposedIsoVertexSet[StartIndex].Coordinate) * AQuarter;
-		for (Index = StartIndex + 1; Index < OutImposedIsoVertexSet.Num() - 1; ++Index)
+		else
 		{
-			OutImposedIsoVertexSet[Index].IsoDeltaU = (OutImposedIsoVertexSet[Index + 1].Coordinate - OutImposedIsoVertexSet[Index - 1].Coordinate) * AEighth;
+			int32 Index = StartIndex + 1;
+			OutImposedIsoVertexSet[StartIndex].IsoDeltaU = (OutImposedIsoVertexSet[StartIndex + 1].Coordinate - OutImposedIsoVertexSet[StartIndex].Coordinate) * AQuarter;
+			for (; Index < OutImposedIsoVertexSet.Num() - 1; ++Index)
+			{
+				OutImposedIsoVertexSet[Index].IsoDeltaU = (OutImposedIsoVertexSet[Index + 1].Coordinate - OutImposedIsoVertexSet[Index - 1].Coordinate) * AEighth;
+			}
+			OutImposedIsoVertexSet[Index].IsoDeltaU = (OutImposedIsoVertexSet[Index].Coordinate - OutImposedIsoVertexSet[Index - 1].Coordinate) * AQuarter;
 		}
-		OutImposedIsoVertexSet[Index].IsoDeltaU = (OutImposedIsoVertexSet[Index].Coordinate - OutImposedIsoVertexSet[Index - 1].Coordinate) * AQuarter;
-		break;
-	}
-
 	}
 }
 
diff --git a/Engine/Source/Runtime/Datasmith/CADKernel/Private/CADKernel/Topo/TopologicalEdge.cpp b/Engine/Source/Runtime/Datasmith/CADKernel/Private/CADKernel/Topo/TopologicalEdge.cpp
--- a/Engine/Source/Runtime/Datasmith/CADKernel/Private/CADKernel/Topo/TopologicalEdge.cpp
+++ b/Engine/Source/Runtime/Datasmith/CADKernel/Private/CADKernel/Topo/TopologicalEdge.cpp
@@ -417,6 +417,11 @@ void FTopologicalEdge::ComputeCrossingPointCoordinates()
 	Sampler.Sample();
 
 	Presampling.SwapCoordinates(CrossingPointUs);
+	// #cadkernel_check: To investigate - It looks like the sampler does not start and end the sample data with the boundary values???
+	if (!FMath::IsNearlyEqual(CrossingPointUs.Last(), Boundary.GetMax(), UE_DOUBLE_SMALL_NUMBER))
+	{
+		CrossingPointUs.Add(Boundary.GetMax());
+	}
 
 	// Check sampling:
 	// the main idea is to avoid very small delta U between two or more points.
