Moved CADKernel library code from /Engine/Source/Runtime/Datasmith/CADKernel to /Engine/Source/Runtime/Datasmith/CADKernel/Base
This is in preparation of the creation of a CADKernelEngine module

#rnx

diff --git a/Engine/Plugins/Enterprise/DatasmithCADImporter/Source/CADInterfaces/Private/TechSoftBridge.cpp b/Engine/Plugins/Enterprise/DatasmithCADImporter/Source/CADInterfaces/Private/TechSoftBridge.cpp
--- a/Engine/Plugins/Enterprise/DatasmithCADImporter/Source/CADInterfaces/Private/TechSoftBridge.cpp
+++ b/Engine/Plugins/Enterprise/DatasmithCADImporter/Source/CADInterfaces/Private/TechSoftBridge.cpp
@@ -10,25 +10,25 @@
 #include "TechSoftUtils.h"
 #include "TUniqueTechSoftObj.h"
 
-#include "CADKernel/Core/Session.h"
+#include "Core/Session.h"
 
-#include "CADKernel/Geo/Curves/Curve.h"
-#include "CADKernel/Geo/Curves/NURBSCurveData.h"
+#include "Geo/Curves/Curve.h"
+#include "Geo/Curves/NURBSCurveData.h"
 
-#include "CADKernel/Geo/Surfaces/NurbsSurfaceData.h"
-#include "CADKernel/Geo/Surfaces/Surface.h"
+#include "Geo/Surfaces/NurbsSurfaceData.h"
+#include "Geo/Surfaces/Surface.h"
 
-#include "CADKernel/Topo/Body.h"
-#include "CADKernel/Topo/Model.h"
-#include "CADKernel/Topo/Shell.h"
-#include "CADKernel/Topo/TopologicalEdge.h"
-#include "CADKernel/Topo/TopologicalFace.h"
-#include "CADKernel/Topo/TopologicalLink.h"
-#include "CADKernel/Topo/TopologicalVertex.h"
+#include "Topo/Body.h"
+#include "Topo/Model.h"
+#include "Topo/Shell.h"
+#include "Topo/TopologicalEdge.h"
+#include "Topo/TopologicalFace.h"
+#include "Topo/TopologicalLink.h"
+#include "Topo/TopologicalVertex.h"
 
-#include "CADKernel/Utils/StringUtil.h"
+#include "Utils/StringUtil.h"
 
-#include "CADKernel/UI/Display.h"
+#include "UI/Display.h"
 
 #ifdef CADKERNEL_DEV
 #include "CADFileReport.h"
diff --git a/Engine/Plugins/Enterprise/DatasmithCADImporter/Source/CADInterfaces/Private/TechSoftBridge.h b/Engine/Plugins/Enterprise/DatasmithCADImporter/Source/CADInterfaces/Private/TechSoftBridge.h
--- a/Engine/Plugins/Enterprise/DatasmithCADImporter/Source/CADInterfaces/Private/TechSoftBridge.h
+++ b/Engine/Plugins/Enterprise/DatasmithCADImporter/Source/CADInterfaces/Private/TechSoftBridge.h
@@ -1,11 +1,11 @@
 // Copyright Epic Games, Inc. All Rights Reserved.
 #pragma once
-#include "CADKernel/Core/Types.h"
+#include "Core/Types.h"
 
 #ifdef USE_TECHSOFT_SDK
 
-#include "CADKernel/Geo/GeoEnum.h"
-#include "CADKernel/Math/MatrixH.h"
+#include "Geo/GeoEnum.h"
+#include "Math/MatrixH.h"
 #include "CADOptions.h"
 
 #include "TechSoftInterface.h"
diff --git a/Engine/Plugins/Enterprise/DatasmithCADImporter/Source/CADInterfaces/Private/TechSoftFileParserCADKernelTessellator.cpp b/Engine/Plugins/Enterprise/DatasmithCADImporter/Source/CADInterfaces/Private/TechSoftFileParserCADKernelTessellator.cpp
--- a/Engine/Plugins/Enterprise/DatasmithCADImporter/Source/CADInterfaces/Private/TechSoftFileParserCADKernelTessellator.cpp
+++ b/Engine/Plugins/Enterprise/DatasmithCADImporter/Source/CADInterfaces/Private/TechSoftFileParserCADKernelTessellator.cpp
@@ -10,17 +10,17 @@
 #include "TechSoftBridge.h"
 #include "TUniqueTechSoftObj.h"
 
-#include "CADKernel/Core/CADEntity.h"
-#include "CADKernel/Core/Session.h"
-#include "CADKernel/Core/Types.h"
+#include "Core/CADEntity.h"
+#include "Core/Session.h"
+#include "Core/Types.h"
 
-#include "CADKernel/Mesh/Meshers/Mesher.h"
-#include "CADKernel/Mesh/Structure/ModelMesh.h"
+#include "Mesh/Meshers/Mesher.h"
+#include "Mesh/Structure/ModelMesh.h"
 
-#include "CADKernel/Topo/Body.h"
-#include "CADKernel/Topo/Model.h"
-#include "CADKernel/Topo/TopologicalShapeEntity.h"
-#include "CADKernel/Topo/Topomaker.h"
+#include "Topo/Body.h"
+#include "Topo/Model.h"
+#include "Topo/TopologicalShapeEntity.h"
+#include "Topo/Topomaker.h"
 
 #include "HAL/PlatformTime.h"
 
diff --git a/Engine/Plugins/Enterprise/DatasmithCADImporter/Source/CADInterfaces/Public/TechSoftInterface.h b/Engine/Plugins/Enterprise/DatasmithCADImporter/Source/CADInterfaces/Public/TechSoftInterface.h
--- a/Engine/Plugins/Enterprise/DatasmithCADImporter/Source/CADInterfaces/Public/TechSoftInterface.h
+++ b/Engine/Plugins/Enterprise/DatasmithCADImporter/Source/CADInterfaces/Public/TechSoftInterface.h
@@ -2,7 +2,7 @@
 #pragma once
 
 #ifdef CADKERNEL_DEV
-#include "CADKernel/Core/Types.h"
+#include "Core/Types.h"
 #else
 #include "CoreMinimal.h"
 #endif
diff --git a/Engine/Plugins/Enterprise/DatasmithCADImporter/Source/CADKernelSurface/Private/CADKernelSurfaceExtension.cpp b/Engine/Plugins/Enterprise/DatasmithCADImporter/Source/CADKernelSurface/Private/CADKernelSurfaceExtension.cpp
--- a/Engine/Plugins/Enterprise/DatasmithCADImporter/Source/CADKernelSurface/Private/CADKernelSurfaceExtension.cpp
+++ b/Engine/Plugins/Enterprise/DatasmithCADImporter/Source/CADKernelSurface/Private/CADKernelSurfaceExtension.cpp
@@ -13,13 +13,13 @@
 #include "StaticMeshAttributes.h"
 #include "UObject/EnterpriseObjectVersion.h"
 
-#include "CADKernel/Core/CADKernelArchive.h"
-#include "CADKernel/Core/CADEntity.h"
-#include "CADKernel/Core/Session.h"
-#include "CADKernel/Core/Types.h"
-#include "CADKernel/Mesh/Structure/ModelMesh.h"
-#include "CADKernel/Topo/Model.h"
-#include "CADKernel/Topo/Body.h"
+#include "Core/CADKernelArchive.h"
+#include "Core/CADEntity.h"
+#include "Core/Session.h"
+#include "Core/Types.h"
+#include "Mesh/Structure/ModelMesh.h"
+#include "Topo/Model.h"
+#include "Topo/Body.h"
 
 void UCADKernelParametricSurfaceData::Serialize(FArchive& Ar)
 {
diff --git a/Engine/Plugins/Enterprise/DatasmithCADImporter/Source/CADKernelSurface/Public/CADModelToCADKernelConverterBase.h b/Engine/Plugins/Enterprise/DatasmithCADImporter/Source/CADKernelSurface/Public/CADModelToCADKernelConverterBase.h
--- a/Engine/Plugins/Enterprise/DatasmithCADImporter/Source/CADKernelSurface/Public/CADModelToCADKernelConverterBase.h
+++ b/Engine/Plugins/Enterprise/DatasmithCADImporter/Source/CADKernelSurface/Public/CADModelToCADKernelConverterBase.h
@@ -12,9 +12,9 @@
 #include "IDatasmithSceneElements.h"
 #include "MeshDescriptionHelper.h"
 
-#include "CADKernel/Core/Session.h"
-#include "CADKernel/Topo/Model.h"
-#include "CADKernel/Topo/Topomaker.h"
+#include "Core/Session.h"
+#include "Topo/Model.h"
+#include "Topo/Topomaker.h"
 
 struct FDatasmithMeshElementPayload;
 struct FDatasmithTessellationOptions;
diff --git a/Engine/Plugins/Enterprise/DatasmithCADImporter/Source/CADLibrary/Private/CADKernelTools.cpp b/Engine/Plugins/Enterprise/DatasmithCADImporter/Source/CADLibrary/Private/CADKernelTools.cpp
--- a/Engine/Plugins/Enterprise/DatasmithCADImporter/Source/CADLibrary/Private/CADKernelTools.cpp
+++ b/Engine/Plugins/Enterprise/DatasmithCADImporter/Source/CADLibrary/Private/CADKernelTools.cpp
@@ -9,23 +9,23 @@
 #include "StaticMeshAttributes.h"
 #include "StaticMeshOperations.h"
 
-#include "CADKernel/Core/CADEntity.h"
-#include "CADKernel/Core/MetadataDictionary.h"
-#include "CADKernel/Core/Session.h"
-#include "CADKernel/Core/Types.h"
-
-#include "CADKernel/Mesh/Criteria/Criterion.h"
-#include "CADKernel/Mesh/Meshers/Mesher.h"
-#include "CADKernel/Mesh/Structure/FaceMesh.h"
-#include "CADKernel/Mesh/Structure/ModelMesh.h"
-
-#include "CADKernel/Topo/Model.h"
-#include "CADKernel/Topo/Body.h"
-#include "CADKernel/Topo/Shell.h"
-#include "CADKernel/Topo/TopologicalEdge.h"
-#include "CADKernel/Topo/TopologicalShapeEntity.h"
-#include "CADKernel/Topo/TopologicalFace.h"
-#include "CADKernel/Topo/TopologicalVertex.h"
+#include "Core/CADEntity.h"
+#include "Core/MetadataDictionary.h"
+#include "Core/Session.h"
+#include "Core/Types.h"
+
+#include "Mesh/Criteria/Criterion.h"
+#include "Mesh/Meshers/Mesher.h"
+#include "Mesh/Structure/FaceMesh.h"
+#include "Mesh/Structure/ModelMesh.h"
+
+#include "Topo/Model.h"
+#include "Topo/Body.h"
+#include "Topo/Shell.h"
+#include "Topo/TopologicalEdge.h"
+#include "Topo/TopologicalShapeEntity.h"
+#include "Topo/TopologicalFace.h"
+#include "Topo/TopologicalVertex.h"
 
 typedef uint32 TriangleIndex[3];
 
diff --git a/Engine/Plugins/Enterprise/DatasmithCADImporter/Source/DatasmithOpenNurbsTranslator/Private/OpenNurbsBRepToCADKernelConverter.cpp b/Engine/Plugins/Enterprise/DatasmithCADImporter/Source/DatasmithOpenNurbsTranslator/Private/OpenNurbsBRepToCADKernelConverter.cpp
--- a/Engine/Plugins/Enterprise/DatasmithCADImporter/Source/DatasmithOpenNurbsTranslator/Private/OpenNurbsBRepToCADKernelConverter.cpp
+++ b/Engine/Plugins/Enterprise/DatasmithCADImporter/Source/DatasmithOpenNurbsTranslator/Private/OpenNurbsBRepToCADKernelConverter.cpp
@@ -12,22 +12,22 @@
 
 #include "CADKernelTools.h"
 
-#include "CADKernel/Core/Session.h"
-#include "CADKernel/Geo/Curves/Curve.h"
-#include "CADKernel/Geo/Curves/NURBSCurveData.h"
-#include "CADKernel/Geo/Surfaces/Surface.h"
-#include "CADKernel/Geo/Surfaces/NurbsSurfaceData.h"
-
-#include "CADKernel/Math/Boundary.h"
-#include "CADKernel/Math/Point.h"
-#include "CADKernel/Mesh/Structure/ModelMesh.h"
-
-#include "CADKernel/Topo/Body.h"
-#include "CADKernel/Topo/Model.h"
-#include "CADKernel/Topo/Shell.h"
-#include "CADKernel/Topo/TopologicalEdge.h"
-#include "CADKernel/Topo/TopologicalFace.h"
-#include "CADKernel/Topo/TopologicalLoop.h"
+#include "Core/Session.h"
+#include "Geo/Curves/Curve.h"
+#include "Geo/Curves/NURBSCurveData.h"
+#include "Geo/Surfaces/Surface.h"
+#include "Geo/Surfaces/NurbsSurfaceData.h"
+
+#include "Math/Boundary.h"
+#include "Math/Point.h"
+#include "Mesh/Structure/ModelMesh.h"
+
+#include "Topo/Body.h"
+#include "Topo/Model.h"
+#include "Topo/Shell.h"
+#include "Topo/TopologicalEdge.h"
+#include "Topo/TopologicalFace.h"
+#include "Topo/TopologicalLoop.h"
 
 
 namespace OpenNurbsBRepConverterTool
diff --git a/Engine/Plugins/Enterprise/DatasmithCADImporter/Source/DatasmithOpenNurbsTranslator/Private/OpenNurbsBRepToCADKernelConverter.h b/Engine/Plugins/Enterprise/DatasmithCADImporter/Source/DatasmithOpenNurbsTranslator/Private/OpenNurbsBRepToCADKernelConverter.h
--- a/Engine/Plugins/Enterprise/DatasmithCADImporter/Source/DatasmithOpenNurbsTranslator/Private/OpenNurbsBRepToCADKernelConverter.h
+++ b/Engine/Plugins/Enterprise/DatasmithCADImporter/Source/DatasmithOpenNurbsTranslator/Private/OpenNurbsBRepToCADKernelConverter.h
@@ -8,8 +8,8 @@
 #include "CADModelToCADKernelConverterBase.h"
 #include "OpenNurbsBRepConverter.h"
 
-#include "CADKernel/Core/Session.h"
-#include "CADKernel/Core/Types.h"
+#include "Core/Session.h"
+#include "Core/Types.h"
 
 class ON_BoundingBox;
 class ON_Brep;
diff --git a/Engine/Plugins/Enterprise/DatasmithCADImporter/Source/WireInterface/Private/AliasModelToCADKernelConverter.cpp b/Engine/Plugins/Enterprise/DatasmithCADImporter/Source/WireInterface/Private/AliasModelToCADKernelConverter.cpp
--- a/Engine/Plugins/Enterprise/DatasmithCADImporter/Source/WireInterface/Private/AliasModelToCADKernelConverter.cpp
+++ b/Engine/Plugins/Enterprise/DatasmithCADImporter/Source/WireInterface/Private/AliasModelToCADKernelConverter.cpp
@@ -12,21 +12,21 @@
 #include "Hal/PlatformMemory.h"
 #include "MeshDescriptionHelper.h"
 
-#include "CADKernel/Core/Session.h"
-#include "CADKernel/Geo/Curves/NURBSCurveData.h"
-#include "CADKernel/Geo/GeoEnum.h"
-#include "CADKernel/Geo/Surfaces/NurbsSurfaceData.h"
-#include "CADKernel/Geo/Surfaces/Surface.h"
-
-#include "CADKernel/Math/Point.h"
-#include "CADKernel/Mesh/Structure/ModelMesh.h"
-
-#include "CADKernel/Topo/Body.h"
-#include "CADKernel/Topo/Model.h"
-#include "CADKernel/Topo/Shell.h"
-#include "CADKernel/Topo/TopologicalEdge.h"
-#include "CADKernel/Topo/TopologicalFace.h"
-#include "CADKernel/Topo/TopologicalLoop.h"
+#include "Core/Session.h"
+#include "Geo/Curves/NURBSCurveData.h"
+#include "Geo/GeoEnum.h"
+#include "Geo/Surfaces/NurbsSurfaceData.h"
+#include "Geo/Surfaces/Surface.h"
+
+#include "Math/Point.h"
+#include "Mesh/Structure/ModelMesh.h"
+
+#include "Topo/Body.h"
+#include "Topo/Model.h"
+#include "Topo/Shell.h"
+#include "Topo/TopologicalEdge.h"
+#include "Topo/TopologicalFace.h"
+#include "Topo/TopologicalLoop.h"
 
 #if PLATFORM_WINDOWS
 #include "Windows/AllowWindowsPlatformTypes.h"
diff --git a/Engine/Plugins/Enterprise/DatasmithCADImporter/Source/WireInterface/Private/AliasModelToCADKernelConverter.h b/Engine/Plugins/Enterprise/DatasmithCADImporter/Source/WireInterface/Private/AliasModelToCADKernelConverter.h
--- a/Engine/Plugins/Enterprise/DatasmithCADImporter/Source/WireInterface/Private/AliasModelToCADKernelConverter.h
+++ b/Engine/Plugins/Enterprise/DatasmithCADImporter/Source/WireInterface/Private/AliasModelToCADKernelConverter.h
@@ -9,8 +9,8 @@
 
 #include "CADModelToCADKernelConverterBase.h"
 
-#include "CADKernel/Core/Session.h"
-#include "CADKernel/Core/Types.h"
+#include "Core/Session.h"
+#include "Core/Types.h"
 
 class AlDagNode;
 class AlShell;
diff --git a/Engine/Source/Runtime/Datasmith/CADKernel/Base/Private/Core/System.cpp b/Engine/Source/Runtime/Datasmith/CADKernel/Base/Private/Core/System.cpp
new file mode 100644
--- /dev/null
+++ b/Engine/Source/Runtime/Datasmith/CADKernel/Base/Private/Core/System.cpp
@@ -0,0 +1,172 @@
+// Copyright Epic Games, Inc. All Rights Reserved.
+#include "Core/System.h"
+
+#include "Core/Database.h"
+#include "Core/CADEntity.h"
+#include "UI/Visu.h"
+#include "Utils/Util.h"
+
+#include "HAL/FileManager.h"
+#include "Misc/Paths.h"
+
+#ifdef CADKERNEL_DEV
+#include "Core/KernelParameters.h"
+#include "Core/Version.h"
+
+#include <stdlib.h>
+#include <signal.h>
+#endif 
+
+namespace UE::CADKernel
+{
+
+TUniquePtr<FSystem> FSystem::Instance = nullptr;
+
+FSystem::FSystem()
+	: DefaultVisu()
+	, Viewer(&DefaultVisu)
+	, Console(&DefaultConsole)
+	, ProgressManager(&DefaultProgressManager)
+#ifdef CADKERNEL_DEV
+	, Parameters(MakeShared<FKernelParameters>())
+#endif
+{
+	LogLevel = Log;
+	VerboseLevel = Log;
+}
+
+void FSystem::Initialize(bool bIsDll, const FString& LogFilePath, const FString& SpyFilePath)
+{
+	SetVerboseLevel(Log);
+
+	if (LogFilePath.Len() > 0)
+	{
+		DefineLogFile(LogFilePath);
+	}
+	if (SpyFilePath.Len() > 0)
+	{
+		DefineSpyFile(SpyFilePath);
+	}
+
+	PrintHeader();
+
+	fflush(stdout);
+
+	if(bIsDll) 
+	{
+		SetVerboseLevel(EVerboseLevel::NoVerbose);
+	} 
+	else 
+	{
+		SetVerboseLevel(EVerboseLevel::Log);
+	}
+}
+
+void FSystem::CloseLogFiles()
+{
+	if (LogFile)
+	{
+		LogFile->Close();
+		LogFile.Reset();
+	}
+	if (SpyFile)
+	{
+		SpyFile->Close();
+		SpyFile.Reset();
+	}
+}
+
+void FSystem::Shutdown()
+{
+	CloseLogFiles();
+	Instance.Reset();
+}
+
+void FSystem::DefineLogFile(const FString& InLogFilePath, EVerboseLevel InLevel)
+{
+	if(LogFile) 
+	{
+		LogFile->Close();
+		LogFile.Reset();
+	}
+
+	LogFile = MakeShareable<FArchive>(IFileManager::Get().CreateFileWriter(*InLogFilePath, IO_WRITE));
+	LogLevel = InLevel;
+}
+
+void FSystem::DefineSpyFile(const FString& InSpyFilePath)
+{
+	if(SpyFile) 
+	{
+		SpyFile->Close();
+		SpyFile.Reset();
+	}
+	SpyFile = MakeShareable<FArchive>(IFileManager::Get().CreateFileWriter(*InSpyFilePath, IO_WRITE));
+}
+
+
+#if defined(CADKERNEL_DEV) || defined(CADKERNEL_STDA)
+void FSystem::DefineReportFile(const FString& InReportFilePath)
+{
+	if (ReportFile.IsValid())
+	{
+		ReportFile->Close();
+		ReportFile.Reset();
+
+		if (ReportHeaderFile.IsValid())
+		{
+			ReportHeaderFile->Close();
+			ReportHeaderFile.Reset();
+		}
+	}
+
+	FString ReportHeaderPath = FPaths::GetPath(InReportFilePath);
+	if (!FPaths::DirectoryExists(ReportHeaderPath))
+	{
+		IFileManager::Get().MakeDirectory(*ReportHeaderPath, true);
+	}
+
+	ReportFile = MakeShareable<FArchive>(IFileManager::Get().CreateFileWriter(*InReportFilePath, IO_WRITE));
+
+	ReportHeaderPath = FPaths::ConvertRelativePathToFull(FPaths::Combine(ReportHeaderPath, TEXT("ReportHeader.csv")));
+	if(!FPaths::FileExists(ReportHeaderPath))
+	{
+		ReportHeaderFile = MakeShareable<FArchive>(IFileManager::Get().CreateFileWriter(*ReportHeaderPath, IO_WRITE));
+	}
+}
+#endif
+
+void FSystem::InitializeCADKernel()
+{
+	FSystem::Get().Initialize();
+	FSystem::Get().SetVerboseLevel(EVerboseLevel::Log);
+}
+
+FString FSystem::GetToolkitVersion()
+{
+#ifdef CADKERNEL_DEV
+	return UTF8_TO_TCHAR(TOOLKIT_VERSION_ASCII);
+#endif
+	return FString();
+}
+
+FString FSystem::GetCompilationDate()
+{
+	return UTF8_TO_TCHAR(__DATE__);
+}
+
+void FSystem::PrintHeader()
+{
+	FMessage::Printf(Log, TEXT("_______________________________________________________________________________\n"));
+	FMessage::Printf(Log, TEXT("\n"));
+#ifdef CADKERNEL_DEV
+	FMessage::Printf(Log, TEXT("\tDatasmith CAD Kernel Toolkit release %s (%s)\n"), UTF8_TO_TCHAR(TOOLKIT_VERSION_ASCII), UTF8_TO_TCHAR(__DATE__));
+#endif
+	FMessage::Printf(Log, TEXT("\tCopyright Epic Games, Inc. All Rights Reserved.\n"));
+	FMessage::Printf(Log, TEXT("\n"));
+	FMessage::Printf(Log, TEXT("_______________________________________________________________________________\n"));
+	FMessage::Printf(Log, TEXT("\n"));
+}
+
+
+}
\ No newline at end of file
diff --git a/Engine/Source/Runtime/Datasmith/CADKernel/Base/Private/Geo/Sampling/SurfacicPolyline.cpp b/Engine/Source/Runtime/Datasmith/CADKernel/Base/Private/Geo/Sampling/SurfacicPolyline.cpp
new file mode 100644
--- /dev/null
+++ b/Engine/Source/Runtime/Datasmith/CADKernel/Base/Private/Geo/Sampling/SurfacicPolyline.cpp
@@ -0,0 +1,213 @@
+// Copyright Epic Games, Inc. All Rights Reserved.
+
+#include "Geo/Sampling/SurfacicPolyline.h"
+
+#include "Core/System.h"
+#include "Geo/Curves/Curve.h"
+#include "Geo/Surfaces/Surface.h"
+#include "Geo/Sampler/SamplerOnParam.h"
+#include "Geo/Sampling/PolylineTools.h"
+#include "Math/SlopeUtils.h"
+
+namespace UE::CADKernel
+{
+
+FSurfacicPolyline::FSurfacicPolyline(TSharedRef<FSurface> InCarrierSurface, TSharedRef<FCurve> Curve2D, const double InTolerance)
+	: bWithNormals(true)
+	, bWithTangent(false)
+{
+	FSurfacicCurveSamplerOnParam Sampler(InCarrierSurface.Get(), Curve2D.Get(), Curve2D->GetBoundary(), InTolerance, InTolerance, *this);
+	Sampler.Sample();
+	BoundingBox.Set(Points2D);
+}
+
+FSurfacicPolyline::FSurfacicPolyline(TSharedRef<FSurface> InCarrierSurface, TSharedRef<FCurve> Curve2D)
+	: FSurfacicPolyline(InCarrierSurface, Curve2D, InCarrierSurface->Get3DTolerance())
+{
+}
+
+FSurfacicPolyline::FSurfacicPolyline(TSharedRef<FSurface> InCarrierSurface, TSharedRef<FCurve> Curve2D, const double ChordTolerance, const double ParamTolerance, bool bInWithNormals, bool bInWithTangents)
+	: bWithNormals(bInWithNormals)
+	, bWithTangent(bInWithTangents)
+{
+	FSurfacicCurveSamplerOnParam Sampler(InCarrierSurface.Get(), Curve2D.Get(), Curve2D->GetBoundary(), ChordTolerance, ParamTolerance, *this);
+	Sampler.Sample();
+	BoundingBox.Set(Points2D);
+}
+
+void FSurfacicPolyline::CheckIfDegenerated(const double Tolerance3D, const FSurfacicTolerance& ToleranceIso, const FLinearBoundary& Boudary, bool& bDegeneration2D, bool& bDegeneration3D, double& Length3D) const
+{
+	TPolylineApproximator<FPoint> Approximator(Coordinates, Points3D);
+	int32 BoundaryIndices[2];
+	Approximator.GetStartEndIndex(Boudary, BoundaryIndices);
+
+	TPolylineApproximator<FPoint> Approximator3D(Coordinates, Points3D);
+	Length3D = Approximator3D.ComputeLengthOfSubPolyline(BoundaryIndices, Boudary);
+
+	if (!FMath::IsNearlyZero(Length3D, Tolerance3D))
+	{
+		bDegeneration3D = false;
+		bDegeneration2D = false;
+		return;
+	}
+
+	bDegeneration3D = true;
+	Length3D = 0.;
+
+	// Tolerance along Iso U/V is very costly to compute and not accurate.
+	// To test if a curve is degenerated, its 2d bounding box is compute and its compare to the surface boundary along U and along V
+	// Indeed, defining a Tolerance2D has no sense has the boundary length along an Iso could be very very huge compare to the boundary length along the other Iso like [[0, 1000] [0, 1]]
+	// The tolerance along ans iso is the length of the boundary along this iso divide by 100 000: if the curve length in 3d is 10m, the tolerance is 0.01mm
+
+	FAABB2D Aabb;
+	TPolylineApproximator<FPoint2D> Approximator2D(Coordinates, Points2D);
+	Approximator2D.ComputeBoundingBox(BoundaryIndices, Boudary, Aabb);
+
+	bDegeneration2D = (Aabb.GetSize(0) < ToleranceIso[EIso::IsoU] && Aabb.GetSize(1) < ToleranceIso[EIso::IsoV]);
+}
+
+void FSurfacicPolyline::GetExtremities(const FLinearBoundary& InBoundary, const double Tolerance3D, const FSurfacicTolerance& MinToleranceIso, FSurfacicCurveExtremities& Extremities) const
+{
+	FDichotomyFinder Finder(Coordinates);
+	const int32 StartIndex = Finder.Find(InBoundary.Min);
+	const int32 EndIndex = Finder.Find(InBoundary.Max);
+
+	Extremities[0].Point2D = PolylineTools::ComputePoint(Coordinates, Points2D, StartIndex, InBoundary.Min);
+	Extremities[0].Point = PolylineTools::ComputePoint(Coordinates, Points3D, StartIndex, InBoundary.Min);
+	Extremities[0].Tolerance = ComputeTolerance(Tolerance3D, MinToleranceIso, StartIndex);
+
+	Extremities[1].Point2D = PolylineTools::ComputePoint(Coordinates, Points2D, EndIndex, InBoundary.Max);
+	Extremities[1].Point = PolylineTools::ComputePoint(Coordinates, Points3D, EndIndex, InBoundary.Max);
+	if (EndIndex == StartIndex)
+	{
+		Extremities[1].Tolerance = Extremities[0].Tolerance;
+	}
+	else
+	{
+		Extremities[1].Tolerance = ComputeTolerance(Tolerance3D, MinToleranceIso, EndIndex);
+	}
+}
+
+void FSurfacicPolyline::ComputeIntersectionsWithIsos(const FLinearBoundary& InBoundary, const TArray<double>& IsoCoordinates, const EIso TypeIso, const FSurfacicTolerance& ToleranceIso, TArray<double>& Intersection) const
+{
+	if (BoundingBox.Length(TypeIso) < ToleranceIso[TypeIso])
+	{
+		// the edge is on an IsoCurve on Iso axis
+		return;
+	}
+
+#ifdef DEBUG_COMPUTEINTERSECTIONSWITHISOS
+	F3DDebugSession _(TEXT("IntersectEdgeIsos"));
+#endif
+ 
+	Intersection.Reserve(IsoCoordinates.Num());
+
+	double UMin;
+	double UMax;
+
+	int32 IsoCoordinateIndex = 0;
+	double IsoCoordinate = IsoCoordinates[0];
+
+	for (int32 Index = 0; Index < Points2D.Num() - 1; ++Index)
+	{
+		// if the segment is outside edge boundary, go to the next one
+		if (Coordinates[Index + 1] < InBoundary.GetMin())
+		{
+			continue;
+		}
+
+		if (Coordinates[Index] > InBoundary.GetMax())
+		{
+			break;
+		}
+
+		// if the segment is too parallel to Iso, go to the next one
+		double EdgeLocalSlope = ComputeUnorientedSlope(Points2D[Index], Points2D[Index + 1], 0);
+		if (TypeIso == EIso::IsoV)
+		{
+			if (EdgeLocalSlope < 0.1 || EdgeLocalSlope > 3.9)
+			{
+				continue;
+			}
+		}
+		else
+		{
+			if (EdgeLocalSlope < 2.1 && EdgeLocalSlope > 1.9)
+			{
+				continue;
+			}
+		}
+
+		GetMinMax(Points2D[Index][TypeIso], Points2D[Index + 1][TypeIso], UMin, UMax);
+		UMin -= DOUBLE_SMALL_NUMBER;
+		UMax += DOUBLE_SMALL_NUMBER;
+
+		if (IsoCoordinate < UMin || IsoCoordinate > UMax)
+		{
+			if (IsoCoordinateIndex > 0 && IsoCoordinateIndex < IsoCoordinates.Num() - 1 && IsoCoordinates[IsoCoordinateIndex - 1] < UMin && UMax < IsoCoordinates[IsoCoordinateIndex + 1])
+			{
+				continue;
+			}
+		}
+
+		if (IsoCoordinate > UMin)
+		{
+			while (IsoCoordinateIndex > 0 && IsoCoordinates[IsoCoordinateIndex - 1] > UMin)
+			{
+				IsoCoordinateIndex--;
+			}
+			IsoCoordinate = IsoCoordinates[IsoCoordinateIndex];
+		}
+
+		while (IsoCoordinateIndex < IsoCoordinates.Num() - 1 && IsoCoordinates[IsoCoordinateIndex] < UMin)
+		{
+			IsoCoordinateIndex++;
+		}
+		IsoCoordinate = IsoCoordinates[IsoCoordinateIndex];
+
+		if (UMax < IsoCoordinate)
+		{
+			continue;
+		}
+
+		for (;; ++IsoCoordinateIndex)
+		{
+			if (IsoCoordinateIndex == IsoCoordinates.Num() || IsoCoordinates[IsoCoordinateIndex] > UMax)
+			{
+				--IsoCoordinateIndex;
+				break;
+			}
+			IsoCoordinate = IsoCoordinates[IsoCoordinateIndex];
+
+			if ((IsoCoordinate > UMin) && (IsoCoordinate < UMax))
+			{
+				double EdgeCoord = 0.;
+				const double Delta = (Points2D[Index + 1][TypeIso] - Points2D[Index][TypeIso]);
+				if (FMath::IsNearlyZero(Delta))
+				{
+					EdgeCoord = Coordinates[Index];
+				}
+				else
+				{
+					const double LocalCoord = (IsoCoordinate - Points2D[Index][TypeIso]) / Delta;
+					EdgeCoord = PolylineTools::LinearInterpolation(Coordinates, Index, LocalCoord);
+				}
+
+#ifdef DEBUG_COMPUTEINTERSECTIONSWITHISOS
+				::DisplayPoint(Points2D[Index + 1], EVisuProperty::RedPoint);
+				::DisplayPoint(Points2D[Index], EVisuProperty::RedPoint);
+				FPoint2D IntersectionPoint = PolylineTools::LinearInterpolation(Points2D, Index, LocalCoord);
+				::DisplayPoint(IntersectionPoint);
+				Wait();
+#endif // #ifdef DEBUG_COMPUTEINTERSECTIONSWITHISOS
+
+				if (!(Intersection.Num() && FMath::IsNearlyEqual(EdgeCoord, Intersection.Last())))
+				{
+					Intersection.Add(EdgeCoord);
+				}
+			}
+		}
+	}
+}
+
+} //namespace UE::CADKernel
diff --git a/Engine/Source/Runtime/Datasmith/CADKernel/Base/Private/Geo/Surfaces/PlaneSurface.cpp b/Engine/Source/Runtime/Datasmith/CADKernel/Base/Private/Geo/Surfaces/PlaneSurface.cpp
new file mode 100644
--- /dev/null
+++ b/Engine/Source/Runtime/Datasmith/CADKernel/Base/Private/Geo/Surfaces/PlaneSurface.cpp
@@ -0,0 +1,165 @@
+// Copyright Epic Games, Inc. All Rights Reserved.
+
+#include "Geo/Surfaces/PlaneSurface.h"
+
+#include "Core/System.h"
+#include "Geo/GeoPoint.h"
+#include "Geo/Sampling/SurfacicSampling.h"
+
+namespace UE::CADKernel
+{
+
+FPlaneSurface::FPlaneSurface(const double InToleranceGeometric, const FMatrixH& InMatrix, const FSurfacicBoundary& InBoundary)
+	: FSurface(InToleranceGeometric, InBoundary)
+{
+	Matrix = InMatrix;
+
+	ensureCADKernel(FMath::IsNearlyZero(Matrix.Get(3, 0)) && FMath::IsNearlyZero(Matrix.Get(3, 1)) && FMath::IsNearlyZero(Matrix.Get(3, 2)));
+
+	InverseMatrix = Matrix;
+	InverseMatrix.Inverse();
+	ComputeMinToleranceIso();
+}
+
+FPlaneSurface::FPlaneSurface(const double InToleranceGeometric, const FPoint& InPosition, FPoint InNormal, const FSurfacicBoundary& InBoundary)
+	: FSurface(InToleranceGeometric, InBoundary)
+{
+	InNormal.Normalize();
+	Matrix.FromAxisOrigin(InNormal, InPosition);
+
+	InverseMatrix = Matrix;
+	InverseMatrix.Inverse();
+	ComputeMinToleranceIso();
+}
+
+FPlane FPlaneSurface::GetPlane() const
+{
+	FSurfacicPoint Point3D;
+	EvaluatePoint(FPoint2D::ZeroPoint, Point3D, 0);
+	FPoint Normal = Matrix.Column(2);
+	return FPlane(Point3D.Point, Normal);
+}
+
+void FPlaneSurface::EvaluatePoint(const FPoint2D& InSurfacicCoordinate, FSurfacicPoint& OutPoint3D, int32 InDerivativeOrder) const
+{
+	OutPoint3D.DerivativeOrder = InDerivativeOrder;
+	OutPoint3D.Point = Matrix.Multiply(InSurfacicCoordinate);
+	if(InDerivativeOrder>0) 
+	{
+		OutPoint3D.GradientU = Matrix.Column(0);
+		OutPoint3D.GradientV = Matrix.Column(1);
+	}
+
+	if(InDerivativeOrder>1) 
+	{
+		OutPoint3D.LaplacianU = FPoint::ZeroPoint;
+		OutPoint3D.LaplacianV = FPoint::ZeroPoint;
+		OutPoint3D.LaplacianUV = FPoint::ZeroPoint;
+	}
+}
+
+void FPlaneSurface::EvaluatePoints(const TArray<FPoint2D>& InSurfacicCoordinates, TArray<FSurfacicPoint>& OutPoint3D, int32 InDerivativeOrder) const
+{
+	int32 PointNum = InSurfacicCoordinates.Num();
+	OutPoint3D.SetNum(PointNum);
+
+	for (int32 Index = 0; Index < PointNum; ++Index)
+	{
+		OutPoint3D[Index].DerivativeOrder = InDerivativeOrder;
+	}
+
+	for (int32 Index = 0; Index < PointNum; ++Index)
+	{
+		OutPoint3D[Index].Point = Matrix.Multiply(InSurfacicCoordinates[Index]);
+	}
+
+	if (InDerivativeOrder > 0)
+	{
+		FPoint GradientU(Matrix.Column(0));
+		FPoint GradientV(Matrix.Column(1));
+
+		for (int32 Index = 0; Index < PointNum; ++Index)
+		{
+			OutPoint3D[Index].GradientU = GradientU;
+		}
+
+		for (int32 Index = 0; Index < PointNum; ++Index)
+		{
+			OutPoint3D[Index].GradientV = GradientV;
+		}
+	}
+}
+
+void FPlaneSurface::EvaluatePointGrid(const FCoordinateGrid& Coordinates, FSurfacicSampling& OutPoints, bool bComputeNormals) const
+{
+	OutPoints.bWithNormals = bComputeNormals;
+
+	int32 PointNum = Coordinates.Count();
+	OutPoints.Reserve(PointNum);
+
+	OutPoints.Set2DCoordinates(Coordinates);
+
+	for (FPoint Point : OutPoints.Points2D)
+	{
+		OutPoints.Points3D.Emplace(Matrix.Multiply(Point));
+	}
+
+	if(bComputeNormals)
+	{
+		FPoint Normal(Matrix.Column(2));
+		OutPoints.Normals.Init(Normal, PointNum);
+	}
+}
+
+FPoint FPlaneSurface::ProjectPoint(const FPoint& Point, FPoint* OutProjectedPoint) const
+{
+	FPoint PointCoordinate = InverseMatrix.Multiply(Point);
+	PointCoordinate.Z = 0.0;
+
+	if(OutProjectedPoint)
+	{
+		*OutProjectedPoint = Matrix.Multiply(PointCoordinate);
+	}
+
+	return PointCoordinate;
+}
+
+void FPlaneSurface::ProjectPoints(const TArray<FPoint>& Points, TArray<FPoint>* PointCoordinates, TArray<FPoint>* OutProjectedPointS) const
+{
+	PointCoordinates->Reserve(Points.Num());
+	if(OutProjectedPointS) 
+	{
+		OutProjectedPointS->Reserve(Points.Num());
+	}
+
+	for (const FPoint& Point : Points)
+	{
+		FPoint& PointCoordinate = PointCoordinates->Emplace_GetRef(InverseMatrix.Multiply(Point));
+		PointCoordinate.Z = 0.0;
+	}
+
+	if (OutProjectedPointS)
+	{
+		for (const FPoint& PointCoordinate : *PointCoordinates)
+		{
+			OutProjectedPointS->Emplace(Matrix.Multiply(PointCoordinate));
+		}
+	}
+}
+
+TSharedPtr<FEntityGeom> FPlaneSurface::ApplyMatrix(const FMatrixH& InMatrix) const
+{
+	FMatrixH NewMatrix = InMatrix * Matrix;
+	return FEntity::MakeShared<FPlaneSurface>(Tolerance3D, NewMatrix, Boundary);
+}
+
+#ifdef CADKERNEL_DEV
+FInfoEntity& FPlaneSurface::GetInfo(FInfoEntity& Info) const
+{
+	return FSurface::GetInfo(Info)
+		.Add(TEXT("Matrix"), Matrix)
+		.Add(TEXT("Inverse"), InverseMatrix);
+}
+#endif
+
+}
\ No newline at end of file
diff --git a/Engine/Source/Runtime/Datasmith/CADKernel/Base/Private/Math/Point.cpp b/Engine/Source/Runtime/Datasmith/CADKernel/Base/Private/Math/Point.cpp
new file mode 100644
--- /dev/null
+++ b/Engine/Source/Runtime/Datasmith/CADKernel/Base/Private/Math/Point.cpp
@@ -0,0 +1,54 @@
+// Copyright Epic Games, Inc. All Rights Reserved.
+
+#include "Math/Point.h"
+
+#include "Utils/Util.h"
+#include "Math/MatrixH.h"
+
+namespace UE::CADKernel
+{
+
+const FPoint FPoint::ZeroPoint(0.,0.,0.);
+const FPoint FPoint::UnitPoint(1., 1., 1.);
+const FPoint FPoint::FarawayPoint(HUGE_VALUE, HUGE_VALUE, HUGE_VALUE);
+const int32 FPoint::Dimension = 3;
+
+const FFPoint FFPoint::ZeroPoint(0.f, 0.f, 0.f);
+const FFPoint FFPoint::FarawayPoint(HUGE_VALUE, HUGE_VALUE, HUGE_VALUE);
+const int32 FFPoint::Dimension = 3;
+
+const FPoint2D FPoint2D::ZeroPoint(0., 0.);
+const FPoint2D FPoint2D::FarawayPoint(HUGE_VALUE, HUGE_VALUE);
+const int32 FPoint2D::Dimension = 2;
+
+const FPointH FPointH::ZeroPoint(0., 0., 0., 1.);
+const FPointH FPointH::FarawayPoint(HUGE_VALUE, HUGE_VALUE, HUGE_VALUE, 1.);
+const int32 FPointH::Dimension = 4;
+
+double FPoint::SignedAngle(const FPoint & Other, const FPoint & Normal) const
+{
+	FPoint Vector1 = *this; 
+	FPoint Vector2 = Other; 
+	FPoint Vector3 = Normal; 
+
+	Vector1.Normalize();
+	Vector2.Normalize();
+	Vector3.Normalize();
+
+	double ScalarProduct = Vector1 * Vector2;
+
+	if (ScalarProduct >= 1 - DOUBLE_SMALL_NUMBER)
+	{
+		return 0.;
+	}
+
+	if (ScalarProduct <= -1 + DOUBLE_SMALL_NUMBER)
+	{
+		return DOUBLE_PI;
+	}
+
+	return MixedTripleProduct(Vector1, Vector2, Vector3) > 0 ? acos(ScalarProduct) : -acos(ScalarProduct);
+}
+
+
+}
\ No newline at end of file
