Moved CADKernel library code from /Engine/Source/Runtime/Datasmith/CADKernel to /Engine/Source/Runtime/Datasmith/CADKernel/Base

diff --git a/Engine/Source/Runtime/Datasmith/CADKernel/Base/Private/Mesh/Meshers/IsoTriangulator.cpp b/Engine/Source/Runtime/Datasmith/CADKernel/Base/Private/Mesh/Meshers/IsoTriangulator.cpp
new file mode 100644
--- /dev/null
+++ b/Engine/Source/Runtime/Datasmith/CADKernel/Base/Private/Mesh/Meshers/IsoTriangulator.cpp
@@ -0,0 +1,2520 @@
+// Copyright Epic Games, Inc. All Rights Reserved.
+
+#include "Mesh/Meshers/IsoTriangulator.h"
+
+#include "Core/Chrono.h"
+#include "Math/MathConst.h"
+#include "Math/Point.h"
+#include "Mesh/Meshers/BowyerWatsonTriangulator.h"
+#include "Mesh/Meshers/CycleTriangulator.h"
+#include "Mesh/Meshers/IsoTriangulator/IsoCell.h"
+#include "Mesh/Meshers/IsoTriangulator/IsoNode.h"
+#include "Mesh/Meshers/IsoTriangulator/IsoSegment.h"
+#include "Mesh/Structure/Grid.h"
+#include "Mesh/Structure/EdgeMesh.h"
+#include "Mesh/Structure/FaceMesh.h"
+#include "Mesh/Structure/LoopCleaner.h"
+#include "Mesh/Structure/ThinZone2D.h"
+#include "Topo/TopologicalEdge.h"
+#include "Topo/TopologicalFace.h"
+#include "UI/Display.h"
+#include "Utils/ArrayUtils.h"
+
+namespace UE::CADKernel
+{
+
+#ifdef DEBUG_BOWYERWATSON
+bool FBowyerWatsonTriangulator::bDisplay = false;
+#endif
+
+namespace IsoTriangulatorImpl
+{
+const double MaxSlopeToBeIso = 0.125;
+
+const double LimitValueMin(double Slope)
+{
+	return Slope - MaxSlopeToBeIso;
+}
+
+const double LimitValueMax(double Slope)
+{
+	return Slope + MaxSlopeToBeIso;
+}
+
+struct FCandidateSegment
+{
+	FLoopNode& StartNode;
+	FLoopNode& EndNode;
+	double Length;
+
+	FCandidateSegment(const FGrid& Grid, FLoopNode& Node1, FLoopNode& Node2)
+		: StartNode(Node1)
+		, EndNode(Node2)
+	{
+		Length = Node1.Get2DPoint(EGridSpace::UniformScaled, Grid).Distance(Node2.Get2DPoint(EGridSpace::UniformScaled, Grid));
+	}
+};
+
+}
+
+FIsoTriangulator::FIsoTriangulator(FGrid& InGrid, FFaceMesh& OutMesh, const FMeshingTolerances& InTolerances)
+	: Grid(InGrid)
+	, Mesh(OutMesh)
+	, LoopSegmentsIntersectionTool(InGrid, InTolerances.GeometricTolerance)
+	, InnerSegmentsIntersectionTool(InGrid, InTolerances.GeometricTolerance)
+	, InnerToOuterIsoSegmentsIntersectionTool(InGrid, InTolerances.GeometricTolerance)
+	, ThinZoneIntersectionTool(InGrid, InTolerances.GeometricTolerance)
+	, Tolerances(InTolerances)
+{
+	FinalInnerSegments.Reserve(3 * Grid.InnerNodesCount());
+	IndexOfLowerLeftInnerNodeSurroundingALoop.Reserve(Grid.GetLoopCount());
+
+#ifdef DEBUG_ISOTRIANGULATOR
+
+#ifdef DEBUG_ONLY_SURFACE_TO_DEBUG
+	if (Grid.GetFace().GetId() == FaceToDebug)
+#endif
+	{
+		bDisplay = true;
+	}
+
+#endif
+}
+
+bool FIsoTriangulator::Triangulate()
+{
+	EGridSpace DisplaySpace = EGridSpace::UniformScaled;
+#ifdef DEBUG_ISOTRIANGULATOR
+	F3DDebugSession _(bDisplay, FString::Printf(TEXT("Triangulate %d"), Grid.GetFace().GetId()));
+#endif
+
+	FTimePoint StartTime = FChrono::Now();
+
+	// =============================================================================================================
+	// Build the first elements (IsoNodes (i.e. Inner nodes), Loops nodes, and knows segments 
+	// =============================================================================================================
+
+	BuildNodes();
+
+#ifdef DEBUG_ISOTRIANGULATOR
+	DisplayIsoNodes(DisplaySpace);
+#endif
+
+	FillMeshNodes();
+	BuildLoopSegments();
+
+#ifdef DEBUG_ISOTRIANGULATOR
+	DisplayLoops(TEXT("FIsoTrianguler::LoopSegments"), true, false);
+	Wait(false);
+#endif
+
+	FLoopCleaner LoopCleaner(*this);
+	if (!LoopCleaner.Run())
+	{
+#ifdef CADKERNEL_DEV
+		FMesherReport::Get().Logs.AddDegeneratedLoop();
+#endif
+		FMessage::Printf(EVerboseLevel::Log, TEXT("The meshing of the surface %d failed due to a degenerated loop\n"), Grid.GetFace().GetId());
+		return false;
+	}
+
+#ifdef DEBUG_ISOTRIANGULATOR
+	DisplayLoops(TEXT("FIsoTrianguler::LoopSegments after fix intersection"), true, false);
+	Grid.DisplayIsoSegments(TEXT("FIsoTrianguler::LoopSegments orientation"), DisplaySpace, LoopSegments, false, true, false, EVisuProperty::YellowCurve);
+	Grid.DisplayIsoSegments(TEXT("FIsoTrianguler::LoopSegments"), DisplaySpace, LoopSegments, true, false, false, EVisuProperty::YellowCurve);
+	Wait(false);
+#endif
+
+	//Fill Intersection tool
+	LoopSegmentsIntersectionTool.Empty(LoopSegments.Num());
+	LoopSegmentsIntersectionTool.AddSegments(LoopSegments);
+	LoopSegmentsIntersectionTool.Sort();
+
+	GetThinZonesMesh();
+
+	LoopSegmentsIntersectionTool.AddSegments(ThinZoneSegments);
+	LoopSegmentsIntersectionTool.Sort();
+
+	FinalToLoops.Append(ThinZoneSegments);
+
+#ifdef DEBUG_ISOTRIANGULATOR
+	Grid.DisplayIsoSegments(TEXT("FIsoTrianguler::ThinZoneSegments"), DisplaySpace, ThinZoneSegments, false);
+	LoopSegmentsIntersectionTool.Display(Grid.bDisplay, TEXT("FIsoTrianguler::IntersectionTool Loop"));
+#endif
+
+	BuildInnerSegments();
+
+#ifdef DEBUG_ISOTRIANGULATOR
+	if (bDisplay)
+	{
+		F3DDebugSession A(TEXT("FIsoTrianguler::FinalInnerSegments"));
+		Grid.DisplayIsoSegments(DisplaySpace, FinalInnerSegments, true, false, false, EVisuProperty::BlueCurve);
+		//Wait();
+	}
+	InnerToOuterIsoSegmentsIntersectionTool.Display(bDisplay, TEXT("FIsoTrianguler::IntersectionTool InnerToOuter"), EVisuProperty::RedCurve);
+#endif
+
+	// =============================================================================================================
+	// =============================================================================================================
+
+	BuildInnerSegmentsIntersectionTool();
+#ifdef DEBUG_ISOTRIANGULATOR
+	InnerSegmentsIntersectionTool.Display(Grid.bDisplay, TEXT("FIsoTrianguler::IntersectionTool Inner"));
+#endif
+
+	// =============================================================================================================
+	// 	   For each cell
+	// 	      - Connect loops together and to cell vertices
+	// 	           - Find subset of node of each loop
+	// 	           - build Delaunay connection
+	// 	           - find the shortest segment to connect each connected loop by Delaunay
+	// =============================================================================================================
+#ifdef DEBUG_ISOTRIANGULATOR
+	Grid.DisplayIsoSegments(TEXT("FIsoTrianguler::Final To Loops Before"), DisplaySpace, FinalToLoops, false, false, false, EVisuProperty::YellowCurve);
+#endif
+
+	ConnectCellLoops();
+
+#ifdef DEBUG_ISOTRIANGULATOR
+	Grid.DisplayIsoSegments(TEXT("FIsoTrianguler::Final Iso ToLink Inner To Loops"), DisplaySpace, FinalToLoops, false, false, true, EVisuProperty::YellowCurve);
+#endif
+
+	SelectSegmentsToLinkInnerToLoop();
+
+#ifdef DEBUG_ISOTRIANGULATOR
+	Grid.DisplayIsoSegments(TEXT("FIsoTrianguler::Final Segments"), DisplaySpace, FinalToLoops, false, false, true, EVisuProperty::YellowCurve);
+#endif
+
+	// =============================================================================================================
+	// Make the final tessellation 
+	// =============================================================================================================
+
+	// Triangulate between inner grid boundary and loops
+	TriangulateOverCycle(EGridSpace::Scaled);
+
+	// Finalize the mesh by the tessellation of the inner grid
+	TriangulateInnerNodes();
+
+#ifdef DEBUG_ISOTRIANGULATOR
+	if (bDisplay)
+	{
+		F3DDebugSession A(TEXT("Mesh 3D"));
+		DisplayMesh(Mesh);
+	}
+	//Wait(bDisplay);
+#endif
+
+#ifdef CADKERNEL_DEV
+	//Chronos.PrintTimeElapse();
+#endif
+
+	return true;
+}
+
+void FIsoTriangulator::BuildNodes()
+{
+	FTimePoint StartTime = FChrono::Now();
+
+	LoopNodeCount = 0;
+	for (const TArray<FPoint2D>& LoopPoints : Grid.GetLoops2D(EGridSpace::Default2D))
+	{
+		LoopNodeCount += (int32)LoopPoints.Num();
+	}
+	LoopStartIndex.Reserve(Grid.GetLoops2D(EGridSpace::Default2D).Num());
+	LoopNodes.Reserve((int32)(LoopNodeCount * 1.2 + 5)); // reserve more in case it need to create complementary nodes
+
+	// Loop nodes
+	int32 NodeIndex = 0;
+	int32 LoopIndex = 0;
+	for (const TArray<FPoint2D>& LoopPoints : Grid.GetLoops2D(EGridSpace::Default2D))
+	{
+		LoopStartIndex.Add(LoopNodeCount);
+		const TArray<int32>& LoopIds = Grid.GetNodeIdsOfFaceLoops()[LoopIndex];
+		FLoopNode* NextNode = nullptr;
+		FLoopNode* FirstNode = &LoopNodes.Emplace_GetRef(LoopIndex, 0, NodeIndex++, LoopIds[0]);
+		FLoopNode* PreviousNode = FirstNode;
+		for (int32 Index = 1; Index < LoopPoints.Num(); ++Index)
+		{
+			NextNode = &LoopNodes.Emplace_GetRef(LoopIndex, Index, NodeIndex++, LoopIds[Index]);
+			PreviousNode->SetNextConnectedNode(NextNode);
+			NextNode->SetPreviousConnectedNode(PreviousNode);
+			PreviousNode = NextNode;
+		}
+		PreviousNode->SetNextConnectedNode(FirstNode);
+		FirstNode->SetPreviousConnectedNode(PreviousNode);
+		LoopIndex++;
+	}
+
+	// Inner node
+	InnerNodes.Reserve(Grid.InnerNodesCount());
+	GlobalIndexToIsoInnerNodes.Init(nullptr, Grid.GetTotalCuttingCount());
+
+	InnerNodeCount = 0;
+	for (int32 Index = 0; Index < (int32)Grid.GetTotalCuttingCount(); ++Index)
+	{
+		if (Grid.IsNodeInsideAndMeshable(Index))
+		{
+			FIsoInnerNode& Node = InnerNodes.Emplace_GetRef(Index, NodeIndex++, InnerNodeCount++);
+			GlobalIndexToIsoInnerNodes[Index] = &Node;
+		}
+	}
+
+#ifdef CADKERNEL_DEV
+	Chronos.BuildIsoNodesDuration += FChrono::Elapse(StartTime);
+#endif
+}
+
+void FIsoTriangulator::FillMeshNodes()
+{
+	int32 TriangleNum = 50 + (int32)((2 * InnerNodeCount + LoopNodeCount) * 1.1);
+	Mesh.Init(TriangleNum, InnerNodeCount + LoopNodeCount);
+
+	TArray<FPoint>& InnerNodeCoordinates = Mesh.GetNodeCoordinates();
+	InnerNodeCoordinates.Reserve(InnerNodeCount);
+	for (int32 Index = 0; Index < (int32)Grid.GetInner3DPoints().Num(); ++Index)
+	{
+		if (Grid.IsNodeInsideAndMeshable(Index))
+		{
+			InnerNodeCoordinates.Emplace(Grid.GetInner3DPoints()[Index]);
+		}
+	}
+
+	int32 StartId = Mesh.RegisterCoordinates();
+	for (FIsoInnerNode& Node : InnerNodes)
+	{
+		Node.OffsetId(StartId);
+	}
+
+	Mesh.VerticesGlobalIndex.SetNum(InnerNodeCount + LoopNodeCount);
+	int32 Index = 0;
+	for (FLoopNode& Node : LoopNodes)
+	{
+		Mesh.VerticesGlobalIndex[Index++] = Node.GetNodeId();
+	}
+
+	for (FIsoInnerNode& Node : InnerNodes)
+	{
+		Mesh.VerticesGlobalIndex[Index++] = Node.GetNodeId();
+	}
+
+	for (FLoopNode& Node : LoopNodes)
+	{
+		Mesh.Normals.Emplace(Node.GetNormal(Grid));
+	}
+
+	for (FIsoInnerNode& Node : InnerNodes)
+	{
+		Mesh.Normals.Emplace(Node.GetNormal(Grid));
+	}
+
+	for (FLoopNode& Node : LoopNodes)
+	{
+		const FPoint2D& UVCoordinate = Node.Get2DPoint(EGridSpace::Scaled, Grid);
+		Mesh.UVMap.Emplace(UVCoordinate.U, UVCoordinate.V);
+	}
+
+	for (FIsoInnerNode& Node : InnerNodes)
+	{
+		const FPoint2D& UVCoordinate = Node.Get2DPoint(EGridSpace::Scaled, Grid);
+		Mesh.UVMap.Emplace(UVCoordinate.U, UVCoordinate.V);
+	}
+}
+
+void FIsoTriangulator::BuildLoopSegments()
+{
+	FTimePoint StartTime = FChrono::Now();
+
+	LoopSegments.Reserve(LoopNodeCount);
+
+	int32 LoopIndex = 0;
+	for (FLoopNode& Node : LoopNodes)
+	{
+		if (Node.IsDelete())
+		{
+			continue;
+		}
+
+		FIsoSegment& Segment = IsoSegmentFactory.New();
+		Segment.Init(Node, Node.GetNextNode(), ESegmentType::Loop);
+		if (Segment.ConnectToNode())
+		{
+			LoopSegments.Add(&Segment);
+		}
+		else
+		{
+#ifdef DEBUG_BUILD_SEGMENT_IF_NEEDED
+			F3DDebugSession _(FString::Printf(TEXT("ERROR Segment")));
+			Grid.DisplayIsoSegment(EGridSpace::UniformScaled, Node, Node.GetNextNode(), 0, EVisuProperty::RedCurve);
+			Wait();
+#endif
+			IsoSegmentFactory.DeleteEntity(&Segment);
+		}
+	}
+
+#ifdef CADKERNEL_DEV
+	Chronos.BuildLoopSegmentsDuration += FChrono::Elapse(StartTime);
+#endif
+}
+
+void FIsoTriangulator::GetThinZonesMesh()
+{
+	TMap<int32, FLoopNode*> IndexToNode;
+	for (FLoopNode& Node : LoopNodes)
+	{
+		if (Node.IsDelete())
+		{
+			continue;
+		}
+
+		IndexToNode.Add(Node.GetNodeId(), &Node);
+	}
+
+	{
+		for (const FThinZone2D& ThinZone : Grid.GetFace().GetThinZones())
+		{
+			GetThinZoneMesh(IndexToNode, ThinZone);
+		}
+	}
+
+	ThinZoneIntersectionTool.Empty(0);
+}
+
+void FIsoTriangulator::GetThinZoneMesh(const TMap<int32, FLoopNode*>& IndexToNode, const FThinZone2D& ThinZone)
+{
+	using namespace IsoTriangulatorImpl;
+
+	TArray<TPair<int32, FPairOfIndex>> CrossZoneElements;
+	TArray<FCandidateSegment> MeshOfThinZones;
+
+	FAddMeshNodeFunc AddElement = [&CrossZoneElements](const int32 NodeIndice, const FPoint2D& MeshNode2D, double MeshingTolerance3D, const FEdgeSegment& EdgeSegment, const FPairOfIndex& OppositeNodeIndices)
+	{
+		if (CrossZoneElements.Num() && CrossZoneElements.Last().Key == NodeIndice)
+		{
+			CrossZoneElements.Last().Value.Add(OppositeNodeIndices);
+		}
+		else
+		{
+			CrossZoneElements.Emplace(NodeIndice, OppositeNodeIndices);
+		}
+	};
+
+	FReserveContainerFunc Reserve = [&CrossZoneElements](int32 MeshVertexCount)
+	{
+		CrossZoneElements.Reserve(CrossZoneElements.Num() + MeshVertexCount);
+	};
+
+	ThinZone.GetFirstSide().GetExistingMeshNodes(Grid.GetFace(), Mesh.GetMeshModel(), Reserve, AddElement, /*bWithTolerance*/ false);
+	ThinZone.GetSecondSide().GetExistingMeshNodes(Grid.GetFace(), Mesh.GetMeshModel(), Reserve, AddElement, /*bWithTolerance*/ false);
+
+	MeshOfThinZones.Reserve(CrossZoneElements.Num() * 2);
+
+	TFunction<void(FLoopNode*, FLoopNode*)> AddSegmentFromNode = [&MeshOfThinZones, this](FLoopNode* NodeA, FLoopNode* NodeB)
+	{
+		if (!NodeA)
+		{
+			return;
+		}
+
+		if (!NodeB)
+		{
+			return;
+		}
+
+		if (&NodeA->GetPreviousNode() == NodeB || &NodeB->GetNextNode() == NodeA)
+		{
+			return;
+		}
+
+		if (NodeA->GetSegmentConnectedTo(NodeB))
+		{
+			return;
+		}
+
+		const FPoint2D& CoordinateA = NodeA->Get2DPoint(EGridSpace::UniformScaled, Grid);
+		const FPoint2D& CoordinateB = NodeB->Get2DPoint(EGridSpace::UniformScaled, Grid);
+
+		// Is Outside and not too flat at Node1
+		const double FlatAngle = 0.1;
+		if (NodeA->IsSegmentBeInsideFace(CoordinateB, Grid, FlatAngle))
+		{
+			return;
+		}
+
+		// Is Outside and not too flat at Node2
+		if (NodeB->IsSegmentBeInsideFace(CoordinateA, Grid, FlatAngle))
+		{
+			return;
+		}
+
+		MeshOfThinZones.Emplace(Grid, *NodeA, *NodeB);
+	};
+
+	TFunction<void(int32, int32)> AddSegment = [&IndexToNode, AddSegmentFromNode](int32 IndexNodeA, int32 IndexNodeB)
+	{
+		if (IndexNodeA < 0 || IndexNodeB < 0)
+		{
+			return;
+		}
+
+		if (IndexNodeA == IndexNodeB)
+		{
+			return;
+		}
+
+		FLoopNode* const* NodeA = IndexToNode.Find(IndexNodeA);
+		FLoopNode* const* NodeB = IndexToNode.Find(IndexNodeB);
+		if (NodeA && NodeB)
+		{
+			AddSegmentFromNode(*NodeA, *NodeB);
+		}
+	};
+
+	for (const TPair<int32, FPairOfIndex>& CrossZoneElement : CrossZoneElements)
+	{
+		AddSegment(CrossZoneElement.Key, CrossZoneElement.Value[0]);
+		AddSegment(CrossZoneElement.Key, CrossZoneElement.Value[1]);
+	}
+
+	Algo::Sort(MeshOfThinZones, [](const FCandidateSegment& SegmentA, const FCandidateSegment& SegmentB) { return SegmentA.Length < SegmentB.Length; });
+
+	ThinZoneIntersectionTool.Reserve(ThinZoneIntersectionTool.Count() + MeshOfThinZones.Num());
+
+	for (FCandidateSegment& CandidateSegment : MeshOfThinZones)
+	{
+		if (FIsoSegment::IsItAlreadyDefined(&CandidateSegment.StartNode, &CandidateSegment.EndNode))
+		{
+			continue;
+		}
+
+		if (LoopSegmentsIntersectionTool.DoesIntersect(CandidateSegment.StartNode, CandidateSegment.EndNode))
+		{
+			continue;
+		}
+
+		if (ThinZoneIntersectionTool.DoesIntersect(CandidateSegment.StartNode, CandidateSegment.EndNode))
+		{
+			continue;
+		}
+
+		FIsoSegment& Segment = IsoSegmentFactory.New();
+		Segment.Init(CandidateSegment.StartNode, CandidateSegment.EndNode, ESegmentType::ThinZone);
+
+		if(Segment.ConnectToNode())
+		{
+			CandidateSegment.StartNode.SetThinZoneNodeMarker();
+			CandidateSegment.EndNode.SetThinZoneNodeMarker();
+			Segment.SetFinalMarker();
+			ThinZoneSegments.Add(&Segment);
+			ThinZoneIntersectionTool.AddSegment(Segment);
+		}
+		else
+		{
+#ifdef DEBUG_BUILD_SEGMENT_IF_NEEDED
+			F3DDebugSession _(FString::Printf(TEXT("ERROR Segment")));
+			Grid.DisplayIsoSegment(EGridSpace::UniformScaled, CandidateSegment.StartNode, CandidateSegment.EndNode, 0, EVisuProperty::RedCurve);
+			Wait();
+#endif
+			IsoSegmentFactory.DeleteEntity(&Segment);
+		}
+	}
+}
+
+void FIsoTriangulator::BuildInnerSegments()
+{
+#ifdef DEBUG_BUILDINNERSEGMENTS
+	if (bDisplay)
+	{
+		Grid.DisplayInnerPoints(TEXT("BuildInnerSegments::Points"), EGridSpace::UniformScaled);
+	}
+	F3DDebugSession _(bDisplay, TEXT("BuildInnerSegments"));
+#endif
+
+	FTimePoint StartTime = FChrono::Now();
+
+	// Build segments according to the Grid following u then following v
+	// Build segment must not be in intersection with the loop
+	int32 NumU = Grid.GetCuttingCount(EIso::IsoU);
+	int32 NumV = Grid.GetCuttingCount(EIso::IsoV);
+
+	LoopSegmentsIntersectionTool.Reserve(InnerSegmentsIntersectionTool.Count());
+
+
+	// Check if the loop tangents the grid between Node1 and Node 2
+	//                            
+	//                       \   /  Loop		                       \     /  Loop
+	//                        \./ 				                        \   / 
+	//        Node1 *------------------* Node2 	        Node1 *----------\./-------* Node2 
+	//                                                                       
+	//
+	TFunction<bool(const FPoint2D&, const FPoint2D&, const ESegmentType, const double)> AlmostHitsLoop = [&](const FPoint2D& Node1, const FPoint2D& Node2, const ESegmentType InType, const double Tolerance) -> bool
+	{
+		if (InType == ESegmentType::IsoV)
+		{
+			for (const TArray<FPoint2D>& Loop : Grid.GetLoops2D(EGridSpace::UniformScaled))
+			{
+				for (const FPoint2D& LoopPoint : Loop)
+				{
+					if (FMath::IsNearlyEqual(LoopPoint.V, Node1.V, Tolerance))
+					{
+						if (Node1.U - DOUBLE_SMALL_NUMBER < LoopPoint.U && LoopPoint.U < Node2.U + DOUBLE_SMALL_NUMBER)
+						{
+#ifdef DEBUG_BUILDINNERSEGMENTS
+							if (bDisplay)
+							{
+								F3DDebugSession _(bDisplay, FString::Printf(TEXT("Point")));
+								DisplayPoint2DWithScale(LoopPoint, EVisuProperty::RedPoint);
+								Wait();
+							}
+#endif
+							return true;
+						}
+					}
+				}
+			}
+		}
+		else
+		{
+			for (const TArray<FPoint2D>& Loop : Grid.GetLoops2D(EGridSpace::UniformScaled))
+			{
+				for (const FPoint2D& LoopPoint : Loop)
+				{
+					if (FMath::IsNearlyEqual(LoopPoint.U, Node1.U, Tolerance))
+					{
+						if (Node1.V - DOUBLE_SMALL_NUMBER < LoopPoint.V && LoopPoint.V < Node2.V + DOUBLE_SMALL_NUMBER)
+						{
+#ifdef DEBUG_BUILDINNERSEGMENTS
+							if (bDisplay)
+							{
+								//F3DDebugSession _(bDisplay, FString::Printf(TEXT("Point")));
+								DisplayPoint2DWithScale(LoopPoint, EVisuProperty::RedPoint);
+								Wait();
+							}
+#endif
+							return true;
+						}
+					}
+				}
+			}
+		}
+		return false;
+	};
+
+	TFunction<void(const int32, const int32, const ESegmentType)> AddToInnerToOuterSegmentsIntersectionTool = [&](const int32 IndexNode1, const int32 IndexNode2, const ESegmentType InType)
+	{
+		const FPoint2D& Point1 = Grid.GetInner2DPoint(EGridSpace::UniformScaled, IndexNode1);
+		const FPoint2D& Point2 = Grid.GetInner2DPoint(EGridSpace::UniformScaled, IndexNode2);
+
+		InnerToOuterIsoSegmentsIntersectionTool.AddIsoSegment(Point1, Point2, InType);
+	};
+
+	TFunction<void(const int32, const int32, const ESegmentType)> AddToInnerSegments = [&](const int32 IndexNode1, const int32 IndexNode2, const ESegmentType InType)
+	{
+		FIsoInnerNode& Node1 = *GlobalIndexToIsoInnerNodes[IndexNode1];
+		FIsoInnerNode& Node2 = *GlobalIndexToIsoInnerNodes[IndexNode2];
+		FIsoSegment& Segment = IsoSegmentFactory.New();
+		Segment.Init(Node1, Node2, InType);
+		if (Segment.ConnectToNode())
+		{
+			FinalInnerSegments.Add(&Segment);
+		}
+		else
+		{
+#ifdef DEBUG_BUILD_SEGMENT_IF_NEEDED
+			F3DDebugSession _(FString::Printf(TEXT("ERROR Segment")));
+			Grid.DisplayIsoSegment(EGridSpace::UniformScaled, Node1, Node2, 0, EVisuProperty::RedCurve);
+			Wait();
+#endif
+			IsoSegmentFactory.DeleteEntity(&Segment);
+		}
+	};
+
+	TFunction<void(const int32, const int32, const ESegmentType, const double)> BuildSegmentIfValid = [&](const int32 IndexNode1, const int32 IndexNode2, const ESegmentType InType, const double Tolerance)
+	{
+		if (Grid.IsNodeOusideFaceButClose(IndexNode1) && Grid.IsNodeOusideFaceButClose(IndexNode2))
+		{
+#ifdef DEBUG_BUILDINNERSEGMENTS
+			if (bDisplay)
+			{
+				DisplaySegmentWithScale(Grid.GetInner2DPoint(EGridSpace::UniformScaled, IndexNode1), Grid.GetInner2DPoint(EGridSpace::UniformScaled, IndexNode2), IndexNode1, EVisuProperty::PinkCurve);
+			}
+#endif
+			AddToInnerToOuterSegmentsIntersectionTool(IndexNode1, IndexNode2, InType);
+			return;
+		}
+
+		if (Grid.IsNodeOutsideFace(IndexNode1) && Grid.IsNodeOutsideFace(IndexNode2))
+		{
+#ifdef DEBUG_BUILDINNERSEGMENTS
+			if (bDisplay)
+			{
+				DisplaySegmentWithScale(Grid.GetInner2DPoint(EGridSpace::UniformScaled, IndexNode1), Grid.GetInner2DPoint(EGridSpace::UniformScaled, IndexNode2), IndexNode1, EVisuProperty::RedCurve);
+			}
+#endif
+			return;
+		}
+
+		if (Grid.IsNodeInsideAndCloseToLoop(IndexNode1) && Grid.IsNodeInsideAndCloseToLoop(IndexNode2))
+		{
+			if (LoopSegmentsIntersectionTool.DoesIntersect(Grid.GetInner2DPoint(EGridSpace::UniformScaled, IndexNode1), Grid.GetInner2DPoint(EGridSpace::UniformScaled, IndexNode2))
+				|| AlmostHitsLoop(Grid.GetInner2DPoint(EGridSpace::UniformScaled, IndexNode1), Grid.GetInner2DPoint(EGridSpace::UniformScaled, IndexNode2), InType, Tolerance))
+			{
+#ifdef DEBUG_BUILDINNERSEGMENTS
+				if (bDisplay)
+				{
+					DisplaySegmentWithScale(Grid.GetInner2DPoint(EGridSpace::UniformScaled, IndexNode1), Grid.GetInner2DPoint(EGridSpace::UniformScaled, IndexNode2), IndexNode1, EVisuProperty::PurpleCurve);
+					Wait();
+				}
+#endif
+				AddToInnerToOuterSegmentsIntersectionTool(IndexNode1, IndexNode2, InType);
+			}
+			else
+			{
+#ifdef DEBUG_BUILDINNERSEGMENTS
+				if (bDisplay)
+				{
+					DisplaySegmentWithScale(Grid.GetInner2DPoint(EGridSpace::UniformScaled, IndexNode1), Grid.GetInner2DPoint(EGridSpace::UniformScaled, IndexNode2), IndexNode1, EVisuProperty::OrangeCurve);
+				}
+#endif
+				AddToInnerSegments(IndexNode1, IndexNode2, InType);
+			}
+
+			return;
+		}
+
+		if (Grid.IsNodeInsideAndMeshable(IndexNode1) && Grid.IsNodeInsideAndMeshable(IndexNode2))
+		{
+#ifdef DEBUG_BUILDINNERSEGMENTS
+			if (bDisplay)
+			{
+				DisplaySegmentWithScale(Grid.GetInner2DPoint(EGridSpace::UniformScaled, IndexNode1), Grid.GetInner2DPoint(EGridSpace::UniformScaled, IndexNode2), IndexNode1, EVisuProperty::BlueCurve);
+			}
+#endif
+			AddToInnerSegments(IndexNode1, IndexNode2, InType);
+			return;
+		}
+
+		if (Grid.IsNodeInsideButTooCloseToLoop(IndexNode1) && Grid.IsNodeInsideButTooCloseToLoop(IndexNode2))
+		{
+#ifdef DEBUG_BUILDINNERSEGMENTS
+			if (bDisplay)
+			{
+				DisplaySegmentWithScale(Grid.GetInner2DPoint(EGridSpace::UniformScaled, IndexNode1), Grid.GetInner2DPoint(EGridSpace::UniformScaled, IndexNode2), IndexNode1, EVisuProperty::GreenCurve);
+			}
+#endif
+			return;
+		}
+
+#ifdef DEBUG_BUILDINNERSEGMENTS
+		if (bDisplay)
+		{
+			DisplaySegmentWithScale(Grid.GetInner2DPoint(EGridSpace::UniformScaled, IndexNode1), Grid.GetInner2DPoint(EGridSpace::UniformScaled, IndexNode2), IndexNode1, EVisuProperty::YellowCurve);
+		}
+#endif
+		AddToInnerToOuterSegmentsIntersectionTool(IndexNode1, IndexNode2, InType);
+	};
+
+	TFunction<const TArray<double>(const TArray<double>&)> ComputeLocalTolerance = [](const TArray<double>& UniformCutting) -> const TArray<double>
+	{
+		const int32 Num = UniformCutting.Num();
+		TArray<double> TolerancesAlongU;
+		TArray<double> Temp;
+		TolerancesAlongU.Reserve(Num);
+		Temp.Reserve(Num);
+		for (int32 Index = 1; Index < Num; Index++)
+		{
+			Temp.Add((UniformCutting[Index] - UniformCutting[Index - 1]));
+		}
+		TolerancesAlongU.Add(Temp[0] * 0.1);
+		for (int32 Index = 1; Index < Num - 1; Index++)
+		{
+			TolerancesAlongU.Add((Temp[Index] + Temp[Index - 1]) * 0.05);
+		}
+		TolerancesAlongU.Add(Temp.Last() * 0.1);
+		return MoveTemp(TolerancesAlongU);
+	};
+
+	// Process along V
+	{
+		TArray<double> TolerancesAlongU = ComputeLocalTolerance(Grid.GetUniformCuttingCoordinatesAlongIso(EIso::IsoU));
+		for (int32 UIndex = 0; UIndex < NumU; UIndex++)
+		{
+			for (int32 VIndex = 0; VIndex < NumV - 1; VIndex++)
+			{
+				BuildSegmentIfValid(Grid.GobalIndex(UIndex, VIndex), Grid.GobalIndex(UIndex, VIndex + 1), ESegmentType::IsoU, TolerancesAlongU[UIndex]);
+			}
+		}
+	}
+
+	// Process along U
+	{
+		TArray<double> TolerancesAlongV = ComputeLocalTolerance(Grid.GetUniformCuttingCoordinatesAlongIso(EIso::IsoV));
+		for (int32 VIndex = 0; VIndex < NumV; VIndex++)
+		{
+			for (int32 UIndex = 0; UIndex < NumU - 1; UIndex++)
+			{
+				BuildSegmentIfValid(Grid.GobalIndex(UIndex, VIndex), Grid.GobalIndex(UIndex + 1, VIndex), ESegmentType::IsoV, TolerancesAlongV[VIndex]);
+			}
+		}
+	}
+
+#ifdef DEBUG_BUILDINNERSEGMENTS
+	Wait(bDisplay);
+#endif
+
+	InnerToOuterIsoSegmentsIntersectionTool.Sort();
+
+#ifdef CADKERNEL_DEV
+	Chronos.BuildInnerSegmentsDuration += FChrono::Elapse(StartTime);
+#endif
+}
+
+void FIsoTriangulator::BuildInnerSegmentsIntersectionTool()
+{
+	FTimePoint StartTime = FChrono::Now();
+
+	// Find Boundary Segments Of Inner Triangulation
+	// 
+	// A pixel grid is build. 
+	// A pixel is the quadrangle of the inner grid
+	// The grid pixel are initialized to False
+	//
+	// A pixel is True if one of its boundary segment does not exist 
+	// The inner of the grid is all pixel False
+	// The boundary of the inner triangulation is defined by all segments adjacent to different cells 
+	// 
+	//    T      T	     T                                                  
+	//       0 ----- 0 												     0 ----- 0 
+	//    T  |   F   |   T       T       T      T         			     |       |    
+	//       0 ----- 0               0 ----- 0 						     0       0               0 ----- 0 
+	//    T  |   F   |   T       T   |   F   |  T					     |       |               |       |  
+	//       0 ----- 0 ----- 0 ----- 0 ----- 0 						     0       0 ----- 0 ----- 0       0 	
+	//    T  |   F   |   F   |   F   |   F   |	T					     |                               |	
+	//       0 ----- 0 ----- 0 ----- 0 ----- 0 						     0                               0 	
+	//    T  |   F   |   F   |   F   |   F   |	T					     |                               |	
+	//       0 ----- 0 ----- 0 ----- 0 ----- 0 						     0 ----- 0 ----- 0 ----- 0 ----- 0 	
+	//    T      T 		 T		 T		 T		T					                
+	// 
+	// https://docs.google.com/presentation/d/1qUVOH-2kU_QXBVKyRUcdDy1Y6WGkcaJCiaS8wGjSZ6M/edit?usp=sharing
+	// Slide "Boundary Segments Of Inner Triangulation"
+
+	int32 NumU = Grid.GetCuttingCount(EIso::IsoU);
+	int32 NumV = Grid.GetCuttingCount(EIso::IsoV);
+
+	TArray<uint8> Pixel;
+	Pixel.Init(0, Grid.GetTotalCuttingCount());
+
+#ifdef DEBUG_FINDBOUNDARYSEGMENTS
+	if (bDisplay)
+	{
+		F3DDebugSession _(TEXT("FIsoTrianguler::InnerSegment"));
+		for (int32 IndexV = 0, Index = 0; IndexV < NumV; ++IndexV)
+		{
+			for (int32 IndexU = 0; IndexU < NumU; ++IndexU, ++Index)
+			{
+				if (!Grid.IsNodeInsideAndMeshable(Index))
+				{
+					continue;
+				}
+
+				if (GlobalIndexToIsoInnerNodes[Index]->IsLinkedToNextU())
+				{
+					Grid.DisplayIsoSegment(EGridSpace::UniformScaled, *GlobalIndexToIsoInnerNodes[Index], *GlobalIndexToIsoInnerNodes[Index + 1], Index, EVisuProperty::YellowCurve);
+				}
+
+				if (GlobalIndexToIsoInnerNodes[Index]->IsLinkedToPreviousU())
+				{
+					Grid.DisplayIsoSegment(EGridSpace::UniformScaled, *GlobalIndexToIsoInnerNodes[Index], *GlobalIndexToIsoInnerNodes[Index - 1], Index, EVisuProperty::YellowCurve);
+				}
+
+				if (GlobalIndexToIsoInnerNodes[Index]->IsLinkedToNextV())
+				{
+					Grid.DisplayIsoSegment(EGridSpace::UniformScaled, *GlobalIndexToIsoInnerNodes[Index], *GlobalIndexToIsoInnerNodes[Index + NumU], Index, EVisuProperty::YellowCurve);
+				}
+
+				if (GlobalIndexToIsoInnerNodes[Index]->IsLinkedToPreviousV())
+				{
+					Grid.DisplayIsoSegment(EGridSpace::UniformScaled, *GlobalIndexToIsoInnerNodes[Index], *GlobalIndexToIsoInnerNodes[Index - NumU], Index, EVisuProperty::YellowCurve);
+				}
+			}
+		}
+	}
+#endif
+
+	// A pixel is True if one of its boundary segment does not exist 
+	for (int32 IndexV = 0, Index = 0; IndexV < NumV; ++IndexV)
+	{
+		for (int32 IndexU = 0; IndexU < NumU; ++IndexU, ++Index)
+		{
+			if (!Grid.IsNodeInsideAndMeshable(Index))
+			{
+				continue;
+			}
+
+			if (!GlobalIndexToIsoInnerNodes[Index]->IsLinkedToNextU())
+			{
+				Pixel[Index] = true;
+				Pixel[Index - NumU] = true;
+			}
+
+			if (!GlobalIndexToIsoInnerNodes[Index]->IsLinkedToPreviousU())
+			{
+				Pixel[Index - 1] = true;
+				Pixel[Index - 1 - NumU] = true;
+			}
+
+			if (!GlobalIndexToIsoInnerNodes[Index]->IsLinkedToNextV())
+			{
+				Pixel[Index] = true;
+				Pixel[Index - 1] = true;
+			}
+
+			if (!GlobalIndexToIsoInnerNodes[Index]->IsLinkedToPreviousV())
+			{
+				Pixel[Index - NumU] = true;
+				Pixel[Index - NumU - 1] = true;
+			}
+		}
+	}
+
+	// The boundary of the inner triangulation is defined by all segments adjacent to a "True" cell 
+	// These segments are added to InnerSegmentsIntersectionTool
+	InnerSegmentsIntersectionTool.Reserve((int32)FinalInnerSegments.Num());
+
+	for (FIsoSegment* Segment : FinalInnerSegments)
+	{
+		int32 IndexFirstNode = Segment->GetFirstNode().GetIndex();
+		int32 IndexSecondNode = 0;
+		switch (Segment->GetType())
+		{
+		case ESegmentType::IsoU:
+			IndexSecondNode = IndexFirstNode - NumU;
+			break;
+		case ESegmentType::IsoV:
+			IndexSecondNode = IndexFirstNode - 1;
+			break;
+		default:
+			ensureCADKernel(false);
+		}
+		if (Pixel[IndexFirstNode] || Pixel[IndexSecondNode])
+		{
+			InnerSegmentsIntersectionTool.AddSegment(*Segment);
+		}
+	}
+
+	FindInnerGridCellSurroundingSmallLoop();
+
+	// initialize the intersection tool
+	InnerSegmentsIntersectionTool.Sort();
+
+#ifdef CADKERNEL_DEV
+	Chronos.FindLoopSegmentOfInnerTriangulationDuration += FChrono::Elapse(StartTime);
+#endif
+
+#ifdef DEBUG_FINDBOUNDARYSEGMENTS
+	DisplayPixels(Pixel);
+	//Wait(bDisplay);
+#endif
+}
+
+// =============================================================================================================
+// 	   For each cell
+// 	      - Connect loops together and to cell vertices
+// 	           - Find subset of node of each loop
+// 	           - build Delaunay connection
+// 	           - find the shortest segment to connect each connected loop by Delaunay
+// =============================================================================================================
+void FIsoTriangulator::ConnectCellLoops()
+{
+	TArray<FCell> Cells;
+	FindCellContainingBoundaryNodes(Cells);
+
+#ifdef DEBUG_CONNECT_CELL_LOOPS
+	if (bDisplay)
+	{
+		for (FCell& Cell : Cells)
+		{
+			F3DDebugSession _(FString::Printf(TEXT("Cell %d"), Cell.Id));
+			DisplayCell(Cell);
+		}
+	}
+#endif
+
+	for (FCell& Cell : Cells)
+	{
+		InitCellCorners(Cell);
+		Cell.InitLoopConnexions();
+	}
+
+	InnerToLoopCandidateSegments.Reserve(Cells.Num() * 2);
+
+#ifdef DEBUG_CONNECT_CELL_LOOPS
+	F3DDebugSession _(bDisplay, ("ConnectCellLoops"));
+#endif
+
+	FinalToLoops.Reserve(LoopNodeCount + InnerNodeCount);
+	for (FCell& Cell : Cells)
+	{
+#ifdef DEBUG_CONNECT_CELL_LOOPS
+		F3DDebugSession _(bDisplay, FString::Printf(TEXT("Cell %d"), Cell.Id));
+		DisplayCell(Cell);
+#endif
+
+		if (Cell.CellLoops.Num())
+		{
+			Cell.FindCandidateToConnectLoopsByNeighborhood();
+
+#ifdef DEBUG_CONNECT_CELL_CORNER_TO_INNER_LOOP
+			if (bDisplay)
+			{
+				DisplayCellConnexions(FString::Printf(TEXT("Cell %d CandidateSegments"), Cell.Id), Cell.LoopConnexions, EVisuProperty::YellowCurve);
+			}
+#endif
+
+			FindCandidateToConnectCellCornerToLoops(Cell);
+
+#ifdef DEBUG_CONNECT_CELL_CORNER_TO_INNER_LOOP
+			if (bDisplay)
+			{
+				DisplayCellConnexions(FString::Printf(TEXT("Cell %d CandidateSegments"), Cell.Id), Cell.LoopConnexions, EVisuProperty::BlueCurve);
+			}
+#endif
+
+			Cell.SelectSegmentToConnectLoops(IsoSegmentFactory);
+			Cell.SelectSegmentToConnectLoopToCorner(IsoSegmentFactory);
+			Cell.CheckAllLoopsConnectedTogetherAndConnect();
+#ifdef DEBUG_CONNECT_CELL_LOOPS
+			Grid.DisplayIsoSegments(TEXT("ConnectCellLoops::Step 0"), EGridSpace::UniformScaled, Cell.FinalSegments, false, false, false, EVisuProperty::YellowCurve);
+#endif
+		}
+
+#ifdef DEBUG_CONNECT_CELL_LOOPS
+		Grid.DisplayIsoSegments(TEXT("ConnectCellLoops::Step 1"), EGridSpace::UniformScaled, Cell.FinalSegments, false, false, false, EVisuProperty::YellowCurve);
+#endif
+
+		FinalToLoops.Append(Cell.FinalSegments);
+
+#ifdef DEBUG_CONNECT_CELL_LOOPS
+		Grid.DisplayIsoSegments(TEXT("ConnectCellLoops::Final Iso ToLink Inner To Loops"), EGridSpace::UniformScaled, FinalToLoops, false, false, false, EVisuProperty::YellowCurve);
+#endif
+	}
+}
+
+void FIsoTriangulator::FindCellContainingBoundaryNodes(TArray<FCell>& Cells)
+{
+	FTimePoint StartTime = FChrono::Now();
+
+	TArray<int32> NodeToCellIndices;
+	TArray<int32> SortedIndex;
+
+	const int32 CountU = Grid.GetCuttingCount(EIso::IsoU);
+	const int32 CountV = Grid.GetCuttingCount(EIso::IsoV);
+	const int32 MaxUV = Grid.GetTotalCuttingCount();
+
+	const TArray<double>& IsoUCoordinates = Grid.GetUniformCuttingCoordinatesAlongIso(EIso::IsoU);
+	const TArray<double>& IsoVCoordinates = Grid.GetUniformCuttingCoordinatesAlongIso(EIso::IsoV);
+
+	NodeToCellIndices.Reserve(LoopNodeCount);
+	{
+		int32 IndexU = 0;
+		int32 IndexV = 0;
+		int32 Index = 0;
+		int32 DeletedNodeCount = 0;
+
+#ifdef DEBUG_BUILD_CELLS
+		F3DDebugSession _(bDisplay, TEXT("Debug Find index"));
+#endif
+
+		for (const FLoopNode& LoopPoint : LoopNodes)
+		{
+			if (!LoopPoint.IsDelete())
+			{
+				const FPoint2D& Coordinate = LoopPoint.Get2DPoint(EGridSpace::UniformScaled, Grid);
+				ArrayUtils::FindCoordinateIndex(IsoUCoordinates, Coordinate.U, IndexU);
+				ArrayUtils::FindCoordinateIndex(IsoVCoordinates, Coordinate.V, IndexV);
+
+#ifdef DEBUG_BUILD_CELLS
+				if (bDisplay)
+				{
+					F3DDebugSession S(FString::Printf(TEXT("Index %d"), IndexV * CountU + IndexU));
+					DisplayPoint2DWithScale(LoopPoint.GetPoint(EGridSpace::UniformScaled, Grid), EVisuProperty::BluePoint);
+					DrawCellBoundary(IndexV * CountU + IndexU, EVisuProperty::YellowPoint);
+				}
+#endif
+
+				NodeToCellIndices.Emplace(IndexV * CountU + IndexU);
+			}
+			else
+			{
+#ifdef DEBUG_BUILD_CELLS
+				if (bDisplay)
+				{
+					F3DDebugSession S(TEXT("Deleted Node"));
+					DisplayPoint2DWithScale(LoopPoint.GetPoint(EGridSpace::UniformScaled, Grid), EVisuProperty::RedPoint);
+				}
+#endif
+
+				DeletedNodeCount++;
+				NodeToCellIndices.Emplace(MaxUV);
+			}
+			SortedIndex.Emplace(Index++);
+		}
+
+		Algo::Sort(SortedIndex, [&](const int32& Index1, const int32& Index2)
+			{
+				return (NodeToCellIndices[Index1] < NodeToCellIndices[Index2]);
+			});
+
+		SortedIndex.SetNum(SortedIndex.Num() - DeletedNodeCount);
+	}
+
+	int32 CountOfCellsFilled = 1;
+	{
+		int32 CellIndex = NodeToCellIndices[0];
+		for (int32 Index : SortedIndex)
+		{
+			if (CellIndex != NodeToCellIndices[Index])
+			{
+				CellIndex = NodeToCellIndices[Index];
+				CountOfCellsFilled++;
+			}
+		}
+	}
+
+	// build Cells
+	{
+		Cells.Reserve(CountOfCellsFilled);
+		int32 CellIndex = NodeToCellIndices[SortedIndex[0]];
+		TArray<FLoopNode*> CellNodes;
+		CellNodes.Reserve(LoopNodeCount);
+
+#ifdef DEBUG_BUILD_CELLS
+		F3DDebugSession _(bDisplay, TEXT("Build Cells"));
+#endif
+
+		for (int32 Index : SortedIndex)
+		{
+			if (CellIndex != NodeToCellIndices[Index])
+			{
+				Cells.Emplace(CellIndex, CellNodes, *this);
+
+				CellIndex = NodeToCellIndices[Index];
+				CellNodes.Reset(LoopNodeCount);
+
+#ifdef DEBUG_BUILD_CELLS
+				if(bDisplay)
+				{
+					DisplayCell(Cells.Last());
+					Wait();
+				}
+#endif
+			}
+
+			FLoopNode& LoopNode = LoopNodes[Index];
+			if (!LoopNode.IsDelete())
+			{
+				CellNodes.Add(&LoopNode);
+			}
+		}
+		Cells.Emplace(CellIndex, CellNodes, *this);
+
+#ifdef DEBUG_BUILD_CELLS
+		if (bDisplay)
+		{
+			DisplayCell(Cells.Last());
+			Wait();
+		}
+#endif
+	}
+
+	FChrono::Elapse(StartTime);
+}
+
+bool FIsoTriangulator::CanCycleBeMeshed(const TArray<FIsoSegment*>& Cycle, FIntersectionSegmentTool& CycleIntersectionTool)
+{
+	bool bHasSelfIntersection = true;
+
+	for (const FIsoSegment* Segment : Cycle)
+	{
+		if (CycleIntersectionTool.DoesIntersect(*Segment))
+		{
+			FMessage::Printf(Log, TEXT("A cycle of the surface %d is in self intersecting. The mesh of this sector is canceled.\n"), Grid.GetFace().GetId());
+
+#ifdef DEBUG_CAN_CYCLE_BE_FIXED_AND_MESHED
+			if (bDisplay)
+			{
+				F3DDebugSession _(TEXT("New Segment"));
+				F3DDebugSession A(TEXT("Segments"));
+				Display(EGridSpace::UniformScaled, *Segment, 0, EVisuProperty::RedPoint);
+				Display(EGridSpace::UniformScaled, *IntersectingSegment, 0, EVisuProperty::BluePoint);
+				Wait();
+			}
+#endif
+			return false;
+		}
+	}
+
+	return true;
+}
+
+void FIsoTriangulator::MeshCycle(const TArray<FIsoSegment*>& Cycle, const TArray<bool>& CycleOrientation)
+{
+	switch (Cycle.Num())
+	{
+	case 2:
+		return;
+	case 3:
+		return MeshCycleOf<3>(Cycle, CycleOrientation, Polygon::MeshTriangle);
+	case 4:
+		return MeshCycleOf<4>(Cycle, CycleOrientation, Polygon::MeshQuadrilateral);
+	case 5:
+		return MeshCycleOf<5>(Cycle, CycleOrientation, Polygon::MeshPentagon);
+	default:
+		MeshLargeCycle(Cycle, CycleOrientation);
+	}
+}
+
+void FIsoTriangulator::MeshLargeCycle(const TArray<FIsoSegment*>& Cycle, const TArray<bool>& CycleOrientation)
+{
+	FCycleTriangulator CycleTriangulator(*this, Cycle, CycleOrientation);
+	CycleTriangulator.MeshCycle();
+}
+
+/**
+ * The purpose is to add surrounding segments to the small loop to intersection tool to prevent traversing inner segments
+ * A loop is inside inner segments
+ *									|			 |
+ *								   -----------------
+ *									|	 XXX	 |
+ *									|	XXXXX	 |
+ *									|	 XXX	 |
+ *								   -----------------
+ *									|			 |
+ *
+ */
+void FIsoTriangulator::FindInnerGridCellSurroundingSmallLoop()
+{
+	FTimePoint StartTime = FChrono::Now();
+
+	if (GlobalIndexToIsoInnerNodes.Num() == 0)
+	{
+		// No inner node
+		return;
+	}
+
+	// when an internal loop is inside inner UV cell
+	int32 NumU = Grid.GetCuttingCount(EIso::IsoU);
+	int32 NumV = Grid.GetCuttingCount(EIso::IsoV);
+	const TArray<double>& UCoordinates = Grid.GetCuttingCoordinatesAlongIso(EIso::IsoU);
+	const TArray<double>& VCoordinates = Grid.GetCuttingCoordinatesAlongIso(EIso::IsoV);
+
+	const TArray<TArray<FPoint2D>>& Loops = Grid.GetLoops2D(EGridSpace::Default2D);
+	for (int32 LoopIndex = 1; LoopIndex < Loops.Num(); ++LoopIndex)
+	{
+		FPoint2D FirstPoint = Loops[LoopIndex][0];
+
+		int32 IndexU = 0;
+		for (; IndexU < NumU - 1; ++IndexU)
+		{
+			if ((FirstPoint.U > UCoordinates[IndexU]) && (FirstPoint.U < UCoordinates[IndexU + 1] + DOUBLE_SMALL_NUMBER))
+			{
+				break;
+			}
+		}
+
+		int32 IndexV = 0;
+		for (; IndexV < NumV - 1; ++IndexV)
+		{
+			if ((FirstPoint.V > VCoordinates[IndexV]) && (FirstPoint.V < VCoordinates[IndexV + 1] + DOUBLE_SMALL_NUMBER))
+			{
+				break;
+			}
+		}
+
+		double UMin = UCoordinates[IndexU];
+		double UMax = UCoordinates[IndexU + 1] + DOUBLE_SMALL_NUMBER;
+		double VMin = VCoordinates[IndexV];
+		double VMax = VCoordinates[IndexV + 1] + DOUBLE_SMALL_NUMBER;
+
+		bool bBoudardyIsSurrounded = true;
+		for (const FPoint2D& LoopPoint : Loops[LoopIndex])
+		{
+			if (LoopPoint.U < UMin || LoopPoint.U > UMax || LoopPoint.V < VMin || LoopPoint.V > VMax)
+			{
+				bBoudardyIsSurrounded = false;
+				break;
+			}
+		}
+
+		if (bBoudardyIsSurrounded)
+		{
+			int32 Index = IndexV * NumU + IndexU;
+			IndexOfLowerLeftInnerNodeSurroundingALoop.Add((int32)Index);
+
+			FIsoInnerNode* Node = GlobalIndexToIsoInnerNodes[Index];
+			if (Node == nullptr)
+			{
+				Node = GlobalIndexToIsoInnerNodes[Index + 1];
+			}
+			if (Node != nullptr)
+			{
+				for (FIsoSegment* Segment : Node->GetConnectedSegments())
+				{
+					if (Segment->GetType() == ESegmentType::IsoU)
+					{
+						if (Segment->GetSecondNode().GetIndex() == Index + 1)
+						{
+							InnerSegmentsIntersectionTool.AddSegment(*Segment);
+						}
+					}
+					else if (Segment->GetSecondNode().GetIndex() == Index + NumU)
+					{
+						InnerSegmentsIntersectionTool.AddSegment(*Segment);
+					}
+				}
+			}
+
+			Index = (IndexV + 1) * NumU + IndexU + 1;
+			Node = GlobalIndexToIsoInnerNodes[Index];
+			if (Node == nullptr)
+			{
+				Node = GlobalIndexToIsoInnerNodes[Index - 1];
+			}
+			if (Node != nullptr)
+			{
+				for (FIsoSegment* Segment : Node->GetConnectedSegments())
+				{
+					if (Segment->GetType() == ESegmentType::IsoU)
+					{
+						if (Segment->GetFirstNode().GetIndex() == Index - 1)
+						{
+							InnerSegmentsIntersectionTool.AddSegment(*Segment);
+						}
+					}
+					else if (Segment->GetFirstNode().GetIndex() == Index - NumU)
+					{
+						InnerSegmentsIntersectionTool.AddSegment(*Segment);
+					}
+				}
+			}
+		}
+	}
+
+#ifdef CADKERNEL_DEV
+	Chronos.FindSegmentIsoUVSurroundingSmallLoopDuration += FChrono::Elapse(StartTime);
+#endif
+}
+
+void FIsoTriangulator::TriangulateOverCycle(const EGridSpace Space)
+{
+#ifdef ADD_TRIANGLE_2D
+	F3DDebugSession G(bDisplay, ("Triangulate cycles"));
+#endif 
+
+	FTimePoint StartTime = FChrono::Now();
+
+	TArray<FIsoSegment*> Cycle;
+	Cycle.Reserve(100);
+	TArray<bool> CycleOrientation;
+	CycleOrientation.Reserve(100);
+
+#ifdef FIND_CYCLE
+	int32 OverCycleIndex = 0;
+	if (Grid.bDisplay)
+	{
+		Open3DDebugSession(TEXT("Triangulate Over Cycle"));
+	}
+#endif
+
+	// first the external segments (loop segments) are processed 
+	for (FIsoSegment* Segment : LoopSegments)
+	{
+		if (!Segment->HasCycleOnLeft())
+		{
+			Cycle.Empty();
+			CycleOrientation.Empty();
+			bool bLeftSide = true;
+			if (!FindCycle(Segment, bLeftSide, Cycle, CycleOrientation))
+			{
+				continue;
+			}
+#ifdef FIND_CYCLE
+			if (Grid.bDisplay)
+			{
+				F3DDebugSession G(TEXT("Find & mesh cycles"));
+				FString Message = FString::Printf(TEXT("MeshCycle - cycle %d"), OverCycleIndex++);
+				Grid.DisplayIsoSegments(Message, EGridSpace::UniformScaled, Cycle, false);
+			}
+#endif
+			MeshCycle(Cycle, CycleOrientation);
+		}
+	}
+
+	// then all segments are processed 
+	for (FIsoSegment* Segment : FinalToLoops)
+	{
+		if (!Segment->HasCycleOnLeft())
+		{
+			Cycle.Empty();
+			CycleOrientation.Empty();
+			bool bLeftSide = true;
+			if (!FindCycle(Segment, bLeftSide, Cycle, CycleOrientation))
+			{
+				continue;
+			}
+#ifdef FIND_CYCLE
+			if (Grid.bDisplay)
+			{
+				Open3DDebugSession(TEXT("Find & mesh cycles"));
+				FString Message = FString::Printf(TEXT("MeshCycle - cycle %d"), OverCycleIndex++);
+				Grid.DisplayIsoSegments(Message, EGridSpace::UniformScaled, Cycle, false);
+				Close3DDebugSession();
+			}
+#endif
+			MeshCycle(Cycle, CycleOrientation);
+		}
+
+		if (!Segment->HasCycleOnRight())
+		{
+			Cycle.Empty();
+			CycleOrientation.Empty();
+			bool bLeftSide = false;
+			if (!FindCycle(Segment, bLeftSide, Cycle, CycleOrientation))
+			{
+				continue;
+			}
+#ifdef FIND_CYCLE
+			if (Grid.bDisplay)
+			{
+				Open3DDebugSession(TEXT("Find & mesh cycles"));
+				FString Message = FString::Printf(TEXT("MeshCycle - cycle %d"), OverCycleIndex++);
+				Grid.DisplayIsoSegments(Message, EGridSpace::UniformScaled, Cycle, false);
+				Close3DDebugSession();
+			}
+#endif
+			MeshCycle(Cycle, CycleOrientation);
+		}
+	}
+#ifdef FIND_CYCLE
+	Close3DDebugSession();
+#endif
+
+#ifdef CADKERNEL_DEV
+	Chronos.TriangulateOverCycleDuration = FChrono::Elapse(StartTime);
+#endif
+}
+
+#ifdef DEBUG_FIND_CYCLE
+static int32 CycleId = -1;
+static int32 CycleIndex = 0;
+#endif
+
+bool FIsoTriangulator::FindCycle(FIsoSegment* StartSegment, bool LeftSide, TArray<FIsoSegment*>& Cycle, TArray<bool>& CycleOrientation)
+{
+
+#ifdef DEBUG_FIND_CYCLE
+	CycleIndex++;
+
+	if (Grid.bDisplay)
+	{
+		CycleIndex = CycleId;
+		bDisplay = true;
+	}
+#endif
+
+	Cycle.Empty();
+	CycleOrientation.Empty();
+
+#ifdef DEBUG_FIND_CYCLE
+	if (CycleId == CycleIndex)
+	{
+		Open3DDebugSession(TEXT("Cycle"));
+	}
+#endif
+
+	FIsoSegment* Segment = StartSegment;
+	FIsoNode* Node;
+
+	if (LeftSide)
+	{
+		Segment->SetHaveCycleOnLeft();
+		Node = &StartSegment->GetSecondNode();
+#ifdef DEBUG_FIND_CYCLE
+		if (CycleId == CycleIndex)
+		{
+			FIsoNode* EndNode = &StartSegment->GetFirstNode();
+			F3DDebugSession _(TEXT("FirstSegment left"));
+			Grid.DisplayIsoNode(EGridSpace::UniformScaled, *EndNode);
+			Grid.DisplayIsoSegment(EGridSpace::UniformScaled, *StartSegment);
+			//Wait();
+		}
+#endif
+	}
+	else
+	{
+		Segment->SetHaveCycleOnRight();
+		Node = &StartSegment->GetFirstNode();
+#ifdef DEBUG_FIND_CYCLE
+		if (CycleId == CycleIndex)
+		{
+			FIsoNode* EndNode = &StartSegment->GetSecondNode();
+			F3DDebugSession _(TEXT("FirstSegment right"));
+			Grid.DisplayIsoNode(EGridSpace::UniformScaled, *EndNode);
+			Grid.DisplayIsoSegment(EGridSpace::UniformScaled, *StartSegment);
+			//Wait();
+		}
+#endif
+	}
+
+	Cycle.Add(StartSegment);
+	CycleOrientation.Add(LeftSide);
+	Segment = StartSegment;
+
+	for (;;)
+	{
+		Segment = FindNextSegment(EGridSpace::UniformScaled, Segment, Node, ClockwiseSlope);
+		if (Segment == nullptr)
+		{
+			Cycle.Empty();
+			break;
+		}
+
+		if (Segment == StartSegment)
+		{
+			break;
+		}
+
+		Cycle.Add(Segment);
+
+		if (&Segment->GetFirstNode() == Node)
+		{
+#ifdef DEBUG_HAS_CYCLE
+			if (Segment->HasCycleOnLeft())
+			{
+				F3DDebugSession _(TEXT("Segment HasCycleOnLeft"));
+				Grid.DisplayIsoSegment(EGridSpace::UniformScaled, *Segment, 0, EVisuProperty::RedPoint);
+				Wait();
+			}
+#endif
+			if (Segment->HasCycleOnLeft())
+			{
+				return false;
+			}
+			Segment->SetHaveCycleOnLeft();
+			Node = &Segment->GetSecondNode();
+			CycleOrientation.Add(true);
+#ifdef DEBUG_FIND_CYCLE
+			if (CycleId == CycleIndex)
+			{
+				F3DDebugSession _(TEXT("Next"));
+				Grid.DisplayIsoNode(EGridSpace::UniformScaled, *Node);
+				Grid.DisplayIsoSegment(EGridSpace::UniformScaled, *Segment);
+				//Wait();
+			}
+#endif
+		}
+		else
+		{
+#ifdef DEBUG_HAS_CYCLE
+			if (Segment->HasCycleOnRight())
+			{
+				F3DDebugSession _(TEXT("Segment HasCycleOnRight"));
+				Grid.DisplayIsoSegment(EGridSpace::UniformScaled, *Segment, 0, EVisuProperty::RedPoint);
+				Wait();
+			}
+#endif
+			if (Segment->HasCycleOnRight())
+			{
+				return false;
+			}
+			Segment->SetHaveCycleOnRight();
+			Node = &Segment->GetFirstNode();
+			CycleOrientation.Add(false);
+#ifdef DEBUG_FIND_CYCLE
+			if (CycleId == CycleIndex)
+			{
+				F3DDebugSession _(TEXT("Next"));
+				Grid.DisplayIsoNode(EGridSpace::UniformScaled, *Node);
+				Grid.DisplayIsoSegment(EGridSpace::UniformScaled, *Segment);
+				//Wait();
+			}
+#endif
+		}
+		if (false)
+		{
+			Wait();
+		}
+	}
+#ifdef DEBUG_FIND_CYCLE
+	if (CycleId == CycleIndex)
+	{
+		Close3DDebugSession();
+	}
+#endif
+	return true;
+}
+
+FIsoSegment* FIsoTriangulator::FindNextSegment(EGridSpace Space, const FIsoSegment* StartSegment, const FIsoNode* StartNode, SlopeMethod GetSlope) const
+{
+	const FPoint2D& StartPoint = StartNode->Get2DPoint(Space, Grid);
+	const FPoint2D& EndPoint = (StartNode == &StartSegment->GetFirstNode()) ? StartSegment->GetSecondNode().Get2DPoint(Space, Grid) : StartSegment->GetFirstNode().Get2DPoint(Space, Grid);
+
+	double ReferenceSlope = ComputePositiveSlope(StartPoint, EndPoint, 0);
+
+	double MaxSlope = 8.1;
+	FIsoSegment* NextSegment = nullptr;
+
+	for (FIsoSegment* Segment : StartNode->GetConnectedSegments())
+	{
+		const FPoint2D& OtherPoint = (StartNode == &Segment->GetFirstNode()) ? Segment->GetSecondNode().Get2DPoint(Space, Grid) : Segment->GetFirstNode().Get2DPoint(Space, Grid);
+
+		double Slope = GetSlope(StartPoint, OtherPoint, ReferenceSlope);
+		if (Slope < SMALL_NUMBER_SQUARE)
+		{
+			Slope = 8;
+		}
+
+		if (Slope < MaxSlope || NextSegment == StartSegment)
+		{
+			NextSegment = Segment;
+			MaxSlope = Slope;
+		}
+	}
+
+	return NextSegment;
+}
+
+void FIsoTriangulator::TriangulateInnerNodes()
+{
+	FTimePoint StartTime = FChrono::Now();
+
+	int32 NumU = Grid.GetCuttingCount(EIso::IsoU);
+	int32 NumV = Grid.GetCuttingCount(EIso::IsoV);
+
+#ifdef ADD_TRIANGLE_2D
+	if (bDisplay)
+	{
+		Open3DDebugSession(TEXT("Inner Mesh 2D"));
+	}
+#endif
+	for (int32 vIndex = 0, Index = 0; vIndex < NumV - 1; vIndex++)
+	{
+		for (int32 uIndex = 0; uIndex < NumU - 1; uIndex++, Index++)
+		{
+			// Do the lower nodes of the cell exist
+			if (!GlobalIndexToIsoInnerNodes[Index] || !GlobalIndexToIsoInnerNodes[Index + 1])
+			{
+				continue;
+			}
+
+			// Is the lower left node connected
+			if (!GlobalIndexToIsoInnerNodes[Index]->IsLinkedToNextU() || !GlobalIndexToIsoInnerNodes[Index]->IsLinkedToNextV())
+			{
+				continue;
+			}
+
+			// Do the upper nodes of the cell exist
+			int32 OppositIndex = Index + NumU + 1;
+			if (!GlobalIndexToIsoInnerNodes[OppositIndex] || !GlobalIndexToIsoInnerNodes[OppositIndex - 1])
+			{
+				continue;
+			}
+
+			// Is the top right node connected
+			if (!GlobalIndexToIsoInnerNodes[OppositIndex]->IsLinkedToPreviousU() || !GlobalIndexToIsoInnerNodes[OppositIndex]->IsLinkedToPreviousV())
+			{
+				continue;
+			}
+
+			bool bIsSurroundingALoop = false;
+			for (int32 BorderIndex : IndexOfLowerLeftInnerNodeSurroundingALoop)
+			{
+				if (Index == BorderIndex)
+				{
+					bIsSurroundingALoop = true;
+					break;
+				}
+			}
+			if (bIsSurroundingALoop)
+			{
+				continue;
+			}
+
+#ifdef ADD_TRIANGLE_2D
+			Grid.DisplayTriangle(EGridSpace::UniformScaled, *GlobalIndexToIsoInnerNodes[Index], *GlobalIndexToIsoInnerNodes[Index + 1], *GlobalIndexToIsoInnerNodes[OppositIndex]);
+			Grid.DisplayTriangle(EGridSpace::UniformScaled, *GlobalIndexToIsoInnerNodes[OppositIndex], *GlobalIndexToIsoInnerNodes[OppositIndex - 1], *GlobalIndexToIsoInnerNodes[Index]);
+#endif 
+
+			Mesh.AddTriangle(GlobalIndexToIsoInnerNodes[Index]->GetGlobalIndex(), GlobalIndexToIsoInnerNodes[Index + 1]->GetGlobalIndex(), GlobalIndexToIsoInnerNodes[OppositIndex]->GetGlobalIndex());
+			Mesh.AddTriangle(GlobalIndexToIsoInnerNodes[OppositIndex]->GetGlobalIndex(), GlobalIndexToIsoInnerNodes[OppositIndex - 1]->GetGlobalIndex(), GlobalIndexToIsoInnerNodes[Index]->GetGlobalIndex());
+		}
+		Index++;
+	}
+#ifdef ADD_TRIANGLE_2D
+	Close3DDebugSession();
+#endif
+}
+
+void FCell::FindCandidateToConnectLoopsByNeighborhood()
+{
+#ifdef DEBUG_CONNECT_CELL_SUB_LOOPS_BY_NEIGHBORHOOD
+	F3DDebugSession _(Grid.bDisplay, FString::Printf(TEXT("Cell % d: Build Segments to Connect inner loops"), Id));
+#endif
+
+	for (FCellConnexion& LoopConnexion : LoopConnexions)
+	{
+		TryToConnectTwoSubLoopsWithShortestSegment(LoopConnexion);
+	}
+}
+
+void FCell::SelectSegmentToConnectLoops(TFactory<FIsoSegment>& SegmentFactory)
+{
+	TArray<FCellConnexion*> LoopConnexionPtrs;
+	LoopConnexionPtrs.Reserve(LoopConnexions.Num());
+	for (FCellConnexion& Connexion : LoopConnexions)
+	{
+		if (Connexion.Loop2.IsCellCorner())
+		{
+			continue;
+		}
+		LoopConnexionPtrs.Add(&Connexion);
+	}
+
+	Algo::Sort(LoopConnexionPtrs, [&](const FCellConnexion* LoopConnexion1, const FCellConnexion* LoopConnexion2)
+		{
+			return LoopConnexion1->MinDistance < LoopConnexion2->MinDistance;
+		});
+
+	{
+#ifdef DEBUG_CONNECT_CELL_SUB_LOOPS_BY_NEIGHBORHOOD
+		F3DDebugSession _(Grid.bDisplay, TEXT("Best Path"));
+#endif
+#ifdef DEBUG_CONNECT_CELL_SUB_LOOPS_BY_NEIGHBORHOOD
+		if (Grid.bDisplay)
+		{
+			Grid.DisplayIsoSegments(TEXT("SelectSegmentToConnectLoops: before"), EGridSpace::UniformScaled, FinalSegments, false, false, false, EVisuProperty::YellowCurve);
+			Wait(Grid.bDisplay);
+		}
+#endif
+		const int32 LoopCount = CellLoops.Num();
+		for (FCellConnexion* LoopConnexion : LoopConnexionPtrs)
+		{
+			if (!LoopConnexion->bIsConnexionWithOuter && LoopConnexion->IsShortestPath(LoopCount))
+			{
+				TryToCreateSegment(*LoopConnexion);
+			}
+		}
+
+		for (FCellConnexion* LoopConnexion : LoopConnexionPtrs)
+		{
+			if (LoopConnexion->bIsConnexionWithOuter && LoopConnexion->IsShortestPathToOuterLoop(LoopCount))
+			{
+				TryToCreateSegment(*LoopConnexion);
+			}
+		}
+	}
+
+#ifdef DEBUG_CONNECT_CELL_SUB_LOOPS_BY_NEIGHBORHOOD
+	if (Grid.bDisplay)
+	{
+		Grid.DisplayIsoSegments(TEXT("SelectSegmentToConnectLoops: Final"), EGridSpace::UniformScaled, FinalSegments, false, false, false, EVisuProperty::YellowCurve);
+		Wait(Grid.bDisplay);
+	}
+#endif
+}
+
+void FCell::SelectSegmentToConnectLoopToCorner(TFactory<FIsoSegment>& SegmentFactory)
+{
+	TArray<FCellConnexion*> LoopConnexionPtrs;
+	LoopConnexionPtrs.Reserve(LoopConnexions.Num());
+	for (FCellConnexion& Connexion : LoopConnexions)
+	{
+		if (Connexion.Loop2.IsCellCorner() && !Connexion.Segment)
+		{
+			LoopConnexionPtrs.Add(&Connexion);
+		}
+	}
+
+	Algo::Sort(LoopConnexionPtrs, [&](const FCellConnexion* LoopConnexion1, const FCellConnexion* LoopConnexion2)
+		{
+			return LoopConnexion1->MinDistance < LoopConnexion2->MinDistance;
+		});
+
+	{
+#ifdef DEBUG_CONNECT_CELL_SUB_LOOPS_BY_NEIGHBORHOOD
+		F3DDebugSession _(Grid.bDisplay, TEXT("Best Path"));
+#endif
+		const int32 LoopCount = CellLoops.Num();
+		for (FCellConnexion* LoopConnexion : LoopConnexionPtrs)
+		{
+			if (!LoopConnexion->bIsConnexionWithOuter && LoopConnexion->IsShortestPathToCorner(LoopCount))
+			{
+				TryToCreateSegment(*LoopConnexion);
+			}
+		}
+
+		for (FCellConnexion* LoopConnexion : LoopConnexionPtrs)
+		{
+			if (LoopConnexion->bIsConnexionWithOuter)
+			{
+				// if the outerLoop is already connected to an innerLoop, the segment is not create.
+				// The aim is to avoid to create long segment that generate degenerated triangle
+				bool bCreateSegment = true;
+				for (const FCellConnexion* Connexion : LoopConnexion->Loop1.Connexions)
+				{
+					if (Connexion->Segment && Connexion->Segment->IsAFinalSegment())
+					{
+						const FCellLoop& InnerLoop = Connexion->Loop1.bIsOuterLoop ? Connexion->Loop2 : Connexion->Loop1;
+						if (!InnerLoop.IsCellCorner())
+						{
+							bCreateSegment = false;
+						}
+					}
+				}
+				if (bCreateSegment)
+				{
+					TryToCreateSegment(*LoopConnexion);
+				}
+			}
+		}
+	}
+
+#ifdef DEBUG_CONNECT_CELL_SUB_LOOPS_BY_NEIGHBORHOOD
+	if (Grid.bDisplay)
+	{
+		Grid.DisplayIsoSegments(TEXT("Build Segments Connect inner close loops: Candidates"), EGridSpace::UniformScaled, CandidateSegments, false, false, false, EVisuProperty::YellowCurve);
+		Wait(Grid.bDisplay);
+	}
+#endif
+}
+
+void FCellLoop::PropagateAsConnected()
+{
+	bIsConnected = true;
+	for (FCellConnexion* Connexion : Connexions)
+	{
+		if (!Connexion->Segment)
+		{
+			continue;
+		}
+
+		FCellLoop* OtherCell = Connexion->GetOtherLoop(this);
+		if (!OtherCell->bIsConnected)
+		{
+			OtherCell->PropagateAsConnected();
+		}
+	}
+}
+
+void FCell::CheckAllLoopsConnectedTogetherAndConnect()
+{
+	for (FCellCorner& CellCorner : CellCorners)
+	{
+		CellCorner.PropagateAsConnected();
+	}
+
+	for (FCellLoop& CellLoop : CellLoops)
+	{
+		if (CellLoop.bIsOuterLoop)
+		{
+			CellLoop.PropagateAsConnected();
+		}
+	}
+}
+
+void FCell::TryToCreateSegment(FCellConnexion& LoopConnexion)
+{
+	const FPoint2D& ACoordinates = LoopConnexion.NodeA->Get2DPoint(EGridSpace::UniformScaled, Grid);
+	const FPoint2D& BCoordinates = LoopConnexion.NodeB->Get2DPoint(EGridSpace::UniformScaled, Grid);
+
+	LoopConnexion.Segment = Triangulator.GetOrTryToCreateSegment(*this, LoopConnexion.NodeA, ACoordinates, LoopConnexion.NodeB, BCoordinates, Slope::OneDegree);
+	if (LoopConnexion.Segment && !LoopConnexion.Segment->IsAFinalSegment())
+	{
+		LoopConnexion.Segment->SetFinalMarker();
+		if(LoopConnexion.Segment->ConnectToNode())
+		{
+			IntersectionTool.AddSegment(*LoopConnexion.Segment);
+			FinalSegments.Add(LoopConnexion.Segment);
+		}
+		else
+		{
+#ifdef DEBUG_BUILD_SEGMENT_IF_NEEDED
+			F3DDebugSession _(FString::Printf(TEXT("Not expected error Segment")));
+			Grid.DisplayIsoSegment(EGridSpace::UniformScaled, *LoopConnexion.NodeA, *LoopConnexion.NodeB, 0, EVisuProperty::RedCurve);
+			Wait();
+#endif
+		}
+	}
+#ifdef DEBUG_BUILD_SEGMENT_IF_NEEDED2
+	else if(Grid.bDisplay)
+	{
+		F3DDebugSession _(FString::Printf(TEXT("Failed Segment %d %d"), LoopConnexion.Loop1.Id, LoopConnexion.Loop2.Id));
+		Grid.DisplayIsoSegment(EGridSpace::UniformScaled, *LoopConnexion.NodeA, *LoopConnexion.NodeB, 0, EVisuProperty::RedCurve);
+		Wait();
+	}
+#endif
+}
+
+
+void FCell::InitLoopConnexions()
+{
+	const int32 MaxInnerToCornerConnexions = CellCorners.Num() * (InnerLoopCount + OuterLoopCount);
+
+	switch (InnerLoopCount)
+	{
+	case 1:
+	{
+		LoopConnexions.Reserve(OuterLoopCount + MaxInnerToCornerConnexions);
+		LoopCellBorderIndices = { OuterLoopCount };
+		break;
+	}
+	case 2:
+	{
+		LoopConnexions.Reserve(1 + 2 * OuterLoopCount + MaxInnerToCornerConnexions);
+		LoopConnexions.Emplace(CellLoops[OuterLoopCount], CellLoops[OuterLoopCount + 1]);
+
+		LoopCellBorderIndices = { OuterLoopCount , OuterLoopCount + 1 };
+		break;
+	}
+	case 3:
+	{
+		LoopConnexions.Reserve(3 + 3 * OuterLoopCount + MaxInnerToCornerConnexions);
+		LoopConnexions.Emplace(CellLoops[OuterLoopCount], CellLoops[OuterLoopCount + 1]);
+		LoopConnexions.Emplace(CellLoops[OuterLoopCount], CellLoops[OuterLoopCount + 2]);
+		LoopConnexions.Emplace(CellLoops[OuterLoopCount + 1], CellLoops[OuterLoopCount + 2]);
+
+		LoopCellBorderIndices = { OuterLoopCount , OuterLoopCount + 1 , OuterLoopCount + 2 };
+		break;
+	}
+	default:
+	{
+		TArray<TPair<int32, FPoint2D>> LoopBarycenters = GetLoopBarycenters();
+
+#ifdef DEBUG_BOWYERWATSON
+		FBowyerWatsonTriangulator::bDisplay = bDisplay;
+#endif
+		TArray<int32> EdgeVertexIndices;
+		FBowyerWatsonTriangulator BWTriangulator(LoopBarycenters, EdgeVertexIndices);
+		BWTriangulator.Triangulate();
+		LoopCellBorderIndices = BWTriangulator.GetOuterVertices();
+		LoopConnexions.Reserve((EdgeVertexIndices.Num() >> 1) + LoopCellBorderIndices.Num() * OuterLoopCount + MaxInnerToCornerConnexions);
+
+		for (int32 Index = 0; Index < EdgeVertexIndices.Num();)
+		{
+			const int32 StartIndex = Index++;
+			const int32 EndIndex = Index++;
+			ensureCADKernel(LoopConnexions.Max() > LoopConnexions.Num());
+			LoopConnexions.Emplace(CellLoops[EdgeVertexIndices[StartIndex]], CellLoops[EdgeVertexIndices[EndIndex]]);
+		}
+	}
+	}
+
+	for (int32 Index = 0; Index < OuterLoopCount; ++Index)
+	{
+		for (int32 BIndex : LoopCellBorderIndices)
+		{
+			ensureCADKernel(LoopConnexions.Max() > LoopConnexions.Num());
+			LoopConnexions.Emplace(CellLoops[Index], CellLoops[BIndex]);
+		}
+	}
+
+	const int32 ConnexionCount = LoopConnexions.Num();
+	for (FCellLoop& Loop : CellLoops)
+	{
+		Loop.Connexions.Reserve(ConnexionCount);
+	}
+}
+
+void FCell::TryToConnectTwoSubLoopsWithShortestSegment(FCellConnexion& LoopConnexion)
+{
+	const TArray<FLoopNode*>& LoopA = LoopConnexion.Loop1.Nodes;
+	const TArray<FLoopNode*>& LoopB = LoopConnexion.Loop2.Nodes;
+
+#ifdef DEBUG_CONNECT_CELL_SUB_LOOPS_BY_NEIGHBORHOOD
+	//F3DDebugSession A(Grid.bDisplay, TEXT("ConnectTwoSubLoopsWithShortestSegment"));
+#endif
+	double MinDistanceSquare = HUGE_VALUE_SQUARE;
+
+	for (FLoopNode* NodeA : LoopA)
+	{
+		const FPoint2D& ACoordinates = NodeA->Get2DPoint(EGridSpace::UniformScaled, Grid);
+
+		for (FLoopNode* NodeB : LoopB)
+		{
+			const FPoint2D& BCoordinates = NodeB->Get2DPoint(EGridSpace::UniformScaled, Grid);
+
+#ifdef DEBUG_CONNECT_CELL_SUB_LOOPS_BY_NEIGHBORHOOD_
+			Grid.DisplayIsoSegment(EGridSpace::UniformScaled, *NodeA, *NodeB, 0, EVisuProperty::BlueCurve);
+#endif
+
+			double SquareDistance = ACoordinates.SquareDistance(BCoordinates);
+			if (SquareDistance < MinDistanceSquare)
+			{
+				MinDistanceSquare = SquareDistance;
+				LoopConnexion.NodeA = NodeA;
+				LoopConnexion.NodeB = NodeB;
+			}
+		}
+	}
+
+	LoopConnexion.MinDistance = FMath::Sqrt(MinDistanceSquare);
+
+#ifdef DEBUG_CONNECT_CELL_SUB_LOOPS_BY_NEIGHBORHOOD
+	if (LoopConnexion.NodeA && LoopConnexion.NodeB)
+	{
+		Grid.DisplayIsoSegment(EGridSpace::UniformScaled, *LoopConnexion.NodeA, *LoopConnexion.NodeB, 0, EVisuProperty::RedCurve);
+	}
+#endif
+}
+
+void FIsoTriangulator::TryToConnectTwoLoopsWithIsocelesTriangle(FCell& Cell, const TArray<FLoopNode*>& LoopA, const TArray<FLoopNode*>& LoopB)
+{
+
+	TFunction<FIsoNode* (FIsoSegment*)> FindBestTriangle = [&](FIsoSegment* Segment) -> FIsoNode*
+	{
+		SlopeMethod GetSlopeAtStartNode = ClockwiseSlope;
+		SlopeMethod GetSlopeAtEndNode = CounterClockwiseSlope;
+
+		// StartNode = A
+		FIsoNode& StartNode = Segment->GetSecondNode();
+		// EndNode = B
+		FIsoNode& EndNode = Segment->GetFirstNode();
+
+		//
+		// For each segment of the LoopA, find in loopB a vertex that made the best triangle i.e. the triangle the most isosceles.
+		// According to the knowledge of the orientation, only inside triangles are tested
+		//
+		// These computations are done in the UniformScaled space to avoid numerical error due to lenght distortion between U or V space and U or V Length.
+		// i.e. if:
+		// (UMax - UMin) / (VMax - VMin) is big 
+		// and 
+		// "medium length along U" / "medium length along V" is small 
+		// 
+		// To avoid flat triangle, a candidate point must defined a minimal slop with [A, X0] or [B, Xn] to not be aligned with one of them. 
+		//
+
+#ifdef DEBUG_FIND_BEST_TRIANGLE_TO_LINK_LOOPS
+		int32 TriangleIndex = 0;
+		if (bDisplay)
+		{
+			{
+				TriangleIndex++;
+				F3DDebugSession _(FString::Printf(TEXT("Start Segment %d"), TriangleIndex));
+				Grid.DisplayIsoSegment(EGridSpace::UniformScaled, *Segment);
+				Grid.DisplayIsoNode(EGridSpace::UniformScaled, StartNode, 0, EVisuProperty::RedPoint);
+				Grid.DisplayIsoNode(EGridSpace::UniformScaled, EndNode);
+				Wait();
+			}
+		}
+#endif
+
+		FIsoNode* CandidatNode = nullptr;
+		FIsoSegment* StartToCandiatSegment = nullptr;
+		FIsoSegment* EndToCandiatSegment = nullptr;
+
+		const FPoint2D& StartPoint2D = StartNode.Get2DPoint(EGridSpace::UniformScaled, Grid);
+		const FPoint2D& EndPoint2D = EndNode.Get2DPoint(EGridSpace::UniformScaled, Grid);
+
+		double StartReferenceSlope = ComputePositiveSlope(StartPoint2D, EndPoint2D, 0);
+		double EndReferenceSlope = StartReferenceSlope < 4 ? StartReferenceSlope + 4 : StartReferenceSlope - 4;
+
+		double MinCriteria = HUGE_VALUE;
+		const double MinSlopeToNotBeAligned = 0.0001;
+		double CandidateSlopeAtStartNode = 8.;
+		double CandidateSlopeAtEndNode = 8.;
+
+		for (FLoopNode* Node : LoopB)
+		{
+			if (Node->IsDeleteOrThinNode())
+			{
+				continue;
+			}
+
+			// Check if the node is inside the sector (X) or outside (Z)
+			const FPoint2D& NodePoint2D = Node->Get2DPoint(EGridSpace::UniformScaled, Grid);
+
+			double PointCriteria = IsoTriangulatorImpl::IsoscelesCriteriaMax(StartPoint2D, EndPoint2D, NodePoint2D);
+
+			// Triangle that are too open (more than rectangle triangle) are not tested 
+			if (PointCriteria > Slope::RightSlope)
+			{
+				continue;
+			}
+
+			double SlopeAtStartNode = GetSlopeAtStartNode(StartPoint2D, NodePoint2D, StartReferenceSlope);
+			double SlopeAtEndNode = GetSlopeAtEndNode(EndPoint2D, NodePoint2D, EndReferenceSlope);
+
+			// check the side of the candidate point according to the segment
+			if (SlopeAtStartNode <= MinSlopeToNotBeAligned)
+			{
+				continue;
+			}
+
+			if (
+				// the candidate triangle is inside the current candidate triangle
+				((SlopeAtStartNode < (CandidateSlopeAtStartNode + MinSlopeToNotBeAligned)) && (SlopeAtEndNode < (CandidateSlopeAtEndNode + MinSlopeToNotBeAligned)))
+				||
+				// the candidate triangle is better than the current candidate triangle and doesn't contain the current candidate triangle
+				((PointCriteria < MinCriteria) && ((SlopeAtStartNode > CandidateSlopeAtStartNode) ^ (SlopeAtEndNode > CandidateSlopeAtEndNode))))
+			{
+				// check if the candidate segment is not in intersection with existing segments
+				// if the segment exist, it has already been tested
+				FIsoSegment* StartSegment = StartNode.GetSegmentConnectedTo(Node);
+				FIsoSegment* EndSegment = EndNode.GetSegmentConnectedTo(Node);
+
+				if (!StartSegment && LoopSegmentsIntersectionTool.DoesIntersect(StartNode, *Node))
+				{
+					continue;
+				}
+
+				if (!EndSegment && LoopSegmentsIntersectionTool.DoesIntersect(EndNode, *Node))
+				{
+					continue;
+				}
+
+				MinCriteria = PointCriteria;
+				CandidatNode = Node;
+				StartToCandiatSegment = StartSegment;
+				EndToCandiatSegment = EndSegment;
+				CandidateSlopeAtStartNode = SlopeAtStartNode;
+				CandidateSlopeAtEndNode = SlopeAtEndNode;
+			}
+		}
+
+		return CandidatNode;
+	};
+
+	// for each segment of LoopA
+	for (int32 IndexA = 0; IndexA < LoopA.Num() - 1; ++IndexA)
+	{
+		FLoopNode* NodeA1 = LoopA[IndexA];
+		FLoopNode* NodeA2 = LoopA[IndexA + 1];
+
+		if (NodeA1->IsDeleteOrThinNode() || NodeA2->IsDeleteOrThinNode())
+		{
+			continue;
+		}
+
+		const FPoint2D& A1Coordinates = NodeA1->Get2DPoint(EGridSpace::UniformScaled, Grid);
+		const FPoint2D& A2Coordinates = NodeA2->Get2DPoint(EGridSpace::UniformScaled, Grid);
+
+		FIsoSegment* Segment = NodeA1->GetSegmentConnectedTo(NodeA2);
+
+		FIsoNode* Node = FindBestTriangle(Segment);
+		if (Node && !Node->IsDeleteOrThinNode())
+		{
+#ifdef DEBUG_TRY_TO_CONNECT
+			if (bDisplay)
+			{
+				F3DDebugSession _(FString::Printf(TEXT("Triangle")));
+				DisplayTriangle(EGridSpace::UniformScaled, *NodeA1, *NodeA2, *Node);
+				Wait(true);
+
+			}
+#endif 
+			const FPoint2D& NodeCoordinates = Node->Get2DPoint(EGridSpace::UniformScaled, Grid);
+			if (!NodeA1->IsDeleteOrThinNode())
+			{
+				GetOrTryToCreateSegment(Cell, NodeA1, A1Coordinates, Node, NodeCoordinates, 0.1);
+			}
+			if (!NodeA2->IsDeleteOrThinNode())
+			{
+				GetOrTryToCreateSegment(Cell, NodeA2, A2Coordinates, Node, NodeCoordinates, 0.1);
+			}
+		}
+	}
+
+};
+
+void FIsoTriangulator::TryToConnectVertexSubLoopWithTheMostIsoSegment(FCell& Cell, const TArray<FLoopNode*>& Loop)
+{
+#ifdef DEBUG_FIND_ISO_SEGMENT_TO_LINK_OUTER_LOOP_NODES
+	F3DDebugSession _(bDisplay, TEXT("TryToConnectVertexSubLoopWithTheMostIsoSegment"));
+	Wait(bDisplay);
+#endif
+
+	const double FlatSlope = 0.10; // ~5 deg: The segment must make an angle less than 10 deg with the Iso
+	double MinSlope = FlatSlope + DOUBLE_SMALL_NUMBER; //.25; // ~10 deg: The segment must make an angle less than 10 deg with the Iso
+
+
+	if (Loop.Num() <= 2)
+	{
+		return;
+	}
+
+	int32 LoopCount = Loop.Num();
+	for (int32 IndexA = 0; IndexA < LoopCount - 2; ++IndexA)
+	{
+		FLoopNode* CandidateB = nullptr;
+
+		FLoopNode* CandidateA = Loop[IndexA];
+		if (CandidateA->IsThinZoneNode())
+		{
+			continue;
+		}
+
+		const FPoint2D& ACoordinates = CandidateA->Get2DPoint(EGridSpace::UniformScaled, Grid);
+
+		FLoopNode* NextA = Loop[IndexA + 1];
+		const FPoint2D& NextACoordinates = NextA->Get2DPoint(EGridSpace::UniformScaled, Grid);
+
+		double ReferenceSlope = 0;
+		{
+			// slope SegmentA (NodeA, NextA)
+			double Slope = ComputeUnorientedSlope(ACoordinates, NextACoordinates, 0);
+			if (Slope > 1.5 && Slope < 2.5)
+			{
+				ReferenceSlope = 0;
+			}
+			else if ((Slope < 0.5) || (Slope > 3.5))
+			{
+				ReferenceSlope = 2;
+			}
+			else
+			{
+				// SegmentA is neither close to IsoV nor IsoU
+				continue;
+			}
+		}
+
+		for (int32 IndexB = IndexA + 2; IndexB < LoopCount; ++IndexB)
+		{
+			FLoopNode* NodeB = Loop[IndexB];
+			if (NodeB->IsThinZoneNode())
+			{
+				continue;
+			}
+
+			const FPoint2D& BCoordinates = NodeB->Get2DPoint(EGridSpace::UniformScaled, Grid);
+
+			double Slope = ComputeSlopeRelativeToReferenceAxis(ACoordinates, BCoordinates, ReferenceSlope);
+			if (Slope < MinSlope)
+			{
+				MinSlope = Slope;
+				CandidateB = NodeB;
+			}
+		}
+
+		if (MinSlope < FlatSlope)
+		{
+			const FPoint2D& BCoordinates = CandidateB->Get2DPoint(EGridSpace::UniformScaled, Grid);
+
+#ifdef DEBUG_FIND_ISO_SEGMENT_TO_LINK_OUTER_LOOP_NODES
+			Grid.DisplayIsoSegment(UniformScaled, *CandidateA, *CandidateB, 0, EVisuProperty::RedCurve);
+#endif
+			GetOrTryToCreateSegment(Cell, CandidateA, ACoordinates, CandidateB, BCoordinates, 0.1);
+			MinSlope = FlatSlope + DOUBLE_SMALL_NUMBER;
+		}
+	}
+
+};
+
+void FIsoTriangulator::TryToConnectTwoSubLoopsWithTheMostIsoSegment(FCell& Cell, const TArray<FLoopNode*>& LoopA, const TArray<FLoopNode*>& LoopB)
+{
+#ifdef DEBUG_FIND_ISO_SEGMENT_TO_LINK_OUTER_LOOP_NODES
+	F3DDebugSession _(bDisplay, TEXT("TryToConnectVertexSubLoopWithTheMostIsoSegment"));
+	if (bDisplay)
+	{
+		Wait();
+	}
+#endif
+
+	const double FlatSlope = 0.10; // ~5 deg: The segment must make an angle less than 10 deg with the Iso
+
+	for (FLoopNode* CandidateA : LoopA)
+	{
+		if (CandidateA->IsThinZoneNode())
+		{
+			continue;
+		}
+
+		FLoopNode* CandidateB = nullptr;
+		const FPoint2D& ACoordinates = CandidateA->Get2DPoint(EGridSpace::UniformScaled, Grid);
+
+		double MinSlope = FlatSlope + DOUBLE_SMALL_NUMBER;// 0.25; // ~15 deg: The segment must make an angle less than 10 deg with the Iso
+		double MinLengthSquare = HUGE_VALUE;
+
+		for (FLoopNode* NodeB : LoopB)
+		{
+			if (NodeB->IsThinZoneNode())
+			{
+				continue;
+			}
+
+			const FPoint2D& BCoordinates = NodeB->Get2DPoint(EGridSpace::UniformScaled, Grid);
+
+			double Slope = ComputeSlopeRelativeToNearestAxis(ACoordinates, BCoordinates);
+			if (Slope < MinSlope)
+			{
+				MinSlope = Slope;
+				// If the slope of the candidate segments is nearly zero, then select the shortest
+				if (MinSlope < DOUBLE_KINDA_SMALL_NUMBER)
+				{
+					double DistanceSquare = BCoordinates.SquareDistance(ACoordinates);
+					if (DistanceSquare > MinLengthSquare)
+					{
+						continue;
+					}
+					MinLengthSquare = DistanceSquare;
+				}
+				CandidateB = NodeB;
+			}
+		}
+
+		if (MinSlope < FlatSlope)
+		{
+			const FPoint2D& BCoordinates = CandidateB->Get2DPoint(EGridSpace::UniformScaled, Grid);
+
+#ifdef DEBUG_FIND_ISO_SEGMENT_TO_LINK_OUTER_LOOP_NODES
+			Grid.DisplayIsoSegment(EGridSpace::UniformScaled, *CandidateA, *CandidateB, 0, EVisuProperty::BlueCurve);
+#endif			
+			GetOrTryToCreateSegment(Cell, CandidateA, ACoordinates, CandidateB, BCoordinates, 0.1);
+			MinSlope = FlatSlope + DOUBLE_SMALL_NUMBER;
+		}
+	}
+}
+
+FIsoSegment* FIsoTriangulator::GetOrTryToCreateSegment(FCell& Cell, FLoopNode* NodeA, const FPoint2D& ACoordinates, FIsoNode* NodeB, const FPoint2D& BCoordinates, const double FlatAngle)
+{
+	if (FIsoSegment* Segment = NodeA->GetSegmentConnectedTo(NodeB))
+	{
+		return Segment;
+	}
+
+	if (InnerSegmentsIntersectionTool.DoesIntersect(*NodeA, *NodeB))
+	{
+		return nullptr;
+	}
+
+	if (ThinZoneIntersectionTool.DoesIntersect(*NodeA, *NodeB))
+	{
+		return nullptr;
+	}
+
+	if (Cell.IntersectionTool.DoesIntersect(*NodeA, *NodeB))
+	{
+		return nullptr;
+	}
+
+	if (LoopSegmentsIntersectionTool.DoesIntersect(*NodeA, *NodeB))
+	{
+		return nullptr;
+	}
+
+#ifdef DEBUG_TRY_TO_CREATE_SEGMENT
+	{
+		F3DDebugSession _(TEXT("Test"));
+		DisplaySegmentWithScale(NodeA->Get2DPoint(EGridSpace::UniformScaled, Grid), NodeB->Get2DPoint(EGridSpace::UniformScaled, Grid), 0, EVisuProperty::RedCurve);
+		Wait(false);
+	}
+#endif
+
+	// Is Outside and not too flat at NodeA
+	if (NodeA->IsSegmentBeInsideFace(BCoordinates, Grid, FlatAngle))
+	{
+		return nullptr;
+	}
+
+	// Is Outside and not too flat at NodeB
+	if (NodeB->IsALoopNode())
+	{
+		if (((FLoopNode*)NodeB)->IsSegmentBeInsideFace(ACoordinates, Grid, FlatAngle))
+		{
+			return nullptr;
+		}
+	}
+
+	FIsoSegment& Segment = IsoSegmentFactory.New();
+	Segment.Init(*NodeA, *NodeB, ESegmentType::LoopToLoop);
+	Segment.SetCandidate();
+	Cell.CandidateSegments.Add(&Segment);
+
+#ifdef DEBUG_TRY_TO_CREATE_SEGMENT
+	DisplaySegment(ACoordinates, BCoordinates, 0, EVisuProperty::OrangePoint);
+#endif
+	return &Segment;
+};
+
+void FIsoTriangulator::InitCellCorners(FCell& Cell)
+{
+	FIsoInnerNode* CellNodes[4];
+	int32 Index = Cell.Id;
+	CellNodes[0] = GlobalIndexToIsoInnerNodes[Index++];
+	CellNodes[1] = GlobalIndexToIsoInnerNodes[Index];
+	Index += Grid.GetCuttingCount(EIso::IsoU);;
+	CellNodes[2] = GlobalIndexToIsoInnerNodes[Index--];
+	CellNodes[3] = GlobalIndexToIsoInnerNodes[Index];
+
+	for (int32 ICell = 0; ICell < 4; ++ICell)
+	{
+		if (CellNodes[ICell])
+		{
+			Cell.CellCorners.Emplace(ICell, *CellNodes[ICell], Grid);
+		}
+	}
+}
+
+void FIsoTriangulator::FindCandidateToConnectCellCornerToLoops(FCell& Cell)
+{
+	if (Cell.CellCorners.IsEmpty())
+	{
+		return;
+	}
+
+#ifdef DEBUG_CONNECT_CELL_CORNER_TO_INNER_LOOP
+	F3DDebugSession _(bDisplay, TEXT("With cell corners"));
+#endif
+
+	TFunction<void(FCellLoop&, FCellCorner&)> FindAndTryCreateCandidateSegmentToLinkLoopToCorner = [&](FCellLoop& LoopCell, FCellCorner& CellCorner)
+	{
+		FIsoInnerNode& CornerNode = CellCorner.CornerNode;
+		const FPoint2D& CornerPoint = CellCorner.Barycenter;
+
+		const TArray<FLoopNode*>& LoopA = LoopCell.Nodes;
+
+		double MinDistanceSquare = HUGE_VALUE_SQUARE;
+		int32 MinIndexA = -1;
+		for (int32 IndexA = 0; IndexA < LoopA.Num(); ++IndexA)
+		{
+			FLoopNode* NodeA = LoopA[IndexA];
+
+			const FPoint2D& ACoordinates = NodeA->Get2DPoint(EGridSpace::UniformScaled, Grid);
+
+			double SquareDistance = ACoordinates.SquareDistance(CornerPoint);
+			if (SquareDistance < MinDistanceSquare)
+			{
+				MinDistanceSquare = SquareDistance;
+				MinIndexA = IndexA;
+			}
+		}
+
+		if (MinIndexA >= 0)
+		{
+			FLoopNode* NodeA = LoopA[MinIndexA];
+			const FPoint2D& ACoordinates = NodeA->Get2DPoint(EGridSpace::UniformScaled, Grid);
+
+			double SlopeVsCell = ComputeSlopeRelativeToNearestAxis(ACoordinates, CornerPoint);
+			if (SlopeVsCell < Slope::OneDegree)
+			{
+				// if the candidate segment is to close to the cell border, we cannot check if in the other side of the cell there are not also a parallel candidate segment
+				// Add it in InnerToLoopCandidateSegments to be processed at the end
+
+				FIsoSegment& Segment = IsoSegmentFactory.New();
+				Segment.Init(*NodeA, CornerNode, ESegmentType::InnerToLoopU);
+				Segment.SetCandidate();
+
+				// The connection is nevertheless created to avoid failed in CheckAllLoopsConnectedTogetherAndConnect while the connection while be create in a second time based on InnerToLoopCandidateSegments
+				ensureCADKernel(Cell.LoopConnexions.Max() > Cell.LoopConnexions.Num());
+				FCellConnexion& Connexion = Cell.LoopConnexions.Emplace_GetRef(LoopCell, CellCorner);
+				Connexion.NodeA = NodeA;
+				Connexion.NodeB = &CornerNode;
+				Connexion.Segment = &Segment;
+
+				InnerToLoopCandidateSegments.Add(&Segment);
+			}
+			else
+			{
+				ensureCADKernel(Cell.LoopConnexions.Max() > Cell.LoopConnexions.Num());
+				FCellConnexion& Connexion = Cell.LoopConnexions.Emplace_GetRef(LoopCell, CellCorner);
+				Connexion.NodeA = NodeA;
+				Connexion.NodeB = &CornerNode;
+			}
+		}
+	};
+
+	int32 IntersectionToolCount = Cell.IntersectionTool.Count();
+	int32 NewSegmentCount = Cell.CandidateSegments.Num() - IntersectionToolCount;
+	Cell.IntersectionTool.AddSegments(Cell.CandidateSegments.GetData() + IntersectionToolCount, NewSegmentCount);
+	Cell.IntersectionTool.Sort();
+
+	for (FCellCorner& CellCorner : Cell.CellCorners)
+	{
+#ifdef DEBUG_CONNECT_CELL_CORNER_TO_INNER_LOOP
+		if (bDisplay)
+		{
+			F3DDebugSession _(bDisplay, FString::Printf(TEXT("Cell corner %d"), CellCorner.Id));
+			Grid.DisplayIsoNode(EGridSpace::UniformScaled, CellCorner.CornerNode, 0, EVisuProperty::RedPoint);
+		}
+#endif
+		for (FCellLoop& LoopCell : Cell.CellLoops)
+		{
+			FindAndTryCreateCandidateSegmentToLinkLoopToCorner(LoopCell, CellCorner);
+		}
+	}
+
+#ifdef DEBUG_CONNECT_CELL_CORNER_TO_INNER_LOOP
+	if (bDisplay)
+	{
+		DisplayCellConnexions(TEXT("Cell CandidateSegments"), Cell.LoopConnexions, EVisuProperty::YellowCurve);
+		Wait(false);
+	}
+#endif
+
+}
+
+
+void FIsoTriangulator::SelectSegmentsToLinkInnerToLoop()
+{
+#ifdef DEBUG_ADD_CANDIDATE_SEGMENTS_TO_LINK_INNER_AND_LOOP
+	F3DDebugSession _(Grid.bDisplay, TEXT("AddCandidateSegmentsToLinkInnerAndLoop "));
+	Grid.DisplayIsoSegments(TEXT("CandidateSegments To Link Inner To Loops"), EGridSpace::UniformScaled, InnerToLoopCandidateSegments, false, false, false, EVisuProperty::RedCurve);
+#endif
+
+	LoopSegmentsIntersectionTool.AddSegments(FinalToLoops);
+#ifdef DEBUG_ADD_CANDIDATE_SEGMENTS_TO_LINK_INNER_AND_LOOP
+	LoopSegmentsIntersectionTool.Display(Grid.bDisplay, TEXT("LoopSegmentsIntersectionTool with final to loop"));
+	Wait();
+#endif
+
+	TArray<TPair<double, FIsoSegment*>> LengthOfCandidateSegments;
+	LengthOfCandidateSegments.Reserve(InnerToLoopCandidateSegments.Num());
+	for (FIsoSegment* Segment : InnerToLoopCandidateSegments)
+	{
+		LengthOfCandidateSegments.Emplace(Segment->Get2DLengthSquare(EGridSpace::UniformScaled, Grid), Segment);
+	}
+
+	FIntersectionSegmentTool IntersectionTool(Grid, Tolerances.GeometricTolerance);
+	IntersectionTool.Reserve(InnerToLoopCandidateSegments.Num());
+
+	Algo::Sort(LengthOfCandidateSegments, [&](const TPair<double, FIsoSegment*>& P1, const TPair < double, FIsoSegment*>& P2) { return P1.Key < P2.Key; });
+
+	// Validate the first candidate segments
+	for (const TPair<double, FIsoSegment*>& Candidate : LengthOfCandidateSegments)
+	{
+		FIsoSegment* Segment = Candidate.Value;
+
+		if (IntersectionTool.DoesIntersect(*Segment))
+		{
+			IsoSegmentFactory.DeleteEntity(Segment);
+			continue;
+		}
+
+		if (LoopSegmentsIntersectionTool.DoesIntersect(*Segment))
+		{
+			IsoSegmentFactory.DeleteEntity(Segment);
+			continue;
+		}
+
+		if (FIsoSegment::IsItAlreadyDefined(&Segment->GetFirstNode(), &Segment->GetSecondNode()))
+		{
+			IsoSegmentFactory.DeleteEntity(Segment);
+			continue;
+		}
+
+		FinalToLoops.Add(Segment);
+		IntersectionTool.AddSegment(*Segment);
+		Segment->SetSelected();
+		Segment->SetFinalMarker();
+		if (!Segment->ConnectToNode())
+		{
+#ifdef DEBUG_BUILD_SEGMENT_IF_NEEDED
+			F3DDebugSession _(FString::Printf(TEXT("ERROR Segment")));
+			Grid.DisplayIsoSegment(EGridSpace::UniformScaled, *Segment, 0, EVisuProperty::RedCurve);
+			Wait();
+#endif
+			IsoSegmentFactory.DeleteEntity(Segment);
+		}
+	}
+	CandidateSegments.Empty();
+}
+
+} //namespace UE::CADKernel
\ No newline at end of file
