Moved CADKernel library code from /Engine/Source/Runtime/Datasmith/CADKernel to /Engine/Source/Runtime/Datasmith/CADKernel/Base

diff --git a/Engine/Source/Runtime/Datasmith/CADKernel/Base/Private/Mesh/Meshers/ParametricFaceMesher.cpp b/Engine/Source/Runtime/Datasmith/CADKernel/Base/Private/Mesh/Meshers/ParametricFaceMesher.cpp
new file mode 100644
--- /dev/null
+++ b/Engine/Source/Runtime/Datasmith/CADKernel/Base/Private/Mesh/Meshers/ParametricFaceMesher.cpp
@@ -0,0 +1,1385 @@
+// Copyright Epic Games, Inc. All Rights Reserved.
+
+#include "Mesh/Meshers/ParametricFaceMesher.h"
+
+#include "Mesh/Criteria/Criterion.h"
+#include "Mesh/Meshers/IsoTriangulator.h"
+#include "Mesh/Meshers/MesherTools.h"
+#include "Mesh/Meshers/ParametricMesherConstantes.h"
+#include "Mesh/Structure/EdgeMesh.h"
+#include "Mesh/Structure/FaceMesh.h"
+#include "Mesh/Structure/ModelMesh.h"
+#include "Mesh/Structure/ThinZone2DFinder.h"
+#include "Mesh/Structure/ThinZone2D.h"
+#include "Mesh/Structure/VertexMesh.h"
+#include "Topo/TopologicalEdge.h"
+#include "Topo/TopologicalFace.h"
+
+
+#ifdef CADKERNEL_DEV
+#include "Mesh/Meshers/MesherReport.h"
+#endif
+
+namespace UE::CADKernel
+{
+
+FParametricFaceMesher::FParametricFaceMesher(FTopologicalFace& InFace, FModelMesh& InMeshModel, const FMeshingTolerances& InTolerances, bool bActivateThinZoneMeshing)
+	: Face(InFace)
+	, MeshModel(InMeshModel)
+	, Tolerances(InTolerances)
+#if CADKERNEL_THINZONE
+	, bThinZoneMeshing(bActivateThinZoneMeshing)
+#endif
+	, Grid(Face, MeshModel)
+{
+#ifdef DEBUG_ONLY_SURFACE_TO_DEBUG
+	bDisplay = (Face.GetId() == FaceToDebug);
+#endif
+}
+
+void FParametricFaceMesher::Mesh()
+{
+#ifdef DEBUG_ONLY_SURFACE_TO_DEBUG
+	bDisplay = (Face.GetId() == FaceToDebug);
+#endif
+
+	Face.GetOrCreateMesh(MeshModel);
+
+	if (Face.IsNotMeshable())
+	{
+		return;
+	}
+
+	MeshVerticesOfFace(Face);
+
+	FMessage::Printf(EVerboseLevel::Debug, TEXT("Meshing of surface %d\n"), Face.GetId());
+
+	FProgress _(1, TEXT("Meshing Entities : Mesh Surface"));
+
+#ifdef DEBUG_CADKERNEL
+	F3DDebugSession S(bDisplayDebugMeshStep, FString::Printf(TEXT("Mesh of surface %d"), Face.GetId()));
+#endif
+
+	FTimePoint StartTime = FChrono::Now();
+
+	if (!GenerateCloud() || Grid.IsDegenerated())
+	{
+#ifdef CADKERNEL_DEV
+		FMesherReport::GetLogs().AddDegeneratedGrid();
+#endif
+		FMessage::Printf(EVerboseLevel::Log, TEXT("The meshing of the surface %d failed due to a degenerated grid\n"), Face.GetId());
+
+		Face.IsDegenerated();
+		Face.SetMeshedMarker();
+		return;
+	}
+
+	FTimePoint IsoTriangulerStartTime = FChrono::Now();
+
+	FFaceMesh& SurfaceMesh = Face.GetOrCreateMesh(MeshModel);
+	FIsoTriangulator IsoTrianguler(Grid, SurfaceMesh, Tolerances);
+
+	if (IsoTrianguler.Triangulate())
+	{
+		if (Face.IsBackOriented())
+		{
+			SurfaceMesh.InverseOrientation();
+		}
+		MeshModel.AddMesh(SurfaceMesh);
+	}
+	Face.SetMeshedMarker();
+
+	FDuration TriangulateDuration = FChrono::Elapse(IsoTriangulerStartTime);
+	FDuration Duration = FChrono::Elapse(StartTime);
+
+#ifdef CADKERNEL_DEV
+	FMesherReport::GetChronos().GlobalPointCloudDuration += Grid.Chronos.GeneratePointCloudDuration;
+	FMesherReport::GetChronos().GlobalTriangulateDuration += TriangulateDuration;
+	FMesherReport::GetChronos().GlobalMeshDuration += Duration;
+#endif
+
+	FChrono::PrintClockElapse(EVerboseLevel::Debug, TEXT("   "), TEXT("Meshing"), Duration);
+
+}
+
+bool FParametricFaceMesher::GenerateCloud()
+{
+	FTimePoint GenerateCloudStartTime = FChrono::Now();
+	Grid.DefineCuttingParameters();
+	if (!Grid.GeneratePointCloud())
+	{
+		return false;
+	}
+
+	if (bThinZoneMeshing)
+	{
+		FTimePoint StartTime = FChrono::Now();
+		if (Grid.GetFace().HasThinZone())
+		{
+			MeshThinZones();
+		}
+#ifdef CADKERNEL_DEV
+		FMesherReport::GetChronos().GlobalMeshThinZones += FChrono::Elapse(StartTime);
+#endif
+	}
+
+	MeshFaceLoops();
+
+	Grid.ProcessPointCloud();
+
+#ifdef CADKERNEL_DEV
+	FMesherReport::GetChronos().GlobalGeneratePointCloudDuration += FChrono::Elapse(GenerateCloudStartTime);
+#endif
+
+	return true;
+}
+
+void FParametricFaceMesher::MeshThinZones(TArray<FTopologicalEdge*>& EdgesToMesh, const bool bFinalMeshing)
+{
+	{
+#ifdef DEBUG_MESH_THINZONE_EDGES
+		F3DDebugSession A(bDisplay, TEXT("   - Define Cutting Points"));
+#endif
+		for (FTopologicalEdge* Edge : EdgesToMesh)
+		{
+			if (Edge->IsPreMeshed())
+			{
+				continue;
+			}
+
+#ifdef DEBUG_MESH_THINZONE_EDGES
+			if(bDisplay)
+			{
+				F3DDebugSession B(bDisplay, FString::Printf(TEXT("Edge %d"), Edge->GetId()));
+				Display2DWithScale(*Edge, EVisuProperty::GreenCurve);
+			}
+#endif
+
+			{
+#ifdef DEBUG_MESH_THINZONE_EDGES
+				F3DDebugSession B(bDisplay, FString::Printf(TEXT("Edge %d"), Edge->GetId()));
+#endif
+				for (FThinZoneSide* ZoneSide : Edge->GetThinZoneSides())
+				{
+					if (ZoneSide->HasMarker1())
+					{
+						continue;
+					}
+					ZoneSide->SetMarker1();
+
+					DefineImposedCuttingPointsBasedOnOtherSideMesh(*ZoneSide);
+				}
+			}
+		}
+
+		for (FTopologicalEdge* Edge : EdgesToMesh)
+		{
+			for (FThinZoneSide* ZoneSide : Edge->GetThinZoneSides())
+			{
+				if (ZoneSide->HasMarker1())
+				{
+					ZoneSide->ResetMarker1();
+				}
+			}
+		}
+	}
+
+	{
+#ifdef DEBUG_MESH_THINZONE_EDGES
+		F3DDebugSession A(bDisplay, TEXT("   - Mesh"));
+#endif
+		for (FTopologicalEdge* Edge : EdgesToMesh)
+		{
+			if (Edge->IsMeshed())
+			{
+				continue;
+			}
+
+			Mesh(*Edge, bFinalMeshing);
+		}
+	}
+	//Wait();
+}
+
+namespace ThinZoneMesherTools
+{
+void ResetMarkers(TArray<FTopologicalEdge*>& EdgesWithThinZones, TArray<FThinZone2D>& ThinZones)
+{
+	for (FTopologicalEdge* Edge : EdgesWithThinZones)
+	{
+		Edge->ResetMarkers();
+	}
+
+	for (FThinZone2D& Zone : ThinZones)
+	{
+		Zone.ResetMarkers();
+		Zone.GetFirstSide().ResetMarkers();
+		Zone.GetSecondSide().ResetMarkers();
+	}
+};
+};
+
+void FParametricFaceMesher::SortThinZoneSides(TArray<FThinZone2D*>& ThinZones)
+{
+#ifdef DEBUG_SORT_THIN_ZONE_SIDES
+	F3DDebugSession A(bDisplay, ("SortThinZoneSides"));
+#endif
+
+	Face.ResetMarkersRecursively();
+	int32 EdgeCount = Face.EdgeCount();
+	ZoneAEdges.Reserve(EdgeCount);
+	ZoneBEdges.Reserve(EdgeCount);
+
+	WaitingThinZones.Reserve(ThinZones.Num());
+
+	TFunction<void(FThinZone2D*)> AddToWaitingList = [&WaitingList = WaitingThinZones](FThinZone2D* Zone)
+	{
+		Zone->SetWaitingMarker();
+		WaitingList.Add(Zone);
+	};
+
+	TFunction<void(FThinZone2D*)> SetAndGet = [&WaitingList = WaitingThinZones, &ZoneA = ZoneAEdges, &ZoneB = ZoneBEdges](FThinZone2D* Zone)
+	{
+		Zone->SetEdgesZoneSide();
+		Zone->GetEdges(ZoneA, ZoneB);
+	};
+
+
+	int32 Index = 1;
+	for (FThinZone2D* Zone : ThinZones)
+	{
+#ifdef DEBUG_SORT_THIN_ZONE_SIDES
+		Zone.Display(FString::Printf(TEXT("Zone %d"), Index++), EVisuProperty::YellowCurve);
+#endif
+
+		Zone->CheckEdgesZoneSide();
+
+
+		if (Zone->GetFirstSide().HasMarker1And2())
+		{
+			AddToWaitingList(Zone);
+			continue;
+		}
+
+		if (Zone->GetSecondSide().HasMarker1And2())
+		{
+			AddToWaitingList(Zone);
+			continue;
+		}
+
+		if (!Zone->GetFirstSide().HasMarker1Or2() && !Zone->GetSecondSide().HasMarker1Or2())
+		{
+			SetAndGet(Zone);
+			continue;
+		}
+
+		if (Zone->GetFirstSide().HasMarker1() && !Zone->GetSecondSide().HasMarker1())
+		{
+			SetAndGet(Zone);
+			continue;
+		}
+
+		if (Zone->GetFirstSide().HasMarker1() && Zone->GetSecondSide().HasMarker1())
+		{
+			AddToWaitingList(Zone);
+			continue;
+		}
+
+		if (!Zone->GetFirstSide().HasMarker2() && !Zone->GetSecondSide().HasMarker1())
+		{
+			SetAndGet(Zone);
+			continue;
+		}
+
+		if (Zone->GetFirstSide().HasMarker2() && Zone->GetSecondSide().HasMarker2()) // E
+		{
+			AddToWaitingList(Zone);
+			continue;
+		}
+
+		if (Zone->GetFirstSide().HasMarker2() && !Zone->GetSecondSide().HasMarker2()) // F
+		{
+			Zone->Swap();
+			SetAndGet(Zone);
+			continue;
+		}
+
+		if (!Zone->GetFirstSide().HasMarker1() && Zone->GetSecondSide().HasMarker1())
+		{
+			Zone->Swap();
+			SetAndGet(Zone);
+			continue;
+		}
+
+		ensureCADKernel(false);
+	}
+
+	Face.ResetMarkersRecursively();
+}
+
+void FParametricFaceMesher::MeshThinZones()
+{
+	TArray<FThinZone2D>& FaceThinZones = Face.GetThinZones();
+
+	if (FaceThinZones.IsEmpty())
+	{
+		return;
+	}
+
+	TArray<FThinZone2D*> ThinZones;
+	ThinZones.Reserve(FaceThinZones.Num());
+	for (FThinZone2D& FaceThinZone : FaceThinZones)
+	{
+		ThinZones.Add(&FaceThinZone);
+	}
+
+	while (ThinZones.Num())
+	{
+#ifdef DEBUG_THIN_ZONES
+		F3DDebugSession A(bDisplay, ("Mesh thin zone"));
+		DisplayMeshOfFaceLoop();
+#endif
+		int32 WaitingThinZoneCount = ThinZones.Num();
+
+		MeshThinZones(ThinZones);
+
+		if (WaitingThinZoneCount == WaitingThinZones.Num())
+		{
+			break;
+		}
+		ThinZones = MoveTemp(WaitingThinZones);
+	}
+}
+
+
+void FParametricFaceMesher::MeshThinZones(TArray<FThinZone2D*>& ThinZones)
+{
+	SortThinZoneSides(ThinZones);
+
+	TFunction<void(TArray<FTopologicalEdge*>)> TransfereCuttingPointsFromMeshedEdges = [](TArray<FTopologicalEdge*> Edges)
+	{
+		for (FTopologicalEdge* Edge : Edges)
+		{
+			FAddCuttingPointFunc AddCuttingPoint = [&Edge](const double Coordinate, const ECoordinateType Type, const FPairOfIndex OppositNodeIndices, const double DeltaU)
+			{
+				Edge->AddTwinsCuttingPoint(Coordinate, DeltaU);
+			};
+
+			constexpr bool bOnlyOppositNode = false;
+			Edge->TransferCuttingPointFromMeshedEdge(bOnlyOppositNode, AddCuttingPoint);
+		}
+	};
+
+	TransfereCuttingPointsFromMeshedEdges(ZoneAEdges);
+	TransfereCuttingPointsFromMeshedEdges(ZoneBEdges);
+
+#ifdef DEBUG_THIN_ZONES
+	if(bDisplay)
+	{
+		DisplayThinZoneEdges(TEXT("Side A"), ZoneAEdges, EVisuProperty::BlueCurve, EVisuProperty::PurpleCurve);
+		DisplayThinZoneEdges(TEXT("Side B"), ZoneBEdges, EVisuProperty::RedCurve, EVisuProperty::PinkCurve);
+
+		{
+			F3DDebugSession A(bDisplay, ("Waiting Zone"));
+			for(FThinZone2D* Zone : WaitingThinZones)
+			{
+				Zone->Display(TEXT("Waiting Zone"), EVisuProperty::YellowCurve);
+			}
+		}
+
+		//Wait();
+	}
+#endif
+
+	bool bFinalMeshing = false;
+	{
+#ifdef DEBUG_MESHTHINSURF
+		F3DDebugSession _(bDisplay, TEXT("Step 1"));
+#endif
+		MeshThinZones(ZoneAEdges, bFinalMeshing);
+	}
+
+	{
+#ifdef DEBUG_MESHTHINSURF
+		F3DDebugSession _(bDisplay, TEXT("Step 2"));
+#endif
+		bFinalMeshing = true;
+		MeshThinZones(ZoneBEdges, bFinalMeshing);
+	}
+
+	for (FTopologicalEdge* Edge : ZoneAEdges)
+	{
+		Edge->RemovePreMesh();
+	}
+
+	{
+#ifdef DEBUG_MESHTHINSURF
+		F3DDebugSession _(bDisplay, TEXT("Step 3"));
+#endif
+		MeshThinZones(ZoneAEdges, bFinalMeshing);
+	}
+
+#ifdef DEBUG_THIN_ZONES
+	{
+		DisplayMeshOfFaceLoop();
+	}
+	//Wait(bDisplay);
+#endif
+}
+
+void FParametricFaceMesher::DefineImposedCuttingPointsBasedOnOtherSideMesh(FThinZoneSide& SideToConstrain)
+{
+	using namespace ParametricMesherTool;
+
+	FThinZoneSide& FrontSide = SideToConstrain.GetFrontThinZoneSide();
+
+	TMap<int32, FPoint2D> ExistingMeshNodes;
+	TArray<FCrossZoneElement> CrossZoneElements;
+
+	FAddMeshNodeFunc AddToCrossZoneElements = [&CrossZoneElements](const int32 NodeIndice, const FPoint2D& MeshNode2D, double MeshingTolerance3D, const FEdgeSegment& EdgeSegment, const FPairOfIndex& OppositeNodeIndices)
+	{
+		if (CrossZoneElements.Num() && CrossZoneElements.Last().VertexId >= 0 && CrossZoneElements.Last().VertexId == NodeIndice)
+		{
+			CrossZoneElements.Last().Add(OppositeNodeIndices);
+		}
+		else
+		{
+			CrossZoneElements.Emplace(NodeIndice, MeshNode2D, MeshingTolerance3D, &EdgeSegment, OppositeNodeIndices);
+		}
+	};
+
+	FAddMeshNodeFunc AddToExistingMeshNodes = [&ExistingMeshNodes](const int32 NodeIndice, const FPoint2D& MeshNode2D, double MeshingTolerance3D, const FEdgeSegment& EdgeSegment, const FPairOfIndex& OppositeNodeIndices)
+	{
+		ExistingMeshNodes.Emplace(NodeIndice, MeshNode2D);
+	};
+
+	FReserveContainerFunc ReserveCrossZoneElements = [&CrossZoneElements](int32 MeshVertexCount)
+	{
+		CrossZoneElements.Reserve(MeshVertexCount);
+	};
+
+	FReserveContainerFunc ReserveExistingMeshNodes = [&ExistingMeshNodes](int32 MeshVertexCount)
+	{
+		ExistingMeshNodes.Reserve(MeshVertexCount);
+	};
+
+	SideToConstrain.GetExistingMeshNodes(Face, MeshModel, ReserveExistingMeshNodes, AddToExistingMeshNodes, /*bWithTolerance*/ false);
+	FrontSide.GetExistingMeshNodes(Face, MeshModel, ReserveCrossZoneElements, AddToCrossZoneElements, /*bWithTolerance*/ true);
+
+#ifdef DEBUG_MESHTHINSURF
+	if (bDisplay)
+	{
+		{
+			F3DDebugSession E(bDisplay, FString::Printf(TEXT("ThinZone")));
+			ThinZone::DisplayThinZoneSide(FrontSide, 0, EVisuProperty::BlueCurve);
+			ThinZone::DisplayThinZoneSide(SideToConstrain, 1, EVisuProperty::YellowPoint);
+			Wait(false);
+		}
+
+		F3DDebugSession A(TEXT("Existing Mesh nodes"));
+		{
+			F3DDebugSession A(TEXT("Existing Mesh nodes side to mesh"));
+			for (const TPair<int32, FPoint2D>& MeshNode : ExistingMeshNodes)
+			{
+				DisplayPoint2DWithScale(MeshNode.Value, EVisuProperty::RedPoint, MeshNode.Key);
+			}
+		}
+		{
+			F3DDebugSession A(TEXT("Existing Mesh nodes reference side"));
+			for (const FCrossZoneElement& MeshNode : CrossZoneElements)
+			{
+				DisplayPoint2DWithScale(MeshNode.VertexPoint2D, EVisuProperty::BluePoint, MeshNode.VertexId);
+				if (MeshNode.OppositeVertexIndices[0] >= 0)
+				{
+					FPoint2D* OppositeVertex = ExistingMeshNodes.Find(MeshNode.OppositeVertexIndices[0]);
+					if (OppositeVertex)
+					{
+						DisplayPoint2DWithScale(*OppositeVertex, EVisuProperty::RedPoint, MeshNode.OppositeVertexIndices[0]);
+						DisplaySegmentWithScale(*OppositeVertex, MeshNode.VertexPoint2D, 0, EVisuProperty::BlueCurve);
+					}
+				}
+			}
+		}
+		Wait(false);
+	}
+
+	F3DDebugSession A(bDisplay, TEXT("Find the best projection of existing mesh vertices"));
+#endif
+
+	const double MaxSquareThickness = FrontSide.GetMaxThickness() > SideToConstrain.GetMaxThickness() ? FMath::Square(3. * FrontSide.GetMaxThickness()) : FMath::Square(3. * SideToConstrain.GetMaxThickness());
+
+	// Find the best projection of existing mesh vertices (CrossZone Vertex)
+	for (FCrossZoneElement& CrossZoneElement : CrossZoneElements)
+	{
+		if (CrossZoneElement.OppositeVertexIndices[0] >= 0)
+		{
+			FPoint2D* OppositeVertex = ExistingMeshNodes.Find(CrossZoneElement.OppositeVertexIndices[0]);
+			if (OppositeVertex)
+			{
+				CrossZoneElement.OppositePoint2D = *OppositeVertex;
+				CrossZoneElement.SquareThickness = 0;
+			}
+			continue;
+		}
+
+#ifdef DEBUG_MESHTHINSURF
+		if (bDisplay)
+		{
+			F3DDebugSession A(TEXT("Point"));
+			DisplayPoint2DWithScale(CrossZoneElement.VertexPoint2D, EVisuProperty::BluePoint, CrossZoneElement.VertexId);
+			Wait(false);
+		}
+#endif
+
+		double MinSquareThickness = MaxSquareThickness;
+		FPoint2D ClosePoint;
+		FEdgeSegment* CloseSegment = nullptr;
+		double ClosePointCoordinate = -1;
+
+		const FPoint2D MeshNodeCoordinate = CrossZoneElement.VertexPoint2D;
+
+		for (FEdgeSegment& Segment : SideToConstrain.GetSegments())
+		{
+			// check the angle between segment and Middle-SegementStart, Middle-SegementEnd.
+			const double SlopeS = Segment.ComputeOrientedSlopeOf(MeshNodeCoordinate, Segment.GetExtemity(ELimit::Start));
+			const double SlopeE = Segment.ComputeOrientedSlopeOf(MeshNodeCoordinate, Segment.GetExtemity(ELimit::End));
+			if (SlopeE < 1. || SlopeS > 3.)
+			{
+#ifdef DEBUG_MESHTHINSURF_
+				if (bDisplay)
+				{
+					F3DDebugSession A(FString::Printf(TEXT("Not candidate slope end %f"), SlopeE));
+					DisplaySegmentWithScale(Segment.GetExtemity(Start), Segment.GetExtemity(End), 0, EVisuProperty::RedCurve);
+					DisplaySegmentWithScale(MeshNodeCoordinate, Segment.GetExtemity(End), 0, EVisuProperty::YellowCurve);
+					DisplaySegmentWithScale(MeshNodeCoordinate, Segment.GetExtemity(Start), 0, EVisuProperty::YellowCurve);
+				}
+#endif
+				continue;
+			}
+
+			double CoordSegmentU;
+			FPoint2D Projection = Segment.ProjectPoint(MeshNodeCoordinate, CoordSegmentU);
+
+
+#ifdef DEBUG_MESHTHINSURF_
+			if (bDisplay)
+			{
+				const double Slope = Segment.ComputeOrientedSlopeOf(MeshNodeCoordinate, Projection);
+				F3DDebugSession A(FString::Printf(TEXT("Close Point slope %f %f %f"), Slope, SlopeS, SlopeE));
+				DisplaySegmentWithScale(Segment.GetExtemity(Start), Segment.GetExtemity(End), 0, EVisuProperty::BlueCurve);
+				DisplaySegmentWithScale(MeshNodeCoordinate, Projection, 0, EVisuProperty::GreenCurve);
+				DisplayPoint2DWithScale(Projection, EVisuProperty::YellowPoint, CrossZoneElement.VertexId);
+				//Wait();
+			}
+#endif
+
+			const double SquareDistance = MeshNodeCoordinate.SquareDistance(Projection);
+			if (MinSquareThickness > SquareDistance)
+			{
+				FEdgeSegment* SegmentPtr = &Segment;
+				// Forbid the common extremity as candidate
+				{
+					constexpr double NearlyZero = 0. + DOUBLE_SMALL_NUMBER;
+					constexpr double NearlyOne = 1. - DOUBLE_SMALL_NUMBER;
+
+					const FEdgeSegment* CrossZoneSegment = CrossZoneElement.Segment;
+					if (SegmentPtr == CrossZoneSegment->GetPrevious() && CoordSegmentU > NearlyOne)
+					{
+						continue;
+					}
+					if (SegmentPtr == CrossZoneSegment->GetNext() && CoordSegmentU < NearlyZero)
+					{
+						continue;
+					}
+				}
+
+				MinSquareThickness = SquareDistance;
+				ClosePoint = Projection;
+				ClosePointCoordinate = CoordSegmentU;
+				CloseSegment = SegmentPtr;
+			}
+		}
+
+		if (CloseSegment)
+		{
+			CrossZoneElement.OppositePoint2D = ClosePoint;
+			CrossZoneElement.OppositeSegment = CloseSegment;
+			CrossZoneElement.OppositePointCoordinate = ClosePointCoordinate;
+			CrossZoneElement.SquareThickness = MinSquareThickness;
+		}
+	}
+
+	Algo::Sort(CrossZoneElements, [](FCrossZoneElement& ElementA, FCrossZoneElement& ElementB) { return ElementA.SquareThickness < ElementB.SquareThickness; });
+
+#ifdef DEBUG_MESHTHINSURF_
+	if (bDisplay)
+	{
+		F3DDebugSession A(TEXT("Result before intersection filtering"));
+		for (FCrossZoneElement& CrossZoneElement : CrossZoneElements)
+		{
+			//F3DDebugSession A(TEXT("Point and Close"));
+			DisplayPoint2DWithScale(CrossZoneElement.VertexPoint2D, EVisuProperty::BluePoint, CrossZoneElement.VertexId);
+			DisplayPoint2DWithScale(CrossZoneElement.OppositePoint2D, EVisuProperty::RedPoint);
+			DisplaySegmentWithScale(CrossZoneElement.VertexPoint2D, CrossZoneElement.OppositePoint2D);
+			if (CrossZoneElement.OppositeSegment)
+			{
+				DisplaySegmentWithScale(CrossZoneElement.OppositeSegment->GetExtemity(ELimit::End), CrossZoneElement.OppositeSegment->GetExtemity(ELimit::Start));
+			}
+		}
+		Wait();
+	}
+#endif
+
+	// Find candidates i.e. CrossZoneElement that are not in intersection with the sides and with selected CrossZoneElement
+	FIntersectionTool IntersectionTool(FrontSide.GetSegments(), SideToConstrain.GetSegments(), CrossZoneElements.Num());
+	for (FCrossZoneElement& CrossZoneElement : CrossZoneElements)
+	{
+		if (!CrossZoneElement.OppositeSegment)
+		{
+			continue;
+		}
+
+		if (!IntersectionTool.IsIntersectSides(CrossZoneElement))
+		{
+			IntersectionTool.AddCrossZoneElement(CrossZoneElement);
+			CrossZoneElement.bIsSelected = true;
+		}
+	}
+
+#ifdef DEBUG_MESHTHINSURF
+	if (bDisplay)
+	{
+		F3DDebugSession M(TEXT("Result final"));
+		for (FCrossZoneElement& CrossZoneElement : CrossZoneElements)
+		{
+			if (CrossZoneElement.bIsSelected)
+			{
+				//F3DDebugSession A(TEXT("Point and Close"));
+				DisplayPoint2DWithScale(CrossZoneElement.VertexPoint2D, EVisuProperty::BluePoint, CrossZoneElement.VertexId);
+				DisplayPoint2DWithScale(CrossZoneElement.OppositePoint2D, EVisuProperty::RedPoint);
+				DisplaySegmentWithScale(CrossZoneElement.VertexPoint2D, CrossZoneElement.OppositePoint2D, 0, EVisuProperty::YellowCurve);
+			}
+		}
+		Wait(false);
+	}
+#endif
+
+	// Add ImposedCuttingPointU
+	for (FCrossZoneElement& CrossZoneElement : CrossZoneElements)
+	{
+		if (CrossZoneElement.bIsSelected && (CrossZoneElement.OppositeVertexIndices[0] < 0))
+		{
+			FTopologicalEdge* OppositeEdge = CrossZoneElement.OppositeSegment->GetEdge();
+			if (OppositeEdge == nullptr)
+			{
+				continue;
+			}
+
+			const double OppositeCuttingPointU = CrossZoneElement.OppositeSegment->ComputeEdgeCoordinate(CrossZoneElement.OppositePointCoordinate);
+			const double DeltaU = CrossZoneElement.OppositeSegment->ComputeDeltaU(CrossZoneElement.Tolerance3D);
+
+			OppositeEdge->AddImposedCuttingPointU(OppositeCuttingPointU, CrossZoneElement.VertexId, DeltaU);
+		}
+	}
+}
+
+void FParametricFaceMesher::MeshFaceLoops()
+{
+#ifdef DEBUG_MESHTHINSURF
+	F3DDebugSession M(bDisplay, ("Mesh other edges"));
+#endif
+
+	for (const TSharedPtr<FTopologicalLoop>& Loop : Face.GetLoops())
+	{
+		for (const FOrientedEdge& Edge : Loop->GetEdges())
+		{
+			Mesh(*Edge.Entity);
+		}
+	}
+}
+
+void FParametricFaceMesher::Mesh(FTopologicalVertex& InVertex)
+{
+	InVertex.GetOrCreateMesh(MeshModel);
+}
+
+void FParametricFaceMesher::MeshVerticesOfFace(FTopologicalFace& FaceToProcess)
+{
+	for (const TSharedPtr<FTopologicalLoop>& Loop : FaceToProcess.GetLoops())
+	{
+		for (FOrientedEdge& Edge : Loop->GetEdges())
+		{
+			Mesh(*Edge.Entity->GetStartVertex());
+			Mesh(*Edge.Entity->GetEndVertex());
+		}
+	}
+}
+
+#ifdef DEBUG_INTERSECTEDGEISOS
+void DebugIntersectEdgeIsos(const FTopologicalFace& Face, const TArray<double>& IsoCoordinates, EIso TypeIso);
+#endif
+
+void FParametricFaceMesher::Mesh(FTopologicalEdge& InEdge, bool bFinalMeshing)
+{
+  	FTopologicalEdge& ActiveEdge = *InEdge.GetLinkActiveEntity();
+	if (ActiveEdge.IsMeshed())
+	{
+		if (ActiveEdge.GetMesh()->GetNodeCount() > 0)
+		{
+			return;
+		}
+
+		// In some case the 2d curve is a smooth curve and the 3d curve is a line and vice versa
+		// In the particular case where the both case are opposed, we can have the 2d line sampled with 4 points, and the 2d curve sampled with 2 points (because in 3d, the 2d curve is a 3d line)
+		// In this case, the loop is flat i.e. in 2d the meshes of the 2d line and 2d curve are coincident
+		// So the grid is degenerated and the surface is not meshed
+		// to avoid this case, the Edge is virtually meshed i.e. the nodes inside the edge have the id of the mesh of the vertices.
+		InEdge.SetVirtuallyMeshedMarker();
+	}
+
+	if (ActiveEdge.IsThinPeak())
+	{
+		TArray<FCuttingPoint>& FinalEdgeCuttingPointCoordinates = ActiveEdge.GetCuttingPoints();
+		FinalEdgeCuttingPointCoordinates.Emplace(ActiveEdge.GetStartCurvilinearCoordinates(), ECoordinateType::VertexCoordinate);
+		FinalEdgeCuttingPointCoordinates.Emplace(ActiveEdge.GetEndCurvilinearCoordinates(), ECoordinateType::VertexCoordinate);
+		ActiveEdge.GenerateMeshElements(MeshModel);
+		return;
+	}
+
+	const FSurfacicTolerance& ToleranceIso = Face.GetIsoTolerances();
+
+	// Get Edge intersection with inner surface mesh grid
+	TArray<double> EdgeIntersectionWithIsoU_Coordinates;
+	TArray<double> EdgeIntersectionWithIsoV_Coordinates;
+
+	const TArray<double>& SurfaceTabU = Face.GetCuttingCoordinatesAlongIso(EIso::IsoU);
+	const TArray<double>& SurfaceTabV = Face.GetCuttingCoordinatesAlongIso(EIso::IsoV);
+
+	ApplyEdgeCriteria(InEdge);
+
+#ifdef DEBUG_MESH_EDGE
+	if (bDisplay)
+	{
+		F3DDebugSession _(FString::Printf(TEXT("EdgePointsOnDomain %d"), InEdge.GetId()));
+		Display2D(InEdge);
+		Wait();
+	}
+#endif
+
+#ifdef DEBUG_INTERSECTEDGEISOS
+	DebugIntersectEdgeIsos(Face, SurfaceTabU, EIso::IsoU);
+	DebugIntersectEdgeIsos(Face, SurfaceTabV, EIso::IsoV);
+	{
+		F3DDebugSession _(FString::Printf(TEXT("Edge 2D %d"), InEdge.GetId()));
+		UE::CADKernel::Display2D(InEdge);
+	}
+#endif
+
+	InEdge.ComputeIntersectionsWithIsos(SurfaceTabU, EIso::IsoU, ToleranceIso, EdgeIntersectionWithIsoU_Coordinates);
+	InEdge.ComputeIntersectionsWithIsos(SurfaceTabV, EIso::IsoV, ToleranceIso, EdgeIntersectionWithIsoV_Coordinates);
+
+#ifdef DEBUG_INTERSECTEDGEISOS
+	{
+		F3DDebugSession _(FString::Printf(TEXT("Edge %d Intersect with iso"), InEdge.GetId()));
+		TArray<FPoint2D> Intersections;
+		InEdge.Approximate2DPoints(EdgeIntersectionWithIsoU_Coordinates, Intersections);
+		for (const FPoint2D& Point : Intersections)
+		{
+			DisplayPoint(Point);
+		}
+		Intersections.Empty();
+		InEdge.Approximate2DPoints(EdgeIntersectionWithIsoV_Coordinates, Intersections);
+		for (const FPoint2D& Point : Intersections)
+		{
+			DisplayPoint(Point);
+		}
+		Wait();
+	}
+	{
+		F3DDebugSession _(FString::Printf(TEXT("Thin Zone")));
+		for (FLinearBoundary ThinZone : InEdge.GetThinZoneBounds())
+		{
+			TArray<double> Coords;
+			Coords.Add(ThinZone.GetMin());
+			Coords.Add(ThinZone.GetMax());
+			TArray<FPoint2D> ThinZone2D;
+			InEdge.Approximate2DPoints(Coords, ThinZone2D);
+			DisplaySegment(ThinZone2D[0], ThinZone2D[1], EVisuProperty::YellowCurve);
+		}
+		Wait();
+
+	}
+#endif
+
+	FLinearBoundary EdgeBounds = InEdge.GetBoundary();
+
+	TArray<double>& DeltaUs = InEdge.GetDeltaUMaxs();
+
+	FAddCuttingPointFunc AddCuttingPoint = [&InEdge](const double Coordinate, const ECoordinateType Type, const FPairOfIndex OppositNodeIndices, const double DeltaU)
+	{
+		for (int32 Index = 0; Index < 2; ++Index)
+		{
+			if (OppositNodeIndices[Index] >= 0)
+			{
+				InEdge.AddImposedCuttingPointU(Coordinate, OppositNodeIndices[Index], DeltaU);
+			}
+		}
+	};
+
+	// Case of self connected surface (e.g. cylinder) an edge 
+	// The first edge is premeshed at step 1, but the activeEdge is not yet meshed
+	// the twin edge is meshed at step 2
+	if (ActiveEdge.IsPreMeshed())
+	{
+		constexpr bool bOnlyOppositNode = true;
+		InEdge.TransferCuttingPointFromMeshedEdge(bOnlyOppositNode, AddCuttingPoint);
+
+		FTopologicalEdge* PreMeshEdge = InEdge.GetPreMeshedTwin();
+		if(PreMeshEdge)
+		{
+			PreMeshEdge->RemovePreMesh();
+		}
+	}
+
+	InEdge.SortImposedCuttingPoints();
+	const TArray<FImposedCuttingPoint>& EdgeImposedCuttingPoints = InEdge.GetImposedCuttingPoints();
+
+	// build a edge mesh compiling inner surface cutting (based on criteria applied on the surface) and edge cutting (based on criteria applied on the curve)
+	TArray<FCuttingPoint> ImposedIsoCuttingPoints;
+
+	TFunction<void(int32&, int32&)> UpdateDeltaU = [&ImposedIsoCuttingPoints](int32& NewIndex, int32& Index)
+	{
+		if (ImposedIsoCuttingPoints[NewIndex].IsoDeltaU > ImposedIsoCuttingPoints[Index].IsoDeltaU)
+		{
+			ImposedIsoCuttingPoints[NewIndex].IsoDeltaU = ImposedIsoCuttingPoints[Index].IsoDeltaU;
+		}
+	};
+
+	TFunction<void(int32&, int32&)> UpdateOppositNodeIndices = [&ImposedIsoCuttingPoints](int32& NewIndex, int32& Index)
+	{
+		if (ImposedIsoCuttingPoints[NewIndex].OppositNodeIndices[0] == -1)
+		{
+			ImposedIsoCuttingPoints[NewIndex].OppositNodeIndices[0] = ImposedIsoCuttingPoints[Index].OppositNodeIndices[0];
+		}
+		else if (ImposedIsoCuttingPoints[NewIndex].OppositNodeIndices[0] != ImposedIsoCuttingPoints[Index].OppositNodeIndices[0])
+		{
+			ImposedIsoCuttingPoints[NewIndex].OppositNodeIndices[1] = ImposedIsoCuttingPoints[Index].OppositNodeIndices[0];
+		}
+	};
+
+	TFunction<void(int32&, int32&, ECoordinateType)> MergeImposedCuttingPoints = [&ImposedIsoCuttingPoints, UpdateOppositNodeIndices, UpdateDeltaU](int32& Index, int32& NewIndex, ECoordinateType NewType)
+	{
+		double DeltaU = FMath::Max(ImposedIsoCuttingPoints[NewIndex].IsoDeltaU, ImposedIsoCuttingPoints[Index].IsoDeltaU);
+
+		if (ImposedIsoCuttingPoints[NewIndex].Coordinate + DeltaU > ImposedIsoCuttingPoints[Index].Coordinate)
+		{
+			if (ImposedIsoCuttingPoints[Index].Type == VertexCoordinate)
+			{
+				ImposedIsoCuttingPoints[NewIndex].Coordinate = ImposedIsoCuttingPoints[Index].Coordinate;
+				ImposedIsoCuttingPoints[NewIndex].IsoDeltaU = ImposedIsoCuttingPoints[Index].IsoDeltaU;
+				ImposedIsoCuttingPoints[NewIndex].Type = ImposedIsoCuttingPoints[Index].Type;
+				UpdateOppositNodeIndices(NewIndex, Index);
+				UpdateDeltaU(NewIndex, Index);
+			}
+			else if (ImposedIsoCuttingPoints[NewIndex].Type == VertexCoordinate)
+			{
+				if (ImposedIsoCuttingPoints[Index].Type == ImposedCoordinate)
+				{
+					UpdateOppositNodeIndices(NewIndex, Index);
+					UpdateDeltaU(NewIndex, Index);
+				}
+			}
+			else if (ImposedIsoCuttingPoints[NewIndex].Type == ImposedCoordinate)
+			{
+				if (ImposedIsoCuttingPoints[Index].Type == ImposedCoordinate)
+				{
+					ImposedIsoCuttingPoints[NewIndex].Coordinate = (ImposedIsoCuttingPoints[NewIndex].Coordinate + ImposedIsoCuttingPoints[Index].Coordinate) * 0.5;
+					UpdateOppositNodeIndices(NewIndex, Index);
+					UpdateDeltaU(NewIndex, Index);
+				}
+			}
+			else if (ImposedIsoCuttingPoints[Index].Type == ImposedCoordinate)
+			{
+				ImposedIsoCuttingPoints[NewIndex].Coordinate = ImposedIsoCuttingPoints[Index].Coordinate;
+				ImposedIsoCuttingPoints[NewIndex].Type = ImposedCoordinate;
+				UpdateOppositNodeIndices(NewIndex, Index);
+				UpdateDeltaU(NewIndex, Index);
+			}
+			else if (ImposedIsoCuttingPoints[NewIndex].Type != ImposedIsoCuttingPoints[Index].Type)
+			{
+				ImposedIsoCuttingPoints[NewIndex].Coordinate = (ImposedIsoCuttingPoints[NewIndex].Coordinate + ImposedIsoCuttingPoints[Index].Coordinate) * 0.5;
+				ImposedIsoCuttingPoints[NewIndex].Type = IsoUVCoordinate;
+				UpdateDeltaU(NewIndex, Index);
+			}
+		}
+		else
+		{
+			++NewIndex;
+			ImposedIsoCuttingPoints[NewIndex] = ImposedIsoCuttingPoints[Index];
+		}
+	};
+
+	{
+		int32 NbImposedCuttingPoints = EdgeImposedCuttingPoints.Num() + EdgeIntersectionWithIsoU_Coordinates.Num() + EdgeIntersectionWithIsoV_Coordinates.Num() + 2;
+		ImposedIsoCuttingPoints.Reserve(NbImposedCuttingPoints);
+	}
+
+
+	const double EdgeBoundsLength = EdgeBounds.Length();
+	const double EdgeDeltaUAtMin = FMath::Min(DeltaUs[0] * AQuarter, EdgeBoundsLength * AEighth);
+	const double EdgeDeltaUAtMax = FMath::Min(DeltaUs.Last() * AQuarter, EdgeBoundsLength * AEighth);
+
+	ImposedIsoCuttingPoints.Emplace(EdgeBounds.GetMin(), ECoordinateType::VertexCoordinate, FPairOfIndex::Undefined, EdgeDeltaUAtMin);
+	ImposedIsoCuttingPoints.Emplace(EdgeBounds.GetMax(), ECoordinateType::VertexCoordinate, FPairOfIndex::Undefined, EdgeDeltaUAtMax);
+
+	int32 Index = 0;
+	for (const FImposedCuttingPoint& CuttingPoint : EdgeImposedCuttingPoints)
+	{
+		const double CuttingPointDeltaU = CuttingPoint.DeltaU;
+		ImposedIsoCuttingPoints.Emplace(CuttingPoint.Coordinate, ECoordinateType::ImposedCoordinate, CuttingPoint.OppositNodeIndex, CuttingPointDeltaU * AThird);
+	}
+
+	// Add Edge intersection with inner surface grid Iso
+	FPoint2D ExtremityTolerances = InEdge.GetCurve()->GetExtremityTolerances(EdgeBounds);
+	double EdgeTolerance = FMath::Min(ExtremityTolerances[0], ExtremityTolerances[1]);
+	if (!EdgeIntersectionWithIsoU_Coordinates.IsEmpty())
+	{
+		FMesherTools::FillImposedIsoCuttingPoints(EdgeIntersectionWithIsoU_Coordinates, IsoUCoordinate, EdgeTolerance, InEdge, ImposedIsoCuttingPoints);
+	}
+
+	if (!EdgeIntersectionWithIsoV_Coordinates.IsEmpty())
+	{
+		FMesherTools::FillImposedIsoCuttingPoints(EdgeIntersectionWithIsoV_Coordinates, IsoVCoordinate, EdgeTolerance, InEdge, ImposedIsoCuttingPoints);
+	}
+
+	ImposedIsoCuttingPoints.Sort([](const FCuttingPoint& Point1, const FCuttingPoint& Point2) { return Point1.Coordinate < Point2.Coordinate; });
+
+	// If a pair of point isoU/isoV is too close, get the middle of the points
+	if (ImposedIsoCuttingPoints.Num() > 1)
+	{
+		int32 NewIndex = 0;
+		for (int32 Andex = 1; Andex < ImposedIsoCuttingPoints.Num(); ++Andex)
+		{
+			if (ImposedIsoCuttingPoints[Andex].Type > ECoordinateType::ImposedCoordinate)
+			{
+				bool bIsDelete = false;
+				for (const FLinearBoundary& ThinZone : InEdge.GetThinZoneBounds())
+				{
+					if (ThinZone.Contains(ImposedIsoCuttingPoints[Andex].Coordinate))
+					{
+						bIsDelete = true;
+						break; // or continue
+					}
+				}
+				if (bIsDelete)
+				{
+					continue;
+				}
+			}
+
+			if (ImposedIsoCuttingPoints[NewIndex].Type == ECoordinateType::ImposedCoordinate || ImposedIsoCuttingPoints[Andex].Type == ECoordinateType::ImposedCoordinate)
+			{
+				MergeImposedCuttingPoints(Andex, NewIndex, ECoordinateType::ImposedCoordinate);
+			}
+			else if (ImposedIsoCuttingPoints[NewIndex].Type != ImposedIsoCuttingPoints[Andex].Type)
+			{
+				MergeImposedCuttingPoints(Andex, NewIndex, ECoordinateType::IsoUVCoordinate);
+			}
+			else
+			{
+				++NewIndex;
+				ImposedIsoCuttingPoints[NewIndex] = ImposedIsoCuttingPoints[Andex];
+			}
+		}
+		ImposedIsoCuttingPoints.SetNum(NewIndex + 1);
+	}
+
+	if (ImposedIsoCuttingPoints.Num() > 1 && (EdgeBounds.GetMax() - ImposedIsoCuttingPoints.Last().Coordinate) < FMath::Min(ImposedIsoCuttingPoints.Last().IsoDeltaU, InEdge.GetDeltaUMaxs().Last()))
+	{
+		ImposedIsoCuttingPoints.Last().Coordinate = EdgeBounds.GetMax();
+		ImposedIsoCuttingPoints.Last().Type = VertexCoordinate;
+	}
+	else
+	{
+		ImposedIsoCuttingPoints.Emplace(EdgeBounds.GetMax(), ECoordinateType::VertexCoordinate, -1, InEdge.GetDeltaUMaxs().Last() * AQuarter);
+	}
+
+	// Final array of the edge mesh vertex 
+	TArray<FCuttingPoint>& FinalEdgeCuttingPointCoordinates = InEdge.GetCuttingPoints();
+	{
+		// max count of vertex
+		double MinDeltaU = HUGE_VALUE;
+		for (const double& DeltaU : DeltaUs)
+		{
+			if (DeltaU < MinDeltaU)
+			{
+				MinDeltaU = DeltaU;
+			}
+		}
+
+		int32 MaxNumberOfVertex = FMath::IsNearlyZero(MinDeltaU) ? 5 : (int32)((EdgeBounds.GetMax() - EdgeBounds.GetMin()) / MinDeltaU) + 5;
+		FinalEdgeCuttingPointCoordinates.Empty(ImposedIsoCuttingPoints.Num() + MaxNumberOfVertex);
+	}
+
+#ifdef DEBUG_GETPREFERREDUVCOORDINATESFROMNEIGHBOURS
+	TArray<double> CuttingPoints2;
+	{
+		double ToleranceGeoEdge = FMath::Min(ExtremityTolerances[0], ExtremityTolerances[1]);
+
+		TArray<FCuttingPoint> Extremities;
+		Extremities.Reserve(2);
+		Extremities.Emplace(EdgeBounds.GetMin(), ECoordinateType::VertexCoordinate, -1, ToleranceGeoEdge);
+		Extremities.Emplace(EdgeBounds.GetMax(), ECoordinateType::VertexCoordinate, -1, ToleranceGeoEdge);
+
+		FMesherTools::ComputeFinalCuttingPointsWithImposedCuttingPoints(InEdge.GetCrossingPointUs(), InEdge.GetDeltaUMaxs(), Extremities, CuttingPoints2);
+	}
+#endif
+
+	if (InEdge.IsDegenerated() || InEdge.IsVirtuallyMeshed())
+	{
+		if (ImposedIsoCuttingPoints.Num() == 2)
+		{
+			ImposedIsoCuttingPoints.EmplaceAt(1, (ImposedIsoCuttingPoints[0].Coordinate + ImposedIsoCuttingPoints[1].Coordinate) * 0.5, ECoordinateType::OtherCoordinate);
+		}
+
+		for (FCuttingPoint CuttingPoint : ImposedIsoCuttingPoints)
+		{
+			FinalEdgeCuttingPointCoordinates.Emplace(CuttingPoint.Coordinate, ECoordinateType::OtherCoordinate);
+		}
+		InEdge.GetLinkActiveEdge()->SetMeshedMarker();
+		return;
+	}
+
+	TArray<double> CuttingPoints;
+	FMesherTools::ComputeFinalCuttingPointsWithImposedCuttingPoints(InEdge.GetCrossingPointUs(), InEdge.GetDeltaUMaxs(), ImposedIsoCuttingPoints, CuttingPoints);
+	int32 ImposedIndex = 0;
+	int32 ImposedIsoCuttingPointsCount = ImposedIsoCuttingPoints.Num();
+	for (const double& Coordinate : CuttingPoints)
+	{
+		if (FMath::IsNearlyEqual(ImposedIsoCuttingPoints[ImposedIndex].Coordinate, Coordinate))
+		{
+			FinalEdgeCuttingPointCoordinates.Emplace(ImposedIsoCuttingPoints[ImposedIndex]);
+			++ImposedIndex;
+		}
+		else
+		{
+			while (ImposedIndex < ImposedIsoCuttingPointsCount && ImposedIsoCuttingPoints[ImposedIndex].Coordinate < Coordinate)
+			{
+				++ImposedIndex;
+			}
+			FinalEdgeCuttingPointCoordinates.Emplace(Coordinate, ECoordinateType::OtherCoordinate);
+		}
+	}
+
+#ifdef DEBUG_GETPREFERREDUVCOORDINATESFROMNEIGHBOURS
+	if (InEdge.IsThinZone() && EdgeImposedCuttingPoints.Num())
+	{
+		F3DDebugSession G(TEXT("Mesh(TSharedRef<FEdge> InEdge"));
+		{
+			F3DDebugSession G(TEXT("U From Iso"));
+			for (const FCuttingPoint& CuttingU : ImposedIsoCuttingPoints)
+			{
+				if (CuttingU.OppositNodeIndex >= 0)
+				{
+					UE::CADKernel::DisplayPoint(FPoint2D(CuttingU.Coordinate, 0.0), EVisuProperty::RedPoint);
+				}
+				else
+				{
+					UE::CADKernel::DisplayPoint(FPoint2D(CuttingU.Coordinate, 0.0));
+				}
+			}
+		}
+		{
+			F3DDebugSession _(InEdge.GetThinZoneBounds().Num() != 0, FString::Printf(TEXT("Thin Zone"), InEdge.GetId()));
+			for (FLinearBoundary ThinZone : InEdge.GetThinZoneBounds())
+			{
+				UE::CADKernel::DisplayPoint(FPoint2D(ThinZone.GetMin(), 0.01), EVisuProperty::BluePoint);
+				UE::CADKernel::DisplayPoint(FPoint2D(ThinZone.GetMax(), 0.01), EVisuProperty::BluePoint);
+				DisplaySegment(FPoint2D(ThinZone.GetMin(), 0.01), FPoint2D(ThinZone.GetMax(), 0.01), EVisuProperty::BlueCurve);
+			}
+			Wait();
+		}
+
+		{
+			F3DDebugSession G(TEXT("U From Criteria"));
+			for (double CuttingU : CuttingPoints2)
+			{
+				UE::CADKernel::DisplayPoint(FPoint2D(CuttingU, 0.02), EVisuProperty::RedPoint);
+			}
+		}
+		{
+			F3DDebugSession G(TEXT("U Final (Criteria & Iso)"));
+			for (double CuttingU : CuttingPoints)
+			{
+				UE::CADKernel::DisplayPoint(FPoint2D(CuttingU, 0.04), EVisuProperty::YellowPoint);
+			}
+		}
+		Wait();
+	}
+#endif
+
+	if (bFinalMeshing)
+	{
+		InEdge.GenerateMeshElements(MeshModel);
+	}
+	else
+	{
+		ActiveEdge.SetPreMeshedMarker();
+		InEdge.SetPreMeshedMarker();
+	}
+
+#ifdef DEBUG_THIN_ZONES
+	if (bFinalMeshing && bDisplay)
+	{
+#ifdef DEBUG_MESHTHINSURF
+		F3DDebugSession A(bDisplay, FString::Printf(TEXT("Edge %d"), InEdge.GetId()));
+#endif
+		Display2DWithScale(InEdge);
+		TArray<FPoint> MeshVertices;
+		if (InEdge.GetMesh())
+		{
+			MeshVertices = InEdge.GetMesh()->GetNodeCoordinates();
+		}
+		if(InEdge.GetStartVertex()->GetMesh())
+		{
+			MeshVertices.Append(InEdge.GetStartVertex()->GetMesh()->GetNodeCoordinates());
+		}
+		if(InEdge.GetEndVertex()->GetMesh())
+		{
+			MeshVertices.Append(InEdge.GetEndVertex()->GetMesh()->GetNodeCoordinates());
+		}
+		if (MeshVertices.Num())
+		{
+			TArray<double> CoordinatesOfMesh;
+			TArray<FPoint> ProjectedPoints;
+			TArray<FPoint2D> Mesh2DPoints;
+			InEdge.ProjectPoints(MeshVertices, CoordinatesOfMesh, ProjectedPoints);
+			Algo::Sort(CoordinatesOfMesh);
+			InEdge.Approximate2DPoints(CoordinatesOfMesh, Mesh2DPoints);
+			int32 Index = 0;
+			for (FPoint2D Point : Mesh2DPoints)
+			{
+				DisplayPoint2DWithScale(Point, EVisuProperty::RedPoint, ++Index);
+			}
+		}
+#ifdef DEBUG_MESHTHINSURF
+		//Wait(bDisplay);
+#endif
+	}
+#endif
+
+}
+
+void FParametricFaceMesher::ApplyEdgeCriteria(FTopologicalEdge& Edge)
+{
+	FTopologicalEdge& ActiveEdge = *Edge.GetLinkActiveEdge();
+
+	if (Edge.Length() < 2. * Tolerances.GeometricTolerance)
+	{
+		for (FTopologicalEdge* TwinEdge : Edge.GetTwinEntities())
+		{
+			TwinEdge->SetAsDegenerated();
+		}
+	}
+
+	Edge.ComputeCrossingPointCoordinates();
+	Edge.InitDeltaUs();
+	const TArray<double>& CrossingPointUs = Edge.GetCrossingPointUs();
+
+	TArray<double> Coordinates;
+	Coordinates.SetNum(CrossingPointUs.Num() * 2 - 1);
+	Coordinates[0] = CrossingPointUs[0];
+	for (int32 ICuttingPoint = 1; ICuttingPoint < Edge.GetCrossingPointUs().Num(); ICuttingPoint++)
+	{
+		Coordinates[2 * ICuttingPoint - 1] = (CrossingPointUs[ICuttingPoint - 1] + CrossingPointUs[ICuttingPoint]) * 0.5;
+		Coordinates[2 * ICuttingPoint] = CrossingPointUs[ICuttingPoint];
+	}
+
+	TArray<FCurvePoint> Points3D;
+	Edge.EvaluatePoints(Coordinates, 0, Points3D);
+
+	const TArray<TSharedPtr<FCriterion>>& Criteria = MeshModel.GetCriteria();
+	for (const TSharedPtr<FCriterion>& Criterion : Criteria)
+	{
+		Criterion->ApplyOnEdgeParameters(Edge, CrossingPointUs, Points3D);
+	}
+
+	Edge.SetApplyCriteriaMarker();
+	ActiveEdge.SetApplyCriteriaMarker();
+}
+
+void FParametricFaceMesher::MeshThinZoneSide(FThinZoneSide& Side, bool bFinalMeshing)
+{
+	if (!Side.HasMarker2())
+	{
+		return;
+	}
+
+	if (Side.IsProcessed())
+	{
+		return;
+	}
+	Side.ResetProcessedMarker();
+
+#ifdef DEBUG_MESHTHINSURF
+	F3DDebugSession A(bFinalMeshing && bDisplay, TEXT("SideMesh"));
+#endif
+	for (FTopologicalEdge* Edge : Side.GetEdges())
+	{
+		if (Edge->IsMeshed())
+		{
+			continue;
+		}
+		Mesh(*Edge, bFinalMeshing);
+#ifdef DEBUG_THIN_ZONES
+		if (bFinalMeshing && bDisplay)
+		{
+#ifdef DEBUG_MESHTHINSURF
+			F3DDebugSession A(bDisplay, FString::Printf(TEXT("Edge %d"), Edge->GetId()));
+#endif
+			Display2DWithScale(*Edge);
+			TArray<FPoint> MeshVertices = Edge->GetMesh()->GetNodeCoordinates();
+			MeshVertices.Append(Edge->GetStartVertex()->GetMesh()->GetNodeCoordinates());
+			MeshVertices.Append(Edge->GetEndVertex()->GetMesh()->GetNodeCoordinates());
+			if (MeshVertices.Num())
+			{
+				TArray<double> CoordinatesOfMesh;
+				TArray<FPoint> ProjectedPoints;
+				TArray<FPoint2D> Mesh2DPoints;
+				Edge->ProjectPoints(MeshVertices, CoordinatesOfMesh, ProjectedPoints);
+				Algo::Sort(CoordinatesOfMesh);
+				Edge->Approximate2DPoints(CoordinatesOfMesh, Mesh2DPoints);
+				int32 Index = 0;
+				for (FPoint2D Point : Mesh2DPoints)
+				{
+					DisplayPoint2DWithScale(Point, EVisuProperty::RedPoint, ++Index);
+				}
+			}
+#ifdef DEBUG_MESHTHINSURF
+					//Wait(bDisplay);
+#endif
+		}
+#endif
+	}
+#ifdef DEBUG_MESHTHINSURF
+	//Wait(bDisplay);
+#endif
+}
+
+#ifdef DEBUG_INTERSECTEDGEISOS
+TMap<int32, int32> SurfaceDrawed;
+bool bDisplayIsoCurve = true;
+
+void DebugIntersectEdgeIsos(const FTopologicalFace& Face, const TArray<double>& IsoCoordinates, EIso TypeIso)
+{
+	if (SurfaceDrawed.Find(Face.GetId()) == nullptr)
+	{
+		SurfaceDrawed.Add(Face.GetId(), 0);
+	}
+
+	if (bDisplayIsoCurve && SurfaceDrawed[Face.GetId()] < 2)
+	{
+		SurfaceDrawed[Face.GetId()]++;
+
+		FSurfacicBoundary Bounds = Face.GetBoundary();
+
+		F3DDebugSession _(FString::Printf(TEXT("Iso %s 2D %d"), TypeIso == EIso::IsoU ? TEXT("U") : TEXT("V"), Face.GetId()));
+		if (TypeIso == EIso::IsoU)
+		{
+			for (double U : IsoCoordinates)
+			{
+				FPoint2D Start(U, Bounds[EIso::IsoV].Min);
+				FPoint2D End(U, Bounds[EIso::IsoV].Max);
+				DisplaySegment(Start, End, 0, EVisuProperty::Iso);
+			}
+		}
+		else
+		{
+			for (double V : IsoCoordinates)
+			{
+				FPoint2D Start(Bounds[EIso::IsoU].Min, V);
+				FPoint2D End(Bounds[EIso::IsoU].Max, V);
+				DisplaySegment(Start, End, 0, EVisuProperty::Iso);
+			}
+		}
+	}
+}
+
+#endif
+
+#ifdef CADKERNEL_DEBUG
+
+void FParametricFaceMesher::DisplayMeshOfFaceLoop()
+{
+	if (bDisplay)
+	{
+		F3DDebugSession _(bDisplay, FString::Printf(TEXT("Mesh of Thin Face %d"), Face.GetId()));
+		for (const TSharedPtr<FTopologicalLoop>& Loop : Face.GetLoops())
+		{
+			for (const FOrientedEdge& OrientedEdge : Loop->GetEdges())
+			{
+				FTopologicalEdge& Edge = *OrientedEdge.Entity;
+				if (Edge.GetLinkActiveEdge()->IsMeshed())
+				{
+					const TArray<FPoint>& MeshVetrices = Edge.GetMesh()->GetNodeCoordinates();
+					TArray<FPoint> ProjectedPoints;
+					TArray<double> Coordinates;
+					Edge.ProjectPoints(MeshVetrices, Coordinates, ProjectedPoints);
+					Algo::Sort(Coordinates);
+					Coordinates.Insert(Edge.GetBoundary().GetMin(), 0);
+					Coordinates.Add(Edge.GetBoundary().GetMax());
+					TArray<FPoint2D> Points2D;
+					Edge.Approximate2DPoints(Coordinates, Points2D);
+
+					DisplayPolylineWithScale(Points2D, EVisuProperty::BlueCurve);
+					for (FPoint2D Point : Points2D)
+					{
+						DisplayPoint2DWithScale(Point);
+					}
+				}
+				else
+				{
+					Display2DWithScale(Edge, EVisuProperty::GreenCurve);
+				}
+			}
+		}
+		//Wait();
+	}
+}
+
+void FParametricFaceMesher::DisplayThinZoneEdges(const TCHAR* Text, TArray<FTopologicalEdge*>& Edges, EVisuProperty Color, EVisuProperty Color2)
+{
+	if (bDisplay)
+	{
+		F3DDebugSession _(bDisplay, Text);
+		for (const FTopologicalEdge* Edge : Edges)
+		{
+			if (Edge->GetLinkActiveEdge()->IsMeshed())
+			{
+				const TArray<FPoint>& MeshVetrices = Edge->GetMesh()->GetNodeCoordinates();
+				TArray<FPoint> ProjectedPoints;
+				TArray<double> Coordinates;
+				Edge->ProjectPoints(MeshVetrices, Coordinates, ProjectedPoints);
+				Algo::Sort(Coordinates);
+				Coordinates.Insert(Edge->GetBoundary().GetMin(), 0);
+				Coordinates.Add(Edge->GetBoundary().GetMax());
+				TArray<FPoint2D> Points2D;
+				Edge->Approximate2DPoints(Coordinates, Points2D);
+
+				DisplayPolylineWithScale(Points2D, Color2);
+				for (FPoint2D Point : Points2D)
+				{
+					DisplayPoint2DWithScale(Point, (EVisuProperty)(Color2 - 1));
+				}
+			}
+			else
+			{
+				Display2DWithScale(*Edge, Color);
+			}
+		}
+	}
+}
+
+
+#endif
+} // namespace
\ No newline at end of file
diff --git a/Engine/Source/Runtime/Datasmith/CADKernel/Base/Private/Mesh/Structure/Grid.cpp b/Engine/Source/Runtime/Datasmith/CADKernel/Base/Private/Mesh/Structure/Grid.cpp
new file mode 100644
--- /dev/null
+++ b/Engine/Source/Runtime/Datasmith/CADKernel/Base/Private/Mesh/Structure/Grid.cpp
@@ -0,0 +1,1769 @@
+// Copyright Epic Games, Inc. All Rights Reserved.
+
+#include "Mesh/Structure/Grid.h"
+
+#include "Geo/Sampling/SurfacicSampling.h"
+#include "Mesh/Meshers/MesherTools.h"
+#include "Mesh/Structure/EdgeMesh.h"
+#include "Mesh/Structure/VertexMesh.h"
+#include "Topo/TopologicalEdge.h"
+#include "Topo/TopologicalFace.h"
+#include "Utils/Util.h"
+#include "Utils/ArrayUtils.h"
+
+namespace UE::CADKernel
+{
+
+FGrid::FGrid(FTopologicalFace& InFace, FModelMesh& InMeshModel)
+	: FGridBase(InFace)
+	, CoordinateGrid(InFace.GetCuttingPointCoordinates())
+	, FaceTolerance(InFace.GetIsoTolerances())
+	, MinimumElementSize(Tolerance3D * 2.)
+	, MeshModel(InMeshModel)
+{
+#ifdef DEBUG_ONLY_SURFACE_TO_DEBUG
+	bDisplay = (InFace.GetId() == FaceToDebug);
+	Open3DDebugSession(bDisplay, FString::Printf(TEXT("Grid %d"), InFace.GetId()));
+#endif
+}
+
+void FGrid::ProcessPointCloud()
+{
+	FTimePoint StartTime = FChrono::Now();
+
+	if (!GetMeshOfLoops())
+	{
+		return;
+	}
+
+#ifdef DEBUG_GRID
+	if(bDisplay)
+	{
+		DisplayGridLoops(TEXT("FGrid::Loop 2D"), GetLoops2D(EGridSpace::Default2D), true, true, false);
+		DisplayInnerDomainPoints(TEXT("FGrid::PointCloud 2D"), GetInner2DPoints(EGridSpace::Default2D));
+
+		DisplayGridLoops(TEXT("FGrid::Loop 2D Scaled"), GetLoops2D(EGridSpace::Scaled), true, false, false);
+		DisplayInnerDomainPoints(TEXT("FGrid::PointCloud 2D Scaled"), GetInner2DPoints(EGridSpace::Scaled));
+
+		DisplayGridLoops(TEXT("FGrid::Loop 2D UniformScaled"), GetLoops2D(EGridSpace::UniformScaled), true, true, false);
+		DisplayInnerDomainPoints(TEXT("FGrid::PointCloud 2D UniformScaled"), GetInner2DPoints(EGridSpace::UniformScaled));
+
+		DisplayGridLoops(TEXT("FGrid::Loop 3D"), GetLoops3D(), true, false, false);
+		DisplayInnerDomainPoints(TEXT("FGrid::PointCloud 3D"), GetInner3DPoints());
+
+		DisplayGridLoops(TEXT("FGrid::Loop 3D"), GetLoops3D(), true, false, false);
+		DisplayInnerDomainPoints(TEXT("FGrid::PointCloud 3D"), GetInner3DPoints());
+
+		DisplayGridNormal();
+
+		DisplayInnerDomainPoints(TEXT("FGrid::PointCloud before FindInnerFacePoints"), GetInner2DPoints(EGridSpace::UniformScaled));
+
+		Wait(false);
+	}
+#endif
+
+	FindInnerFacePoints();
+
+#ifdef DEBUG_GRID
+	DisplayInnerDomainPoints(TEXT("FGrid::PointCloud after FindInnerFacePoints"), GetInner2DPoints(EGridSpace::UniformScaled));
+#endif
+
+	FindPointsCloseToLoop();
+
+#ifdef DEBUG_GRID
+	DisplayInnerDomainPoints(TEXT("FGrid::PointCloud after FindPointsCloseToLoop"), GetInner2DPoints(EGridSpace::UniformScaled));
+#endif
+
+	RemovePointsCloseToLoop();
+
+#ifdef DEBUG_GRID
+	DisplayInnerDomainPoints(TEXT("FGrid::PointCloud after RemovePointsCloseToLoop"), GetInner2DPoints(EGridSpace::UniformScaled));
+#endif
+
+	// Removed of Thin zone boundary (the last boundaries). In case of thin zone, the number of 2d boundary will be biggest than 3d boundary one.
+	// Only EGridSpace::UniformScaled is needed.
+	FaceLoops2D[EGridSpace::UniformScaled].SetNum(FaceLoops3D.Num());
+
+#ifdef DEBUG_GRID
+	if (bDisplay)
+	{
+		DisplayGridLoops(TEXT("FGrid::Final Loop 2D"), GetLoops2D(EGridSpace::UniformScaled), true, true, false);
+		Wait(false);
+	}
+#endif
+
+	Chronos.ProcessPointCloudDuration = FChrono::Elapse(StartTime);
+}
+
+void FGrid::DefineCuttingParameters()
+{
+	FTimePoint StartTime = FChrono::Now();
+
+	FCuttingGrid PreferredCuttingParametersFromLoops;
+	GetPreferredUVCuttingParametersFromLoops(PreferredCuttingParametersFromLoops);
+
+	DefineCuttingParameters(EIso::IsoU, PreferredCuttingParametersFromLoops);
+	DefineCuttingParameters(EIso::IsoV, PreferredCuttingParametersFromLoops);
+
+	CuttingSize = CoordinateGrid.Count();
+
+	Chronos.DefineCuttingParametersDuration = FChrono::Elapse(StartTime);
+}
+
+void FGrid::DefineCuttingParameters(EIso Iso, FCuttingGrid& Neighbors)
+{
+	FTimePoint StartTime = FChrono::Now();
+
+	const FSurfacicBoundary& Boundary = Face.GetBoundary();
+
+#ifdef DEBUG_GET_PREFERRED_UVCOORDINATES_FROM_NEIGHBOURS
+	TArray<double> CuttingPointTmp;
+#endif
+	if (Neighbors[Iso].Num())
+	{
+#ifdef DEBUG_GET_PREFERRED_UVCOORDINATES_FROM_NEIGHBOURS
+		TArray<FCuttingPoint> Extremities;
+		Extremities.Reserve(2);
+		Extremities.Emplace(Boundary[Iso].Min, ECoordinateType::VertexCoordinate, -1, 0.001);
+		Extremities.Emplace(Boundary[Iso].Max, ECoordinateType::VertexCoordinate, -1, 0.001);
+		FMesherTools::ComputeFinalCuttingPointsWithImposedCuttingPoints(Face.GetCrossingPointCoordinates(Iso), Face.GetCrossingPointDeltaMaxs(Iso), Extremities, CuttingPointTmp);
+#endif
+		FMesherTools::ComputeFinalCuttingPointsWithPreferredCuttingPoints(Face.GetCrossingPointCoordinates(Iso), Face.GetCrossingPointDeltaMaxs(Iso), Neighbors[Iso], Boundary[Iso], Face.GetCuttingCoordinatesAlongIso(Iso));
+	}
+	else
+	{
+		TArray<FCuttingPoint> Extremities;
+		Extremities.Reserve(2);
+		Extremities.Emplace(Boundary[Iso].Min, ECoordinateType::VertexCoordinate, -1, 0.001);
+		Extremities.Emplace(Boundary[Iso].Max, ECoordinateType::VertexCoordinate, -1, 0.001);
+		FMesherTools::ComputeFinalCuttingPointsWithImposedCuttingPoints(Face.GetCrossingPointCoordinates(Iso), Face.GetCrossingPointDeltaMaxs(Iso), Extremities, Face.GetCuttingCoordinatesAlongIso(Iso));
+	}
+
+#ifdef DEBUG_GET_PREFERRED_UVCOORDINATES_FROM_NEIGHBOURS
+	if (bDisplay)
+	{
+		EIso OtherIso = Other(Iso);
+		F3DDebugSession _(TEXT("GetPreferredUVCoordinatesFromNeighbours"));
+		{
+			F3DDebugSession _(FString::Printf(TEXT("%s From Neighbours"), IsoNames[Iso]));
+			for (FCuttingPoint CuttingU : Neighbors[Iso])
+			{
+				if (Iso == EIso::IsoU)
+				{
+					DisplayPoint(FPoint(CuttingU.Coordinate, Boundary[OtherIso].GetMin() - Boundary[OtherIso].Length() * 1 / 20), EVisuProperty::GreenPoint);
+				}
+				else
+				{
+					DisplayPoint(FPoint(Boundary[OtherIso].GetMin() - Boundary[OtherIso].Length() * 1 / 20, CuttingU.Coordinate), EVisuProperty::GreenPoint);
+				}
+			}
+		}
+		{
+			F3DDebugSession _(FString::Printf(TEXT("%s From Criteria"), IsoNames[Iso]));
+			for (double CuttingU : CuttingPointTmp)
+			{
+				if (Iso == EIso::IsoU)
+				{
+					DisplayPoint(FPoint(CuttingU, Boundary[OtherIso].GetMin() - Boundary[OtherIso].Length() * 1 / 40), EVisuProperty::YellowPoint);
+				}
+				else
+				{
+					DisplayPoint(FPoint(Boundary[OtherIso].GetMin() - Boundary[OtherIso].Length() * 1 / 40, CuttingU), EVisuProperty::YellowPoint);
+				}
+			}
+		}
+		{
+			F3DDebugSession _(FString::Printf(TEXT("%s From Neighbours"), IsoNames[Iso]));
+			for (double CuttingU : Face.GetCuttingCoordinatesAlongIso(Iso))
+			{
+				if (Iso == EIso::IsoU)
+				{
+					DisplayPoint(FPoint(CuttingU, Boundary[OtherIso].GetMin() - Boundary[OtherIso].Length() * 1 / 60), EVisuProperty::BluePoint);
+				}
+				else
+				{
+					DisplayPoint(FPoint(Boundary[OtherIso].GetMin() - Boundary[OtherIso].Length() * 1 / 60, CuttingU), EVisuProperty::BluePoint);
+				}
+			}
+		}
+		//Wait();
+	}
+#endif
+	CuttingCount[Iso] = CoordinateGrid.IsoCount(Iso);
+
+	Chronos.DefineCuttingParametersDuration = FChrono::Elapse(StartTime);
+}
+
+//#define DEBUG_GET_PREFERRED_UVCOORDINATES_FROM_NEIGHBOURS
+void FGrid::GetPreferredUVCuttingParametersFromLoops(FCuttingGrid& CuttingParametersFromLoops)
+{
+#ifdef DEBUG_GET_PREFERRED_UVCOORDINATES_FROM_NEIGHBOURS
+	F3DDebugSession _(bDisplay, TEXT("GetPreferredUVCoordinatesFromNeighbours"));
+
+	if (bDisplay)
+	{
+		{
+			F3DDebugSession _(TEXT("Surface 2D"));
+			UE::CADKernel::Display2D(*Face.GetCarrierSurface());
+		}
+	}
+#endif
+
+	int32 nbPoints = 0;
+	for (const TSharedPtr<FTopologicalLoop>& Loop : Face.GetLoops())
+	{
+		for (const FOrientedEdge& Edge : Loop->GetEdges())
+		{
+			nbPoints += Edge.Entity->GetOrCreateMesh(MeshModel).GetNodeCoordinates().Num() + 1;
+		}
+	}
+
+	CuttingParametersFromLoops[EIso::IsoU].Reserve(nbPoints);
+	CuttingParametersFromLoops[EIso::IsoV].Reserve(nbPoints);
+
+	for (const TSharedPtr<FTopologicalLoop>& Loop : Face.GetLoops())
+	{
+		for (const FOrientedEdge& OrientedEdge : Loop->GetEdges())
+		{
+			const TSharedPtr<FTopologicalEdge>& Edge = OrientedEdge.Entity;
+
+#ifdef DEBUG_GET_PREFERRED_UVCOORDINATES_FROM_NEIGHBOURS
+			if (bDisplay)
+			{
+				F3DDebugSession _(FString::Printf(TEXT("Edge %d"), Edge->GetId()));
+				Display2D(*Edge);
+			}
+#endif
+
+			TArray<double> ProjectedPointCoords;
+			TSharedRef<FTopologicalEdge> ActiveEdge = Edge->GetLinkActiveEdge();
+			if (ActiveEdge->IsMeshed())
+			{
+				const FMesh* EdgeMesh = ActiveEdge->GetMesh();
+				if (!EdgeMesh)
+				{
+					continue;
+				}
+
+				const TArray<FPoint>& EdgeMeshNodes = EdgeMesh->GetNodeCoordinates();
+				if (EdgeMeshNodes.Num() == 0)
+				{
+					continue;
+				}
+
+				ProjectedPointCoords.Reserve(EdgeMeshNodes.Num() + 2);
+				bool bSameDirection = Edge->IsSameDirection(*ActiveEdge);
+
+				Edge->ProjectTwinEdgePoints(EdgeMeshNodes, bSameDirection, ProjectedPointCoords);
+				ProjectedPointCoords.Insert(Edge->GetStartCurvilinearCoordinates(), 0);
+				ProjectedPointCoords.Add(Edge->GetEndCurvilinearCoordinates());
+			}
+			else // Add Vertices
+			{
+				ProjectedPointCoords.Add(Edge->GetBoundary().GetMin());
+				ProjectedPointCoords.Add(Edge->GetBoundary().GetMax());
+			}
+
+			TArray<FPoint2D> EdgePoints2D;
+			Edge->Approximate2DPoints(ProjectedPointCoords, EdgePoints2D);
+
+#ifdef DEBUG_GET_PREFERRED_UVCOORDINATES_FROM_NEIGHBOURS
+			F3DDebugSession _(bDisplay, ("Nodes"));
+#endif
+			for (int32 Index = 0; Index < EdgePoints2D.Num(); ++Index)
+			{
+				CuttingParametersFromLoops[EIso::IsoU].Emplace(EdgePoints2D[Index].U, ECoordinateType::OtherCoordinate);
+				CuttingParametersFromLoops[EIso::IsoV].Emplace(EdgePoints2D[Index].V, ECoordinateType::OtherCoordinate);
+#ifdef DEBUG_GET_PREFERRED_UVCOORDINATES_FROM_NEIGHBOURS
+				if (bDisplay)
+				{
+					DisplayPoint(EdgePoints2D[Index]);
+				}
+#endif
+			}
+		}
+	}
+
+	TFunction<void(TArray<FCuttingPoint>&, double)> SortAndRemoveDuplicated = [](TArray<FCuttingPoint>& Neighbours, double Tolerance)
+	{
+		if (Neighbours.Num() == 0)
+		{
+			return;
+		}
+
+		Algo::Sort(Neighbours, [](const FCuttingPoint& Point1, const FCuttingPoint& Point2) -> bool
+			{
+				return (Point1.Coordinate) < (Point2.Coordinate);
+			}
+		);
+
+		int32 NewIndex = 0;
+		for (int32 Index = 1; Index < Neighbours.Num() - 1; ++Index)
+		{
+			if (FMath::IsNearlyEqual(Neighbours[Index].Coordinate, Neighbours[NewIndex].Coordinate, Tolerance))
+			{
+				continue;
+			}
+			NewIndex++;
+			Neighbours[NewIndex] = Neighbours[Index];
+		}
+		if (FMath::IsNearlyEqual(Neighbours.Last().Coordinate, Neighbours[NewIndex].Coordinate, Tolerance))
+		{
+			Neighbours[NewIndex] = Neighbours.Last();
+		}
+		else
+		{
+			NewIndex++;
+		}
+		Neighbours.SetNum(NewIndex);
+	};
+
+	SortAndRemoveDuplicated(CuttingParametersFromLoops[EIso::IsoU], FaceTolerance[EIso::IsoU]);
+	SortAndRemoveDuplicated(CuttingParametersFromLoops[EIso::IsoV], FaceTolerance[EIso::IsoV]);
+}
+
+bool FGrid::GeneratePointCloud()
+{
+	FTimePoint StartTime = FChrono::Now();
+
+	if (CheckIf2DGridIsDegenerate())
+	{
+		return false;
+	}
+
+	NodeMarkers.Init(ENodeMarker::None, CuttingSize);
+
+	EvaluatePointGrid(CoordinateGrid, true);
+
+	CountOfInnerNodes = CuttingSize;
+
+	if (!ScaleGrid())
+	{
+		return false;
+	}
+
+	Chronos.GeneratePointCloudDuration += FChrono::Elapse(StartTime);
+	return true;
+}
+
+void FGrid::FindPointsCloseToLoop()
+{
+	FTimePoint StartTime = FChrono::Now();
+
+	int32 IndexLoop = 0;
+	int32 Index[2] = { 1, 1 };
+	int32 GlobalIndex = 1;
+
+	const FPoint2D* PointA = nullptr;
+	const FPoint2D* PointB = nullptr;
+
+#ifdef DEBUG_FINDPOINTSCLOSETOLOOP
+	if (bDisplay)
+	{
+		Open3DDebugSession(TEXT("FGrid::FindPointsClosedToLoop result"));
+	}
+
+	int32 CellIndex = 0;
+	bool bWaitCell = true;
+	TFunction<void()> DisplayCell = [&]()
+	{
+		if (bDisplay)
+		{
+			F3DDebugSession _(*FString::Printf(TEXT("Cell %d"), CellIndex++));
+			DisplayPoint2DWithScale(Points2D[EGridSpace::UniformScaled][GlobalIndex]);
+			DisplayPoint2DWithScale(Points2D[EGridSpace::UniformScaled][GlobalIndex - 1]);
+			DisplayPoint2DWithScale(Points2D[EGridSpace::UniformScaled][GlobalIndex - 1 - CuttingCount[EIso::IsoU]]);
+			DisplayPoint2DWithScale(Points2D[EGridSpace::UniformScaled][GlobalIndex - CuttingCount[EIso::IsoU]]);
+			Wait(bWaitCell);
+		}
+	};
+#endif
+
+	// Find start index
+	TFunction<void(EIso)> FindPointAIndex = [&](EIso Iso)
+	{
+		Index[Iso] = 1;
+		for (; Index[Iso] < CuttingCount[Iso] - 1; ++Index[Iso])
+		{
+			if (UniformCuttingCoordinates[Iso][Index[Iso]] + DOUBLE_SMALL_NUMBER > (*PointA)[Iso])
+			{
+				break;
+			}
+		}
+	};
+
+	TFunction<void()> SetCellCloseToLoop = [&]()
+	{
+		SetCloseToLoop(GlobalIndex);
+		SetCloseToLoop(GlobalIndex - 1);
+		SetCloseToLoop(GlobalIndex - 1 - CuttingCount[EIso::IsoU]);
+		SetCloseToLoop(GlobalIndex - CuttingCount[EIso::IsoU]);
+#ifdef DEBUG_FINDPOINTSCLOSETOLOOP
+		DisplayCell();
+#endif
+	};
+
+	TFunction<void(EIso)> Increase = [&](EIso Iso)
+	{
+		if (Index[Iso] < CuttingCount[Iso] - 1)
+		{
+			Index[Iso]++;
+			GlobalIndex += Iso == EIso::IsoU ? 1 : CuttingCount[EIso::IsoU];
+		}
+	};
+
+	TFunction<void(EIso)> Decrease = [&](EIso Iso)
+	{
+		if (Index[Iso] > 1) //-V547
+		{
+			Index[Iso]--;
+			GlobalIndex -= Iso == EIso::IsoU ? 1 : CuttingCount[EIso::IsoU];
+		}
+	};
+
+	TFunction<void()> IncreaseU = [&]()
+	{
+		Increase(EIso::IsoU);
+	};
+
+	TFunction<void()> IncreaseV = [&]()
+	{
+		Increase(EIso::IsoV);
+	};
+
+	TFunction<void()> DecreaseU = [&]()
+	{
+		Decrease(EIso::IsoU);
+	};
+
+	TFunction<void()> DecreaseV = [&]()
+	{
+		Decrease(EIso::IsoV);
+	};
+
+	TFunction<bool(const double, const double)> IsReallyBigger = [](const double FirstValue, const double SecondValue) ->bool
+	{
+		return FirstValue - DOUBLE_SMALL_NUMBER > SecondValue;
+	};
+
+	TFunction<bool(const double, const double)> IsReallySmaller = [](const double FirstValue, const double SecondValue) ->bool
+	{
+		return FirstValue + DOUBLE_SMALL_NUMBER < SecondValue;
+	};
+
+	double Slope;
+	double Origin;
+
+	TFunction<void(EIso, const int32, const int32, TFunction<void()>, TFunction<void()>)> FindIntersection = [&](EIso MainIso, const int32 DeltaIso, const int32 DeltaOther, TFunction<void()> OffsetIndexIfBigger, TFunction<void()> OffsetIndexIfSmaller)
+	{
+		TFunction<bool(const double, const double)> TestAlongIso = DeltaIso ? IsReallyBigger : IsReallySmaller;
+		TFunction<bool(const double, const double)> TestAlongOther = DeltaOther ? IsReallyBigger : IsReallySmaller;
+
+		EIso OtherIso = Other(MainIso);
+		while (TestAlongIso(UniformCuttingCoordinates[MainIso][Index[MainIso] - DeltaIso], (*PointB)[MainIso]) || TestAlongOther(UniformCuttingCoordinates[OtherIso][Index[OtherIso] - DeltaOther], (*PointB)[OtherIso]))
+		{
+			double CoordinateOther = Slope * UniformCuttingCoordinates[MainIso][Index[MainIso] - DeltaIso] + Origin;
+			if (IsReallyBigger(CoordinateOther, UniformCuttingCoordinates[OtherIso][Index[OtherIso] - DeltaOther]))
+			{
+				OffsetIndexIfBigger();
+			}
+			else if (IsReallySmaller(CoordinateOther, UniformCuttingCoordinates[OtherIso][Index[OtherIso] - DeltaOther]))
+			{
+				OffsetIndexIfSmaller();
+			}
+			else // IsNearlyEqual
+			{
+				OffsetIndexIfBigger();
+				OffsetIndexIfSmaller();
+			}
+			SetCellCloseToLoop();
+		}
+	};
+
+	TFunction<bool(EIso)> FindIntersectionIsoStrip = [&](EIso MainIso) -> bool
+	{
+		EIso OtherIso = Other(MainIso);
+		if ((UniformCuttingCoordinates[OtherIso][Index[OtherIso] - 1] < (*PointB)[OtherIso]) && ((*PointB)[OtherIso] < UniformCuttingCoordinates[OtherIso][Index[OtherIso]]))
+		{
+			if ((UniformCuttingCoordinates[MainIso][Index[MainIso] - 1] < (*PointB)[MainIso]) && ((*PointB)[MainIso] < UniformCuttingCoordinates[MainIso][Index[MainIso]]))
+			{
+				PointA = PointB;
+				return true;
+			}
+
+			if ((*PointA)[MainIso] < (*PointB)[MainIso])
+			{
+				while (UniformCuttingCoordinates[MainIso][Index[MainIso]] < (*PointB)[MainIso])
+				{
+					Increase(MainIso);
+					SetCellCloseToLoop();
+				}
+			}
+			else
+			{
+				while (UniformCuttingCoordinates[MainIso][Index[MainIso] - 1] > (*PointB)[MainIso])
+				{
+					Decrease(MainIso);
+					SetCellCloseToLoop();
+				}
+			}
+			PointA = PointB;
+			return true;
+		}
+		return false;
+	};
+
+	double ABv = 0.;
+	double ABu = 0.;
+	for (const TArray<FPoint2D>& Loop : FaceLoops2D[EGridSpace::UniformScaled])
+	{
+		PointA = &Loop.Last();
+
+		// Find start index
+		FindPointAIndex(EIso::IsoU);
+		FindPointAIndex(EIso::IsoV);
+
+		GlobalIndex = Index[EIso::IsoV] * CuttingCount[EIso::IsoU] + Index[EIso::IsoU];
+		SetCellCloseToLoop();
+
+		for (int32 BIndex = 0; BIndex < Loop.Num(); ++BIndex)
+		{
+			PointB = &Loop[BIndex];
+
+#ifdef DEBUG_FINDPOINTSCLOSETOLOOP
+			if (bDisplay)
+			{
+				F3DDebugSession _(TEXT("SEG"));
+				DisplaySegmentWithScale(*PointB, *PointA);
+			}
+#endif
+
+			// Horizontal case
+			if (FindIntersectionIsoStrip(EIso::IsoU))
+			{
+				continue;
+			}
+
+			if (FindIntersectionIsoStrip(EIso::IsoV))
+			{
+				continue;
+			}
+
+			ABv = PointB->V - PointA->V;
+			ABu = PointB->U - PointA->U;
+			if (FMath::Abs(ABu) > FMath::Abs(ABv))
+			{
+				Slope = ABv / ABu;
+				Origin = PointA->V - Slope * PointA->U;
+				if (ABu > 0)
+				{
+					if (ABv > 0)
+					{
+						FindIntersection(EIso::IsoU, 0, 0, IncreaseV, IncreaseU);
+					}
+					else
+					{
+						FindIntersection(EIso::IsoU, 0, 1, IncreaseU, DecreaseV);
+					}
+				}
+				else
+				{
+					if (ABv > 0)
+					{
+						FindIntersection(EIso::IsoU, 1, 0, IncreaseV, DecreaseU);
+					}
+					else
+					{
+						FindIntersection(EIso::IsoU, 1, 1, DecreaseU, DecreaseV);
+					}
+				}
+			}
+			else
+			{
+				Slope = ABu / ABv;
+				Origin = PointA->U - Slope * PointA->V;
+				if (ABu > 0)
+				{
+					if (ABv > 0)
+					{
+						FindIntersection(EIso::IsoV, 0, 0, IncreaseU, IncreaseV);
+					}
+					else
+					{
+						FindIntersection(EIso::IsoV, 1, 0, IncreaseU, DecreaseV);
+					}
+				}
+				else
+				{
+					if (ABv > 0)
+					{
+						FindIntersection(EIso::IsoV, 0, 1, IncreaseV, DecreaseU);
+					}
+					else
+					{
+						FindIntersection(EIso::IsoV, 1, 1, DecreaseV, DecreaseU);
+					}
+				}
+			}
+			PointA = PointB;
+		}
+	}
+
+#ifdef DEBUG_FINDPOINTSCLOSETOLOOP
+	if (bDisplay)
+	{
+		Close3DDebugSession();
+		Wait();
+	}
+#endif
+
+	Chronos.FindPointsCloseToLoopDuration += FChrono::Elapse(StartTime);
+}
+
+void FGrid::RemovePointsCloseToLoop()
+{
+	FTimePoint StartTime = FChrono::Now();
+
+	struct FGridSegment
+	{
+		const FPoint2D* StartPoint;
+		const FPoint2D* EndPoint;
+		double StartPointWeight;
+		double EndPointWeight;
+		double UMin;
+		double VMin;
+		double UMax;
+		double VMax;
+
+		FGridSegment(const FPoint2D& SPoint, const FPoint2D& EPoint)
+			: StartPoint(&SPoint)
+			, EndPoint(&EPoint)
+		{
+			StartPointWeight = StartPoint->U + StartPoint->V;
+			EndPointWeight = EndPoint->U + EndPoint->V;
+			if (StartPointWeight > EndPointWeight)
+			{
+				Swap(StartPointWeight, EndPointWeight);
+				Swap(StartPoint, EndPoint);
+			}
+			if (StartPoint->U < EndPoint->U)
+			{
+				UMin = StartPoint->U;
+				UMax = EndPoint->U;
+			}
+			else
+			{
+				UMin = EndPoint->U;
+				UMax = StartPoint->U;
+			}
+			if (StartPoint->V < EndPoint->V)
+			{
+				VMin = StartPoint->V;
+				VMax = EndPoint->V;
+			}
+			else
+			{
+				VMin = EndPoint->V;
+				VMax = StartPoint->V;
+			}
+		}
+	};
+
+	TArray<FGridSegment> LoopSegments;
+	{
+		int32 SegmentNum = 0;
+		for (const TArray<FPoint2D>& Loop : FaceLoops2D[EGridSpace::UniformScaled])
+		{
+			SegmentNum += Loop.Num();
+		}
+		LoopSegments.Reserve(SegmentNum);
+
+		for (TArray<FPoint2D>& Loop : FaceLoops2D[EGridSpace::UniformScaled])
+		{
+			for (int32 Index = 0; Index < Loop.Num() - 1; ++Index)
+			{
+				LoopSegments.Emplace(Loop[Index], Loop[Index + 1]);
+			}
+			LoopSegments.Emplace(Loop.Last(), Loop[0]);
+		}
+
+		Algo::Sort(LoopSegments, [](const FGridSegment& Seg1, const FGridSegment& Seg2) -> bool
+			{
+				return (Seg1.EndPointWeight) < (Seg2.EndPointWeight);
+			}
+		);
+	}
+
+	// Sort point border grid
+	TArray<double> GridPointWeight;
+	TArray<int32> IndexOfPointsNearAndInsideLoop;
+	TArray<int32> SortedPointIndexes;
+	{
+		IndexOfPointsNearAndInsideLoop.Reserve(CuttingSize);
+		for (int32 Index = 0; Index < CuttingSize; ++Index)
+		{
+			if (IsNodeInsideAndCloseToLoop(Index))
+			{
+				IndexOfPointsNearAndInsideLoop.Add(Index);
+			}
+		}
+
+		GridPointWeight.Reserve(IndexOfPointsNearAndInsideLoop.Num());
+		SortedPointIndexes.Reserve(IndexOfPointsNearAndInsideLoop.Num());
+		for (const int32& Index : IndexOfPointsNearAndInsideLoop)
+		{
+			GridPointWeight.Add(Points2D[EGridSpace::UniformScaled][Index].U + Points2D[EGridSpace::UniformScaled][Index].V);
+		}
+		for (int32 Index = 0; Index < IndexOfPointsNearAndInsideLoop.Num(); ++Index)
+		{
+			SortedPointIndexes.Add(Index);
+		}
+		SortedPointIndexes.Sort([&GridPointWeight](const int32& Index1, const int32& Index2) { return GridPointWeight[Index1] < GridPointWeight[Index2]; });
+	}
+
+	// only used to reduce the search of neighborhood
+	double DeltaUVMax = 0;
+	{
+		double MaxDeltaU = 0;
+		for (int32 Index = 0; Index < CuttingCount[EIso::IsoU] - 1; ++Index)
+		{
+			MaxDeltaU = FMath::Max(MaxDeltaU, FMath::Abs(UniformCuttingCoordinates[EIso::IsoU][Index + 1] - UniformCuttingCoordinates[EIso::IsoU][Index]));
+		}
+
+		double MaxDeltaV = 0;
+		for (int32 Index = 0; Index < CuttingCount[EIso::IsoV] - 1; ++Index)
+		{
+			MaxDeltaV = FMath::Max(MaxDeltaV, FMath::Abs(UniformCuttingCoordinates[EIso::IsoV][Index + 1] - UniformCuttingCoordinates[EIso::IsoV][Index]));
+		}
+		DeltaUVMax = FMath::Max(MaxDeltaU, MaxDeltaV);
+	}
+
+	// Find DeltaU and DeltaV around a cutting point defined by its index
+	TFunction<void(const int32, double&, double&)> GetDeltaUV = [&](const int32 Index, double& DeltaU, double& DeltaV)
+	{
+		int32 IndexU = Index % CuttingCount[EIso::IsoU];
+		int32 IndexV = Index / CuttingCount[EIso::IsoU];
+
+		if (IndexU == 0)
+		{
+			DeltaU = FMath::Abs(UniformCuttingCoordinates[EIso::IsoU][1] - UniformCuttingCoordinates[EIso::IsoU][0]);
+		}
+		else if (IndexU == CuttingCount[EIso::IsoU] - 1)
+		{
+			DeltaU = FMath::Abs(UniformCuttingCoordinates[EIso::IsoU][CuttingCount[EIso::IsoU] - 1] - UniformCuttingCoordinates[EIso::IsoU][CuttingCount[EIso::IsoU] - 2]);
+		}
+		else
+		{
+			DeltaU = FMath::Abs(UniformCuttingCoordinates[EIso::IsoU][IndexU + 1] - UniformCuttingCoordinates[EIso::IsoU][IndexU - 1]) * .5;
+		}
+
+		if (IndexV == 0)
+		{
+			DeltaV = FMath::Abs(UniformCuttingCoordinates[EIso::IsoV][1] - UniformCuttingCoordinates[EIso::IsoV][0]);
+		}
+		else if (IndexV == CuttingCount[EIso::IsoV] - 1)
+		{
+			DeltaV = FMath::Abs(UniformCuttingCoordinates[EIso::IsoV][CuttingCount[EIso::IsoV] - 1] - UniformCuttingCoordinates[EIso::IsoV][CuttingCount[EIso::IsoV] - 2]);
+		}
+		else
+		{
+			DeltaV = FMath::Abs(UniformCuttingCoordinates[EIso::IsoV][IndexV + 1] - UniformCuttingCoordinates[EIso::IsoV][IndexV - 1]) * .5;
+		}
+	};
+
+	int32 SegmentIndex = 0;
+	SegmentIndex = 0;
+	for (const int32& SortedIndex : SortedPointIndexes)
+	{
+		int32 Index = IndexOfPointsNearAndInsideLoop[SortedIndex];
+		const FPoint2D& Point2D = Points2D[EGridSpace::UniformScaled][Index];
+
+		double DeltaU;
+		double DeltaV;
+		GetDeltaUV(Index, DeltaU, DeltaV);
+
+		//find the first segment that could be close to the point			
+		for (; SegmentIndex < LoopSegments.Num(); ++SegmentIndex)
+		{
+			if (GridPointWeight[SortedIndex] < LoopSegments[SegmentIndex].EndPointWeight + DeltaUVMax)
+			{
+				break;
+			}
+		}
+
+		for (int32 SegmentIndex2 = SegmentIndex; SegmentIndex2 < LoopSegments.Num(); ++SegmentIndex2)
+		{
+			const FGridSegment& Segment = LoopSegments[SegmentIndex2];
+
+			if (GridPointWeight[SortedIndex] < Segment.StartPointWeight - DeltaUVMax)
+			{
+				continue;
+			}
+			if (Point2D.U + DeltaU < Segment.UMin)
+			{
+				continue;
+			}
+			if (Point2D.U - DeltaU > Segment.UMax)
+			{
+				continue;
+			}
+			if (Point2D.V + DeltaV < Segment.VMin)
+			{
+				continue;
+			}
+			if (Point2D.V - DeltaV > Segment.VMax)
+			{
+				continue;
+			}
+
+			double Coordinate;
+			FPoint2D Projection = ProjectPointOnSegment(Point2D, *Segment.StartPoint, *Segment.EndPoint, Coordinate, /*bRestrictCoodinateToInside*/ true);
+
+			// If Projected point is in the oval center on Point2D => the node is too close
+			FPoint2D ProjectionToPoint = Abs(Point2D - Projection);
+
+			if (ProjectionToPoint.U > DeltaU * 0.05 || ProjectionToPoint.V > DeltaV * 0.05)
+			{
+				continue;
+			}
+
+			SetTooCloseToLoop(Index);
+			break;
+		}
+	}
+
+	Chronos.RemovePointsClosedToLoopDuration += FChrono::Elapse(StartTime);
+}
+
+/**
+ * For the surface normal at a StartPoint of the 3D degenerated curve (Not degenerated in 2d)
+ * The normal is swap if StartPoint is too close to the Boundary
+ * The norm of the normal is defined as 1/20 of the parallel boundary Length
+ */
+void ScaleAndSwap(FPoint2D& Normal, const FPoint2D& StartPoint, const FSurfacicBoundary& Boundary)
+{
+	Normal.Normalize();
+	FPoint2D MainDirection = Normal;
+	MainDirection.U /= Boundary[EIso::IsoU].Length();
+	MainDirection.V /= Boundary[EIso::IsoV].Length();
+
+	TFunction<void(const EIso)> SwapAndScale = [&](const EIso Iso)
+	{
+		if (MainDirection[Iso] > 0)
+		{
+			if (FMath::IsNearlyEqual(Boundary[Iso].Max, StartPoint[Iso]))
+			{
+				Normal *= -1.;
+			}
+		}
+		else
+		{
+			if (FMath::IsNearlyEqual(Boundary[IsoU].Min, StartPoint[Iso]))
+			{
+				Normal *= -1.;
+			}
+		}
+		Normal *= Boundary[Iso].Length() / 20;
+	};
+
+	if (MainDirection.U > MainDirection.V)
+	{
+		SwapAndScale(EIso::IsoU);
+	}
+	else
+	{
+		SwapAndScale(EIso::IsoV);
+	}
+}
+
+/**
+ * Displace loop nodes inside to avoid that the nodes are outside the surface Boundary, so outside the grid
+ */
+void SlightlyDisplacedPolyline(TArray<FPoint2D>& D2Points, const FSurfacicBoundary& Boundary)
+{
+	FPoint2D Normal;
+	for (int32 Index = 0; Index < D2Points.Num() - 1; ++Index)
+	{
+		FPoint2D Tangent = D2Points[Index + 1] - D2Points[Index];
+		Normal = Tangent.GetPerpendicularVector();
+
+		ScaleAndSwap(Normal, D2Points[Index], Boundary);
+		D2Points[Index] += Normal;
+	}
+	D2Points.Last() += Normal;
+}
+
+void FGrid::GetMeshOfLoop(const FTopologicalLoop& Loop)
+{
+	int32 LoopNodeCount = 0;
+
+	for (const FOrientedEdge& Edge : Loop.GetEdges())
+	{
+		LoopNodeCount += Edge.Entity->GetLinkActiveEdge()->GetMesh()->GetNodeCount() + 2;
+	}
+
+	TArray<FPoint2D>& Loop2D = FaceLoops2D[EGridSpace::Default2D].Emplace_GetRef();
+	Loop2D.Empty(LoopNodeCount);
+
+	TArray<FPoint>& Loop3D = FaceLoops3D.Emplace_GetRef();
+	Loop3D.Reserve(LoopNodeCount);
+
+	TArray<FVector3f>& LoopNormals = NormalsOfFaceLoops.Emplace_GetRef();
+	LoopNormals.Reserve(LoopNodeCount);
+
+	TArray<int32>& LoopIds = NodeIdsOfFaceLoops.Emplace_GetRef();
+	LoopIds.Reserve(LoopNodeCount);
+
+	for (const FOrientedEdge& OrientedEdge : Loop.GetEdges())
+	{
+		const TSharedPtr<FTopologicalEdge>& Edge = OrientedEdge.Entity;
+		const TSharedRef<FTopologicalEdge>& ActiveEdge = Edge->GetLinkActiveEdge();
+
+#ifdef DEBUG_GET_MESH_OF_LOOP
+		if (bDisplay)
+		{
+			F3DDebugSession _(bDisplay, FString::Printf(TEXT("Edge %d"), Edge->GetId()));
+			Display2D(*Edge);
+			Wait();
+		}
+#endif
+
+		bool bSameDirection = Edge->IsSameDirection(*ActiveEdge);
+
+		TArray<double> EdgeCuttingPointCoordinates;
+		{
+			const TArray<FCuttingPoint>& CuttingPoints = Edge->GetCuttingPoints();
+			if (!CuttingPoints.IsEmpty())
+			{
+				GetCuttingPointCoordinates(CuttingPoints, EdgeCuttingPointCoordinates);
+			}
+		}
+
+		FSurfacicPolyline CuttingPolyline(true);
+
+		if (Edge->IsDegenerated())
+		{
+			if (EdgeCuttingPointCoordinates.IsEmpty())
+			{
+				int32 CuttingPointCount = 2;
+				for (const FTopologicalEdge* TwinEdge : Edge->GetTwinEntities())
+				{
+					int32 TwinCuttingPointCount = TwinEdge->GetCuttingPoints().Num();
+					if (TwinCuttingPointCount > CuttingPointCount)
+					{
+						CuttingPointCount = TwinCuttingPointCount;
+					}
+				}
+				FMesherTools::FillCuttingPointCoordinates(Edge->GetBoundary(), CuttingPointCount, EdgeCuttingPointCoordinates);
+			}
+
+			Swap(CuttingPolyline.Coordinates, EdgeCuttingPointCoordinates);
+			Edge->Approximate2DPoints(CuttingPolyline.Coordinates, CuttingPolyline.Points2D);
+
+			CuttingPolyline.Points3D.Init(ActiveEdge->GetStartBarycenter(), CuttingPolyline.Coordinates.Num());
+
+			TArray<FPoint2D> D2Points = CuttingPolyline.Points2D;
+			const FSurfacicBoundary& Boundary = Edge->GetCurve()->GetCarrierSurface()->GetBoundary();
+			// to compute the normals, the 2D points are slightly displaced perpendicular to the curve
+			SlightlyDisplacedPolyline(D2Points, Boundary);
+			Edge->GetCurve()->GetCarrierSurface()->EvaluateNormals(D2Points, CuttingPolyline.Normals);
+		}
+		else
+		{
+			if (EdgeCuttingPointCoordinates.IsEmpty())
+			{
+				const TArray<FPoint>& MeshVertex3D = ActiveEdge->GetOrCreateMesh(MeshModel).GetNodeCoordinates();
+				TArray<double> ProjectedPointCoords;
+
+				CuttingPolyline.Coordinates.Reserve(MeshVertex3D.Num() + 2);
+				if(MeshVertex3D.Num())
+				{
+					Edge->ProjectTwinEdgePoints(MeshVertex3D, bSameDirection, CuttingPolyline.Coordinates);
+				}
+				CuttingPolyline.Coordinates.Insert(Edge->GetBoundary().GetMin(), 0);
+				CuttingPolyline.Coordinates.Add(Edge->GetBoundary().GetMax());
+
+				// check if there are coincident coordinates
+				bool bProjectionFailed = false;
+				for (int32 Index = 0; Index < CuttingPolyline.Coordinates.Num() - 1;)
+				{
+					double Diff = CuttingPolyline.Coordinates[Index++];
+					Diff -= CuttingPolyline.Coordinates[Index];
+					if (Diff > -DOUBLE_SMALL_NUMBER)
+					{
+						bProjectionFailed = true;
+						break;
+					}
+				}
+				if (bProjectionFailed)
+				{
+					FMesherTools::FillCuttingPointCoordinates(Edge->GetBoundary(), MeshVertex3D.Num() + 2, CuttingPolyline.Coordinates);
+				}
+			}
+			else
+			{
+				Swap(CuttingPolyline.Coordinates, EdgeCuttingPointCoordinates);
+			}
+
+			Edge->ApproximatePolyline(CuttingPolyline);
+		}
+
+		TArray<int32> EdgeVerticesIndex;
+		if (Edge->IsDegenerated())
+		{
+			EdgeVerticesIndex.Init(ActiveEdge->GetStartVertex()->GetLinkActiveEntity()->GetOrCreateMesh(MeshModel).GetMesh(), CuttingPolyline.Coordinates.Num());
+			if (OrientedEdge.Direction == Front)
+			{
+				EdgeVerticesIndex[0] = ActiveEdge->GetEndVertex()->GetLinkActiveEntity()->GetOrCreateMesh(MeshModel).GetMesh();
+			}
+			else
+			{
+				EdgeVerticesIndex.Last() = ActiveEdge->GetEndVertex()->GetLinkActiveEntity()->GetOrCreateMesh(MeshModel).GetMesh();
+			}
+		}
+		else if (Edge->IsVirtuallyMeshed())
+		{
+			// @see FParametricMesher::Mesh(FTopologicalEdge& InEdge, FTopologicalFace& Face)
+			int32 NodeCount = CuttingPolyline.Coordinates.Num();
+			EdgeVerticesIndex.Reserve(NodeCount);
+			int32 MiddleNodeIndex = NodeCount / 2;
+			int32 Index = 0;
+			{
+				int32 StartVertexMeshId = Edge->GetStartVertex()->GetLinkActiveEntity()->GetOrCreateMesh(MeshModel).GetMesh();
+				for (; Index < MiddleNodeIndex; ++Index)
+				{
+					EdgeVerticesIndex.Add(StartVertexMeshId);
+				}
+			}
+			{
+				int32 EndVertexMeshId = Edge->GetEndVertex()->GetLinkActiveEntity()->GetOrCreateMesh(MeshModel).GetMesh();
+				for (; Index < NodeCount; ++Index)
+				{
+					EdgeVerticesIndex.Add(EndVertexMeshId);
+				}
+			}
+		}
+		else
+		{
+			EdgeVerticesIndex = ActiveEdge->GetOrCreateMesh(MeshModel).EdgeVerticesIndex;
+		}
+
+#ifdef DEBUG_GET_MESH_OF_LOOP
+		if (bDisplay)
+		{
+			F3DDebugSession _(bDisplay, FString::Printf(TEXT("Edge %d cutting points on surface"), ActiveEdge->GetId()));
+			for (const FPoint2D& Point2D : CuttingPolyline.Points2D)
+			{
+				DisplayPoint(Point2D);
+			}
+			DisplayPolyline(CuttingPolyline.Points2D, EVisuProperty::BlueCurve);
+			Wait();
+		}
+#endif
+
+		if (OrientedEdge.Direction != EOrientation::Front)
+		{
+			CuttingPolyline.Reverse();
+		}
+
+		if (bSameDirection != (OrientedEdge.Direction == EOrientation::Front))
+		{
+			Algo::Reverse(EdgeVerticesIndex);
+		}
+
+		ensureCADKernel(CuttingPolyline.Size() > 1);
+
+		Loop2D.Append(MoveTemp(CuttingPolyline.Points2D));
+		Loop2D.Pop();
+
+#ifdef DEBUG_GET_MESH_OF_LOOP
+		if (bDisplay)
+		{
+			F3DDebugSession _(bDisplay, FString::Printf(TEXT("Loop with Edge %d"), Edge->GetId()));
+			DisplayPolyline(Loop2D, EVisuProperty::BlueCurve);
+			Wait();
+		}
+#endif
+
+#ifdef CADKERNEL_DEBUG_
+		const FPoint2D* PrevPoint = &Loop2D.Last();
+		if(Loop2D.Num() > 1)
+		{
+			for (const FPoint2D& Point : Loop2D)
+			{
+				double Dist = PrevPoint->Distance(Point);
+				if(Dist > 0.0001)
+				{
+					ensureCADKernel(false);
+				}
+				PrevPoint = &Point;
+			}
+		}
+#endif
+
+		int32 LastIndex = Loop3D.Num();
+		Loop3D.Append(MoveTemp(CuttingPolyline.Points3D));
+		Loop3D[LastIndex] = (ActiveEdge->GetStartVertex((OrientedEdge.Direction == EOrientation::Front) == bSameDirection)->GetLinkActiveEntity()->GetBarycenter());
+		Loop3D.Pop();
+
+		LoopNormals.Append(MoveTemp(CuttingPolyline.Normals));
+		LoopNormals.Pop();
+
+		LoopIds.Append(MoveTemp(EdgeVerticesIndex));
+		LoopIds.Pop();
+
+#ifdef DEBUG_GET_MESH_OF_LOOP
+		static int32 GetMeshOfLoopIter = 0;
+		GetMeshOfLoopIter++;
+		if (LoopIds.Num() != Loop2D.Num() || GetMeshOfLoopIter == 2991)
+		{
+			printf("");
+		}
+#endif
+	}
+
+	if (Loop2D.Num() < 3) // degenerated loop
+	{
+		FaceLoops2D[EGridSpace::Default2D].Pop();
+		FaceLoops3D.Pop();
+		NormalsOfFaceLoops.Pop();
+		NodeIdsOfFaceLoops.Pop();
+	}
+}
+
+void FGrid::GetMeshOfThinZone(const FThinZone2D& ThinZone)
+{
+	// ThinZones are identified during "ApplyCriteria". "ApplyCriteria" use the FCriteriaGrid, The EGridSpace::UniformScaled are not the same between FCriteriaGrid and FGrid
+	// This is the reason why, to get the ThinZone in the FGrid UniformScaled space, we need to get the mesh of the thinZone defined by the node ids, 
+	// And then, from this, we retrieve the thinZone points in Grid UniformScaled space.
+
+	// Return the mesh defined by the node index
+	TFunction<void(const FThinZoneSide&, TArray<int32>&)> GetThinZoneSideMesh = [this](const FThinZoneSide& Side, TArray<int32>& MeshIndices)
+	{
+		FAddMeshNodeFunc AddMeshNode = [&MeshIndices](int32 NodeIndice, const FPoint2D& MeshNode2D, double MeshingTolerance, const FEdgeSegment& EdgeSegment, const FPairOfIndex& OppositeNodeIndices)
+		{
+			if (MeshIndices.IsEmpty() || MeshIndices.Last() != NodeIndice)
+			{
+				MeshIndices.Add(NodeIndice);
+			}
+		};
+
+		FReserveContainerFunc Reserve = [&MeshIndices](int32 MeshVertexCount)
+		{
+			MeshIndices.Reserve(MeshIndices.Num() + MeshVertexCount);
+		};
+		const bool bWithTolerance = false;
+		Side.GetExistingMeshNodes(GetFace(), MeshModel, Reserve, AddMeshNode, bWithTolerance);
+	};
+
+	TFunction<bool(int32, int32&, int32&)> FindLoopAndNodeIndex = [&NodeIdsOfLoops = NodeIdsOfFaceLoops](int32 NodeId, int32& LoopIndex, int32& LoopNodeIndex)
+	{
+		bool bFind = false;
+		for (LoopIndex = 0; LoopIndex < NodeIdsOfLoops.Num(); ++LoopIndex)
+		{
+			TArray<int32>& NodeIdsOfLoop = NodeIdsOfLoops[LoopIndex];
+			for (LoopNodeIndex = 0; LoopNodeIndex < NodeIdsOfLoop.Num(); ++LoopNodeIndex)
+			{
+				if (NodeIdsOfLoop[LoopNodeIndex] == NodeId)
+				{
+					return true;
+				}
+			}
+		}
+		LoopIndex = -1;
+		LoopNodeIndex = -1;
+		return false;
+	};
+
+	TFunction<void(const TArray<int32>&, TArray<FPoint2D>&)> GetThinZoneMeshCoordinates = [&NodeIdsOfLoops = NodeIdsOfFaceLoops, &Loops2D = FaceLoops2D, &FindLoopAndNodeIndex](const TArray<int32>& ThinZoneNodeIds, TArray<FPoint2D>& ThinZonePoints)
+	{
+		ThinZonePoints.Reset(ThinZoneNodeIds.Num());
+
+		int32 LoopIndex = 0;
+		int32 LoopNodeIndex = 0;
+
+		for (int32 NodeId : ThinZoneNodeIds)
+		{
+			if (LoopIndex >= 0)
+			{
+				const int32 NextLoopNodeIndex = (LoopNodeIndex == NodeIdsOfLoops[LoopIndex].Num() - 1) ? 0 : LoopNodeIndex + 1;
+				if (NodeIdsOfLoops[LoopIndex][NextLoopNodeIndex] == NodeId)
+				{
+					LoopNodeIndex = NextLoopNodeIndex;
+				}
+				else
+				{
+					const int32 PrevLoopNodeIndex = LoopNodeIndex ? LoopNodeIndex - 1 : NodeIdsOfLoops[LoopIndex].Num() - 1;
+					if (NodeIdsOfLoops[LoopIndex][PrevLoopNodeIndex] == NodeId)
+					{
+						LoopNodeIndex = PrevLoopNodeIndex;
+					}
+					else
+					{
+						LoopIndex = -1;
+					}
+				}
+			}
+
+			if (LoopIndex < 0 && !FindLoopAndNodeIndex(NodeId, LoopIndex, LoopNodeIndex))
+			{
+				continue;
+			}
+
+			ThinZonePoints.Emplace(Loops2D[EGridSpace::UniformScaled][LoopIndex][LoopNodeIndex]);
+		}
+	};
+
+	TFunction<void(TArray<FPoint2D>&)> RemoveDuplicatedNode = [](TArray<FPoint2D>& ThinZoneMesh)
+	{
+		if (ThinZoneMesh.Num() > 1)
+		{
+			// Remove Duplicated Node
+			constexpr double Tolerance = DOUBLE_SMALL_NUMBER;
+			for (int32 Index = ThinZoneMesh.Num() - 1; Index > 0; --Index)
+			{
+				double SquareDist = ThinZoneMesh[Index - 1].SquareDistance(ThinZoneMesh[Index]);
+				if (SquareDist < Tolerance)
+				{
+					ThinZoneMesh.RemoveAt(Index);
+				}
+			}
+			if (ThinZoneMesh.Num() > 2)
+			{
+				double SquareDist = ThinZoneMesh[0].SquareDistance(ThinZoneMesh.Last());
+				if (SquareDist < Tolerance)
+				{
+					ThinZoneMesh.RemoveAt(ThinZoneMesh.Num() - 1);
+				}
+			}
+		}
+	};
+
+	const bool FirstSideIsClosed = ThinZone.GetFirstSide().IsClosed();
+	const bool SecondSideIsClosed = ThinZone.GetSecondSide().IsClosed();
+
+	TArray<int32> ThinZoneNodeIds;
+	GetThinZoneSideMesh(ThinZone.GetFirstSide(), ThinZoneNodeIds);
+	const int32 FirstSideSize = ThinZoneNodeIds.Num();
+
+	if (FirstSideSize < 2)
+	{
+		return;
+	}
+
+	if (FirstSideIsClosed && SecondSideIsClosed)
+	{
+		TArray<FPoint2D>& ThinZoneMesh = FaceLoops2D[EGridSpace::UniformScaled].Emplace_GetRef();
+		GetThinZoneMeshCoordinates(ThinZoneNodeIds, ThinZoneMesh);
+		RemoveDuplicatedNode(ThinZoneMesh);
+		ThinZoneNodeIds.Reset();
+	}
+
+	GetThinZoneSideMesh(ThinZone.GetSecondSide(), ThinZoneNodeIds);
+	const int32 SecondSideSize = (FirstSideIsClosed && SecondSideIsClosed) ? ThinZoneNodeIds.Num() : ThinZoneNodeIds.Num() - FirstSideSize;
+	if (SecondSideSize < 2)
+	{
+		if (FirstSideIsClosed && SecondSideIsClosed)
+		{
+			FaceLoops2D[EGridSpace::UniformScaled].Pop();
+		}
+		return;
+	}
+
+	TArray<FPoint2D>& ThinZoneMesh = FaceLoops2D[EGridSpace::UniformScaled].Emplace_GetRef();
+	GetThinZoneMeshCoordinates(ThinZoneNodeIds, ThinZoneMesh);
+	RemoveDuplicatedNode(ThinZoneMesh);
+
+	if (ThinZoneMesh.Num() < 4)
+	{
+		FaceLoops2D[EGridSpace::UniformScaled].Pop();
+	}
+}
+
+bool FGrid::GetMeshOfLoops()
+{
+	int32 LoopCount = Face.GetLoops().Num();
+	FaceLoops2D[EGridSpace::Default2D].Reserve(LoopCount);
+
+	FaceLoops3D.Reserve(LoopCount);
+	NormalsOfFaceLoops.Reserve(LoopCount);
+	NodeIdsOfFaceLoops.Reserve(LoopCount);
+
+#ifdef DEBUG_GET_BOUNDARY_MESH
+	F3DDebugSession _(bDisplay, ("GetLoopMesh"));
+#endif
+
+	// Outer loops is processed first
+	for (const TSharedPtr<FTopologicalLoop>& Loop : Face.GetLoops())
+	{
+		if (Loop->IsExternal())
+		{
+			GetMeshOfLoop(*Loop);
+			break;
+		}
+	}
+
+	for (const TSharedPtr<FTopologicalLoop>& Loop : Face.GetLoops())
+	{
+		if (!Loop->IsExternal())
+		{
+			GetMeshOfLoop(*Loop);
+		}
+	}
+
+	if (CheckIfExternalLoopIsDegenerate())
+	{
+		return false;
+	}
+
+	ScaleLoops();
+ 
+	if (Face.HasThinZone())
+	{
+#ifdef DEBUG_GET_BOUNDARY_MESH
+		F3DDebugSession _(bDisplay, ("Thin zone loop"));
+#endif
+		for (const FThinZone2D& ThinZone : Face.GetThinZones())
+		{
+			GetMeshOfThinZone(ThinZone);
+#ifdef DEBUG_GET_BOUNDARY_MESH
+			if(bDisplay)
+			{
+				DisplayGridLoop(TEXT("ThinZone"), FaceLoops2D[EGridSpace::UniformScaled][FaceLoops2D[EGridSpace::UniformScaled].Num() - 1], true, true, false);
+				Wait();
+			}
+#endif
+		}
+#ifdef DEBUG_GET_BOUNDARY_MESH
+		//Wait();
+#endif
+	}
+
+	// Fit loops to Surface bounds.
+	const FSurfacicBoundary& Bounds = Face.GetBoundary();
+	for (TArray<FPoint2D>& Loop : FaceLoops2D[EGridSpace::Default2D])
+	{
+		for (FPoint2D& Point : Loop)
+		{
+			Bounds.MoveInsideIfNot(Point);
+		}
+	}
+
+	FSurfacicBoundary UniformScaleBounds;
+	UniformScaleBounds[IsoU].Set(UniformCuttingCoordinates[IsoU][0], UniformCuttingCoordinates[IsoU].Last());
+	UniformScaleBounds[IsoV].Set(UniformCuttingCoordinates[IsoV][0], UniformCuttingCoordinates[IsoV].Last());
+
+	// Fit loops in UniformScaled to UniformScale bounds.
+	for (TArray<FPoint2D>& Loop : FaceLoops2D[EGridSpace::UniformScaled])
+	{
+		for (FPoint2D& Point : Loop)
+		{
+			UniformScaleBounds.MoveInsideIfNot(Point);
+		}
+	}
+
+	return true;
+}
+
+void FGrid::ScaleLoops()
+{
+	int32 ThinZoneNum = 0;
+	if (Face.HasThinZone())
+	{
+		ThinZoneNum = Face.GetThinZones().Num();
+	}
+
+	FaceLoops2D[EGridSpace::Scaled].SetNum(FaceLoops2D[EGridSpace::Default2D].Num());
+	FaceLoops2D[EGridSpace::UniformScaled].Reserve(FaceLoops2D[EGridSpace::Default2D].Num() + ThinZoneNum);
+	FaceLoops2D[EGridSpace::UniformScaled].SetNum(FaceLoops2D[EGridSpace::Default2D].Num());
+
+	for (int32 IndexBoudnary = 0; IndexBoudnary < FaceLoops2D[EGridSpace::Default2D].Num(); ++IndexBoudnary)
+	{
+		const TArray<FPoint2D>& Loop = FaceLoops2D[EGridSpace::Default2D][IndexBoudnary];
+		TArray<FPoint2D>& ScaledLoop = FaceLoops2D[EGridSpace::Scaled][IndexBoudnary];
+		TArray<FPoint2D>& UniformScaledLoop = FaceLoops2D[EGridSpace::UniformScaled][IndexBoudnary];
+
+		ScaledLoop.SetNum(Loop.Num());
+		UniformScaledLoop.SetNum(Loop.Num());
+
+		int32 IndexU = 0;
+		int32 IndexV = 0;
+		for (int32 Index = 0; Index < Loop.Num(); ++Index)
+		{
+			const FPoint2D& Point = Loop[Index];
+
+			ArrayUtils::FindCoordinateIndex(CoordinateGrid[EIso::IsoU], Point.U, IndexU);
+			ArrayUtils::FindCoordinateIndex(CoordinateGrid[EIso::IsoV], Point.V, IndexV);
+
+			ComputeNewCoordinate(Points2D[EGridSpace::Scaled], IndexU, IndexV, Point, ScaledLoop[Index]);
+			ComputeNewCoordinate(Points2D[EGridSpace::UniformScaled], IndexU, IndexV, Point, UniformScaledLoop[Index]);
+		}
+	}
+}
+
+bool FGrid::CheckIf2DGridIsDegenerate() const
+{
+	TFunction<bool(EIso)> IsDegenerate = [&](EIso Iso) -> bool
+	{
+		double MaxDelta = 0;
+		for (int32 Index = 1; Index < CoordinateGrid[Iso].Num(); ++Index)
+		{
+			double Delta = CoordinateGrid[Iso][Index] - CoordinateGrid[Iso][Index - 1];
+			if (Delta > MaxDelta)
+			{
+				MaxDelta = Delta;
+			}
+		}
+		return MaxDelta < FaceTolerance[Iso];
+	};
+
+	if (IsDegenerate(EIso::IsoU) || IsDegenerate(EIso::IsoV))
+	{
+		SetAsDegenerated();
+		return true;
+	}
+	return false;
+}
+
+void FGrid::FindInnerFacePoints()
+{
+	// FindInnerDomainPoints: Inner Points <-> bIsOfInnerDomain = true
+	// For each points count the number of intersection with the boundary in the four directions U+ U- V+ V-
+	// It for each the number is pair, the point is outside,
+	// If in 3 directions the point is inner, the point is inner else we have a doubt so it preferable to consider it outside. 
+	// Most of the time, there is a doubt if the point is to close of the boundary. So it will be removed be other criteria
+
+	FTimePoint StartTime = FChrono::Now();
+
+	TFunction<void(TArray<char>&, int32)> AddIntersection = [](TArray<char>& Intersect, int32 Index)
+	{
+		Intersect[Index] = Intersect[Index] > 0 ? 0 : 1;
+	};
+
+	TArray<char> VForwardIntersectCount; // we need to know if intersect is pair 0, 2, 4... intersection of impair 1, 3, 5... false is pair, true is impair 
+	TArray<char> VBackwardIntersectCount;
+	TArray<char> UForwardIntersectCount;
+	TArray<char> UBackwardIntersectCount;
+	TArray<char> IntersectLoop;
+
+	IntersectLoop.Init(0, CuttingSize);
+	NodeMarkers.Init(ENodeMarker::IsInside, CuttingSize);
+
+	VForwardIntersectCount.Init(0, CuttingSize);
+	VBackwardIntersectCount.Init(0, CuttingSize);
+	UForwardIntersectCount.Init(0, CuttingSize);
+	UBackwardIntersectCount.Init(0, CuttingSize);
+
+	// Loop node too close to one of CoordinateU or CoordinateV are moved a little to avoid floating error of comparison 
+	// This step is necessary instead of all points could be considered outside...
+	const double SmallToleranceU = DOUBLE_SMALL_NUMBER;
+	const double SmallToleranceV = DOUBLE_SMALL_NUMBER;
+
+	{
+		int32 IndexV = 0;
+		int32 IndexU = 0;
+		for (TArray<FPoint2D>& Loop : FaceLoops2D[EGridSpace::UniformScaled])
+		{
+			for (FPoint2D& Point : Loop)
+			{
+				while (IndexV != 0 && (Point.V < UniformCuttingCoordinates[EIso::IsoV][IndexV]))
+				{
+					IndexV--;
+				}
+				for (; IndexV < CuttingCount[EIso::IsoV]; ++IndexV)
+				{
+					if (Point.V + SmallToleranceV < UniformCuttingCoordinates[EIso::IsoV][IndexV])
+					{
+						break;
+					}
+					if (Point.V - SmallToleranceV > UniformCuttingCoordinates[EIso::IsoV][IndexV])
+					{
+						continue;
+					}
+
+					if (IndexV == 0)
+					{
+						Point.V += SmallToleranceV;
+					}
+					else
+					{
+						Point.V -= SmallToleranceV;
+					}
+
+					break;
+				}
+				if (IndexV == CuttingCount[EIso::IsoV])
+				{
+					IndexV--;
+				}
+
+				while (IndexU != 0 && (Point.U < UniformCuttingCoordinates[EIso::IsoU][IndexU]))
+				{
+					IndexU--;
+				}
+				for (; IndexU < CuttingCount[EIso::IsoU]; ++IndexU)
+				{
+					if (Point.U + SmallToleranceU < UniformCuttingCoordinates[EIso::IsoU][IndexU])
+					{
+						break;
+					}
+					if (Point.U - SmallToleranceU > UniformCuttingCoordinates[EIso::IsoU][IndexU])
+					{
+						continue;
+					}
+
+					if (IndexU == 0)
+					{
+						Point.U += SmallToleranceU;
+					}
+					else
+					{
+						Point.U -= SmallToleranceU;
+					}
+					break;
+				}
+				if (IndexU == CuttingCount[EIso::IsoU])
+				{
+					IndexU--;
+				}
+			}
+		}
+	}
+
+#ifdef DEBUG_FIND_INNER_FACE_POINTS
+	DisplayGridLoops(TEXT("FGrid::Loop 2D After move according tol"), GetLoops2D(EGridSpace::UniformScaled), true, false);
+#endif
+
+	// Intersection along U axis
+	for (const TArray<FPoint2D>& Loop : FaceLoops2D[EGridSpace::UniformScaled])
+	{
+		const FPoint2D* FirstSegmentPoint = &Loop.Last();
+		for (const FPoint2D& LoopPoint : Loop)
+		{
+			const FPoint2D* SecondSegmentPoint = &LoopPoint;
+			double UMin = FirstSegmentPoint->U;
+			double VMin = FirstSegmentPoint->V;
+			double Umax = SecondSegmentPoint->U;
+			double Vmax = SecondSegmentPoint->V;
+			GetMinMax(UMin, Umax);
+			GetMinMax(VMin, Vmax);
+
+			// AB^AP = ABu*APv - ABv*APu
+			// AB^AP = ABu*(Pv-Av) - ABv*(Pu-Au)
+			// AB^AP = Pv*ABu - Pu*ABv + Au*ABv - Av*ABu
+			// AB^AP = Pv*ABu - Pu*ABv + AuABVMinusAvABu
+			FPoint2D PointA;
+			FPoint2D PointB;
+			if (FirstSegmentPoint->V < SecondSegmentPoint->V)
+			{
+				PointA = *FirstSegmentPoint;
+				PointB = *SecondSegmentPoint;
+			}
+			else
+			{
+				PointA = *SecondSegmentPoint;
+				PointB = *FirstSegmentPoint;
+			}
+			double ABv = PointB.V - PointA.V;
+			double ABu = PointB.U - PointA.U;
+			double AuABVMinusAvABu = PointA.U * ABv - PointA.V * ABu;
+
+			int32 IndexV = 0;
+			int32 Index = 0;
+			for (; IndexV < CuttingCount[EIso::IsoV]; ++IndexV)
+			{
+				if (UniformCuttingCoordinates[EIso::IsoV][IndexV] >= VMin)
+				{
+					break;
+				}
+				Index += CuttingCount[EIso::IsoU];
+			}
+
+			for (; IndexV < CuttingCount[EIso::IsoV]; ++IndexV)
+			{
+				if (UniformCuttingCoordinates[EIso::IsoV][IndexV] > Vmax)
+				{
+					break;
+				}
+
+				for (int32 IndexU = 0; IndexU < CuttingCount[EIso::IsoU]; ++IndexU, ++Index)
+				{
+					//Index = IndexV * NumU + IndexU;
+					if (IntersectLoop[Index])
+					{
+						continue;
+					}
+
+					if (UniformCuttingCoordinates[EIso::IsoU][IndexU] < UMin)
+					{
+						AddIntersection(UForwardIntersectCount, Index);
+					}
+					else if (UniformCuttingCoordinates[EIso::IsoU][IndexU] > Umax)
+					{
+						AddIntersection(UBackwardIntersectCount, Index);
+					}
+					else
+					{
+						double APvectAB = UniformCuttingCoordinates[EIso::IsoV][IndexV] * ABu - UniformCuttingCoordinates[EIso::IsoU][IndexU] * ABv + AuABVMinusAvABu;
+						if (APvectAB > DOUBLE_SMALL_NUMBER)
+						{
+							AddIntersection(UForwardIntersectCount, Index);
+						}
+						else if (APvectAB < DOUBLE_SMALL_NUMBER)
+						{
+							AddIntersection(UBackwardIntersectCount, Index);
+						}
+						else
+						{
+							IntersectLoop[Index] = 1;
+						}
+					}
+				}
+			}
+			FirstSegmentPoint = SecondSegmentPoint;
+		}
+	}
+
+	// Intersection along V axis
+	for (const TArray<FPoint2D>& Loop : FaceLoops2D[EGridSpace::UniformScaled])
+	{
+		const FPoint2D* FirstSegmentPoint = &Loop.Last();
+		for (const FPoint2D& LoopPoint : Loop)
+		{
+			const FPoint2D* SecondSegmentPoint = &LoopPoint;
+			double UMin = FirstSegmentPoint->U;
+			double VMin = FirstSegmentPoint->V;
+			double Umax = SecondSegmentPoint->U;
+			double Vmax = SecondSegmentPoint->V;
+			GetMinMax(UMin, Umax);
+			GetMinMax(VMin, Vmax);
+
+			// AB^AP = ABu*APv - ABv*APu
+			// AB^AP = ABu*(Pv-Av) - ABv*(Pu-Au)
+			// AB^AP = Pv*ABu - Pu*ABv + Au*ABv - Av*ABu
+			// AB^AP = Pv*ABu - Pu*ABv + AuABVMinusAvABu
+			FPoint2D PointA;
+			FPoint2D PointB;
+			if (FirstSegmentPoint->U < SecondSegmentPoint->U)
+			{
+				PointA = *FirstSegmentPoint;
+				PointB = *SecondSegmentPoint;
+			}
+			else
+			{
+				PointA = *SecondSegmentPoint;
+				PointB = *FirstSegmentPoint;
+			}
+
+			double ABu = PointB.U - PointA.U;
+			double ABv = PointB.V - PointA.V;
+			double AuABVMinusAvABu = PointA.U * ABv - PointA.V * ABu;
+			int32 Index = 0;
+			for (int32 IndexU = 0; IndexU < CuttingCount[EIso::IsoU]; ++IndexU)
+			{
+				if (UniformCuttingCoordinates[EIso::IsoU][IndexU] < UMin)
+				{
+					continue;
+				}
+
+				if (UniformCuttingCoordinates[EIso::IsoU][IndexU] >= Umax)
+				{
+					continue;
+				}
+
+				Index = IndexU;
+				for (int32 IndexV = 0; IndexV < CuttingCount[EIso::IsoV]; ++IndexV, Index += CuttingCount[EIso::IsoU])
+				{
+					if (IntersectLoop[Index])
+					{
+						continue;
+					}
+
+					if (UniformCuttingCoordinates[EIso::IsoV][IndexV] < VMin)
+					{
+						AddIntersection(VForwardIntersectCount, Index);
+					}
+					else if (UniformCuttingCoordinates[EIso::IsoV][IndexV] > Vmax)
+					{
+						AddIntersection(VBackwardIntersectCount, Index);
+					}
+					else
+					{
+						double APvectAB = UniformCuttingCoordinates[EIso::IsoV][IndexV] * ABu - UniformCuttingCoordinates[EIso::IsoU][IndexU] * ABv + AuABVMinusAvABu;
+						if (APvectAB > DOUBLE_SMALL_NUMBER)
+						{
+							AddIntersection(VBackwardIntersectCount, Index);
+						}
+						else if (APvectAB < DOUBLE_SMALL_NUMBER)
+						{
+							AddIntersection(VForwardIntersectCount, Index);
+						}
+						else
+						{
+							IntersectLoop[Index] = 1;
+						}
+					}
+				}
+			}
+			FirstSegmentPoint = SecondSegmentPoint;
+		}
+	}
+
+	for (int32 Index = 0; Index < CuttingSize; ++Index)
+	{
+		if (IntersectLoop[Index])
+		{
+			ResetInsideLoop(Index);
+			CountOfInnerNodes--;
+			continue;
+		}
+
+		int32 IsInside = 0;
+		if (UForwardIntersectCount[Index] > 0)
+		{
+			IsInside++;
+		}
+		if (UBackwardIntersectCount[Index] > 0)
+		{
+			IsInside++;
+		}
+		if (VForwardIntersectCount[Index] > 0)
+		{
+			IsInside++;
+		}
+		if (VBackwardIntersectCount[Index] > 0)
+		{
+			IsInside++;
+		}
+		if (IsInside < 3)
+		{
+			ResetInsideLoop(Index);
+			CountOfInnerNodes--;
+		}
+	}
+
+	Chronos.FindInnerDomainPointsDuration += FChrono::Elapse(StartTime);
+}
+
+bool FGrid::CheckIfExternalLoopIsDegenerate() const
+{
+	if (FaceLoops2D[EGridSpace::Default2D].Num() == 0)
+	{
+		SetAsDegenerated();
+		return true;
+	}
+
+	// if the external boundary is composed by only 2 points, the mesh of the surface is only an edge.
+	// The grid is degenerated.
+	if (FaceLoops2D[EGridSpace::Default2D][0].Num() < 3)
+	{
+		SetAsDegenerated();
+		return true;
+	}
+
+	return false;
+}
+
+} // NS CADKernel
+
+
