Moved CADKernel library code from /Engine/Source/Runtime/Datasmith/CADKernel to /Engine/Source/Runtime/Datasmith/CADKernel/Base

diff --git a/Engine/Source/Runtime/Datasmith/CADKernel/Base/Private/Mesh/Structure/Grid.h b/Engine/Source/Runtime/Datasmith/CADKernel/Base/Private/Mesh/Structure/Grid.h
new file mode 100644
--- /dev/null
+++ b/Engine/Source/Runtime/Datasmith/CADKernel/Base/Private/Mesh/Structure/Grid.h
@@ -0,0 +1,603 @@
+// Copyright Epic Games, Inc. All Rights Reserved.
+#pragma once
+
+#include "Core/Chrono.h"
+#include "Geo/GeoEnum.h"
+#include "Geo/GeoPoint.h"
+#include "Math/Point.h"
+#include "Mesh/MeshEnum.h"
+#include "Mesh/Structure/GridBase.h"
+#include "Mesh/Structure/ThinZone2D.h"
+#include "UI/Display.h"
+
+#ifdef CADKERNEL_DEV
+#include "UI/DefineForDebug.h"
+#endif
+
+
+namespace UE::CADKernel
+{
+
+struct FCuttingPoint;
+struct FCuttingGrid;
+
+class FIsoNode;
+class FIsoSegment;
+class FLoopNode;
+class FModelMesh;
+class FPoint;
+class FPoint2D;
+class FThinZone2DFinder;
+class FTopologicalFace;
+
+
+enum class ENodeMarker : uint8
+{
+	None = 0x00u,  // No flags.
+
+	IsInside = 0x01u,
+	IsInsideButTooCloseToLoop = 0x02u,  // node inside the loop but too close to the loop to be include in the mesh
+	IsCloseToLoop = 0x04u,
+
+	All = 0xFFu
+};
+ENUM_CLASS_FLAGS(ENodeMarker);
+
+
+class FGrid : public FGridBase
+{
+protected:
+
+	/*
+	 * Cutting coordinates of the face respecting the meshing criteria
+	 */
+	const FCoordinateGrid& CoordinateGrid;
+
+	virtual const FCoordinateGrid& GetCoordinateGrid() const override
+	{
+		return CoordinateGrid;
+	}
+
+	const FSurfacicTolerance FaceTolerance;
+	const double MinimumElementSize;
+
+	FModelMesh& MeshModel;
+
+
+	/**
+	 * 2D Coordinate of Loop's nodes in each space
+	 */
+	TArray<TArray<FPoint2D>> FaceLoops2D[EGridSpace::EndGridSpace];
+
+	/**
+	 * 3D Coordinate of Loop nodes in each space
+	 */
+	TArray<TArray<FPoint>> FaceLoops3D;
+
+	/**
+	 * Surface Normal at each boundary nodes
+	 */
+	TArray<TArray<FVector3f>> NormalsOfFaceLoops;
+
+	TArray<TArray<int32>> NodeIdsOfFaceLoops;
+
+	/**
+	 * count of node inside the face i.e. node inside the external loop and outer the inner loops
+	 */
+	int32 CountOfInnerNodes = 0;
+
+	double MinOfMaxElementSize = 0;
+
+	/**
+	 * Array to flag each points as inside, close, ... the face
+	 * @see IsNode...,
+	 */
+	TArray<ENodeMarker> NodeMarkers;
+
+public:
+#ifdef CADKERNEL_DEV
+	FGridChronos Chronos;
+#endif
+
+	FGrid(FTopologicalFace& InFace, FModelMesh& InShellMesh);
+
+#ifdef CADKERNEL_DEBUG
+	virtual ~FGrid()
+	{
+		Close3DDebugSession(bDisplay);
+	}
+#else
+	virtual ~FGrid() = default;
+#endif
+
+	// ======================================================================================================================================================================================================================
+	// Meshing tools =========================================================================================================================================================================================================
+	// ======================================================================================================================================================================================================================
+
+	/**
+	 * Defines the cutting coordinate of the grid according to mesh criteria and the existing mesh of bordering edges (loop's edges)
+	 * @see GetPreferredUVCoordinatesFromNeighbours
+	 */
+	void DefineCuttingParameters();
+	void DefineCuttingParameters(EIso Iso, FCuttingGrid& Neighbors);
+
+	/**
+	 * Computes the 2d Points, 3D points and Normals of the grid.
+	 * call of DefineCuttingParameters is mandatory called before
+	 * @return false if the grid is degenerated
+	 */
+	bool GeneratePointCloud();
+
+	/**
+	 * Process the generated point cloud to compute the scaled parametric spaces and to identify outer points of the face
+	 * @see IsNodeInsideFace
+	 * @see Points2D
+	 */
+	void ProcessPointCloud();
+
+protected:
+
+	/**
+	 * Projects loop's points in the scaled parametric spaces
+	 * @see FaceLoops2D
+	 * @see ProcessPointClound (called in ProcessPointClound)
+	 */
+	void ScaleLoops();
+
+	/**
+	 * @see ProcessPointClound (called in ProcessPointClound)
+	 */
+	void FindInnerFacePoints();
+
+	/**
+	 * Define if an inner node is close to a loop i.e. if the loop cross the cell [[IndexU - 1, IndexU], [IndexV - 1, IndexV]], each corner node of the cell is close to the loop
+	 *
+	 * This algorithm is inspired of Bresenham's line algorithm
+	 * The loop is traversed segment by segment
+	 * For each segment, according to the slop of the segment, each cell intersecting the segment is selected (FindIntersectionsCaseSlope_) i.e. each corner node of the cell is flagged IsCloseToLoop (SetCellCloseToLoop)
+	 * @see SlopeUtils
+	 * @see IsNodeCloseToLoop
+	 * @see ProcessPointCloud (called in ProcessPointClound)
+	 */
+	void FindPointsCloseToLoop();
+
+	/**
+	 * Find all node close to the boundary.
+	 * For each node, the 3d distance to the boundary is estimated
+	 * If this distance is bigger than the Tolerance, the node is remove
+	 * Removing a node can make its neighbor removable,
+	 * Each removed node make its direct neighbor removable, the process is recursive as long as it is possible to remove a node
+	 *
+	 * @see ProcessPointClound (called in ProcessPointClound)
+	 */
+	void RemovePointsCloseToLoop();
+
+	/**
+	 * Gets the cutting coordinates of the existing mesh of bordering edges (loop's edges)
+	 * @see DefineCuttingParameters (called in DefineCuttingParameters)
+	 */
+	void GetPreferredUVCuttingParametersFromLoops(FCuttingGrid& CuttingFromLoops);
+
+	/**
+	 * @ return false if the mesh of the loop is degenerated
+	 */
+	bool GetMeshOfLoops();
+	void GetMeshOfLoop(const FTopologicalLoop& Loop);
+	void GetMeshOfThinZone(const FThinZone2D& ThinZone);
+
+	// Meshing tools =========================================================================================================================================================================================================
+	// ======================================================================================================================================================================================================================
+
+	/**
+	 * @return true if the grid is not consistent to build a mesh as composed of only two border nodes.
+	 */
+	bool CheckIfExternalLoopIsDegenerate() const;
+
+	/**
+	 * @return true if along an Iso, all (U(i+1) - U(i)) < FaceTolerance
+	 */
+	bool CheckIf2DGridIsDegenerate() const;
+
+public:
+
+	// ======================================================================================================================================================================================================================
+	// GET Methodes =========================================================================================================================================================================================================
+	// ======================================================================================================================================================================================================================
+
+	constexpr const int32 GetCuttingCount(EIso Iso) const
+	{
+		return CuttingCount[Iso];
+	}
+
+	/**
+	 * Return the number of points (inner and outer points) of the grid (i.e. CuttingCount[IsoU] x CuttingCount[IsoV]).
+	 */
+	const int32 GetTotalCuttingCount() const
+	{
+		return CuttingSize;
+	}
+
+	const double GetTolerance(EIso Iso) const
+	{
+		return FaceTolerance[Iso];
+	}
+
+	/**
+	 * @return true if the node is inner the external loop and outer the inner loops and not too close
+	 */
+	const bool IsNodeInsideAndMeshable(int32 Index) const
+	{
+		return (NodeMarkers[Index] & ENodeMarker::IsInside) == ENodeMarker::IsInside;
+	}
+
+	const bool IsNodeInsideButTooCloseToLoop(int32 Index) const
+	{
+		return (NodeMarkers[Index] & ENodeMarker::IsInsideButTooCloseToLoop) == ENodeMarker::IsInsideButTooCloseToLoop;
+	}
+
+	/**
+	 * @return true if the bits EPointMarker::IsInside, EPointMarker::IsInsideButTooCloseToLoop and EPointMarker::IsCloseToLoop are alls equal to 0
+	 */
+	const bool IsNodeFarFromFace(int32 Index) const
+	{
+		constexpr ENodeMarker IsInsideAndClose = ENodeMarker::IsInside | ENodeMarker::IsInsideButTooCloseToLoop | ENodeMarker::IsCloseToLoop;
+		return (NodeMarkers[Index] & IsInsideAndClose) == ENodeMarker::None;
+	}
+
+	/**
+	 * @return true if the bits EPointMarker::IsInside, EPointMarker::IsInsideButTooCloseToLoop and EPointMarker::IsCloseToLoop are alls equal to 0
+	 */
+	const bool IsNodeOutsideFace(int32 Index) const
+	{
+		constexpr ENodeMarker IsInsideAndClose = ENodeMarker::IsInside | ENodeMarker::IsInsideButTooCloseToLoop;
+		return (NodeMarkers[Index] & IsInsideAndClose) == ENodeMarker::None;
+	}
+
+	const bool IsNodeInsideAndCloseToLoop(int32 Index) const
+	{
+		constexpr ENodeMarker IsInsideAndClose = ENodeMarker::IsInside | ENodeMarker::IsCloseToLoop;
+		return (NodeMarkers[Index] & IsInsideAndClose) == IsInsideAndClose;
+	}
+
+	/**
+	 * @return true if the bits EPointMarker::IsInside and EPointMarker::IsInsideButTooCloseToLoop are booth equal to 0 and EPointMarker::IsCloseToLoop  is equal to 1
+	 */
+	const bool IsNodeOusideFaceButClose(int32 Index) const
+	{
+		constexpr ENodeMarker IsInsideAndClose = ENodeMarker::IsInside | ENodeMarker::IsInsideButTooCloseToLoop | ENodeMarker::IsCloseToLoop;
+		return (NodeMarkers[Index] & IsInsideAndClose) == ENodeMarker::IsCloseToLoop;
+	}
+
+	void SetNodeInside(int32 Index)
+	{
+		NodeMarkers[Index] |= ENodeMarker::IsInside;
+	}
+
+	/**
+	 * @return the count of inner nodes
+	 */
+	const int32 InnerNodesCount() const
+	{
+		return CountOfInnerNodes;
+	}
+
+	/**
+	 * @return true if the node is close to a loop i.e. the loop cross the space [[IndexU - 1, IndexU + 1], [IndexV - 1, IndexV + 1]]
+	 */
+	const bool IsNodeCloseToLoop(int32 Index) const
+	{
+		return (NodeMarkers[Index] & ENodeMarker::IsCloseToLoop) == ENodeMarker::IsCloseToLoop;
+	}
+
+	const bool IsNodeTooCloseToLoop(int32 Index) const
+	{
+		return (NodeMarkers[Index] & ENodeMarker::IsInsideButTooCloseToLoop) == ENodeMarker::IsInsideButTooCloseToLoop;
+	}
+
+	void SetCloseToLoop(int32 Index)
+	{
+		NodeMarkers[Index] |= ENodeMarker::IsCloseToLoop;
+	}
+
+	void SetTooCloseToLoop(int32 Index)
+	{
+		NodeMarkers[Index] |= ENodeMarker::IsInsideButTooCloseToLoop;
+		NodeMarkers[Index] &= ~ENodeMarker::IsInside;
+	}
+
+	void ResetInsideLoop(int32 Index)
+	{
+		NodeMarkers[Index] &= ~ENodeMarker::IsInside;
+	}
+
+	void SetInner2DPoint(EGridSpace Space, int32 Index, const FPoint2D& NewCoordinate)
+	{
+		Points2D[(int32)Space][Index] = NewCoordinate;
+	}
+
+	/**
+	 * @return the normal of the surface at the point at the Index of the grid
+	 */
+	const FVector3f& GetPointNormal(int32 IndexU, int32 IndexV) const
+	{
+		return Normals[GobalIndex(IndexU, IndexV)];
+	}
+
+	/**
+	 * @return the normal of the surface at the point at the Index of the grid
+	 */
+	const FVector3f& GetPointNormal(int32 Index) const
+	{
+		return Normals[Index];
+	}
+
+	const TArray<double>& GetCuttingCoordinatesAlongIso(EIso Iso) const
+	{
+		return CoordinateGrid[Iso];
+	}
+
+	const FCoordinateGrid& GetCuttingCoordinates() const
+	{
+		return CoordinateGrid;
+	}
+
+	/**
+	 * @return the array of normal of the points of the grid
+	 */
+	TArray<FVector3f>& GetNormals()
+	{
+		return Normals;
+	}
+
+	const TArray<TArray<int32>>& GetNodeIdsOfFaceLoops() const
+	{
+		return NodeIdsOfFaceLoops;
+	}
+
+	const FPoint2D& GetLoop2DPoint(EGridSpace Space, int32 LoopIndex, int32 Index) const
+	{
+		return FaceLoops2D[(int32)Space][LoopIndex][Index];
+	}
+
+	void SetLoop2DPoint(EGridSpace Space, int32 LoopIndex, int32 Index, const FPoint2D& NewCoordinate)
+	{
+		FaceLoops2D[(int32)Space][LoopIndex][Index] = NewCoordinate;
+	}
+
+	const FPoint& GetLoop3DPoint(int32 LoopIndex, int32 Index) const
+	{
+		return FaceLoops3D[LoopIndex][Index];
+	}
+
+	const int32 GetLoopCount() const
+	{
+		return FaceLoops2D[0].Num();
+	}
+
+	/**
+	 * @return the array of array of 2d points of the loops according to the defined grid space
+	 */
+	const TArray<TArray<FPoint2D>>& GetLoops2D(EGridSpace Space) const
+	{
+		return FaceLoops2D[(int32)Space];
+	}
+
+	/**
+	 * @return the array of array of 3d points of the loops
+	 */
+	TArray<TArray<FPoint>>& GetLoops3D()
+	{
+		return FaceLoops3D;
+	}
+
+	/**
+	 * @return the array of array of 3d points of the loops
+	 */
+	const TArray<TArray<FPoint>>& GetLoops3D() const
+	{
+		return FaceLoops3D;
+	}
+
+	/**
+	 * @return the array of array of normal of the points of the loops
+	 */
+	const TArray<TArray<FVector3f>>& GetLoopNormals() const
+	{
+		return NormalsOfFaceLoops;
+	}
+
+	/**
+	 * @return the Index of the position in the arrays of a point [IndexU, IndexV] of the grid
+	 */
+	void UVIndexFromGlobalIndex(int32 GLobalIndex, int32& OutIndexU, int32& OutIndexV) const
+	{
+		OutIndexU = GLobalIndex % CuttingCount[EIso::IsoU];
+		OutIndexV = GLobalIndex / CuttingCount[EIso::IsoU];
+	}
+
+	/**
+	 * @return the minimal size of the mesh elements according to 3d geometric tolerance
+	 */
+	double GetMinElementSize() const
+	{
+		return MinimumElementSize;
+	}
+
+	// ======================================================================================================================================================================================================================
+	// Display Methodes   ================================================================================================================================================================================================
+	// ======================================================================================================================================================================================================================
+#ifdef CADKERNEL_DEV
+	void PrintTimeElapse() const;
+#endif
+
+#ifdef CADKERNEL_DEBUG
+
+	void DisplayIsoNode(EGridSpace Space, const int32 PointIndex, FIdent Ident = 0, EVisuProperty Property = EVisuProperty::BluePoint) const;
+	void DisplayIsoNode(EGridSpace Space, const FIsoNode& Node, FIdent Ident = 0, EVisuProperty Property = EVisuProperty::BluePoint) const;
+	void DisplayIsoNodes(EGridSpace Space, const TArray<const FIsoNode*>& Nodes, EVisuProperty Property = EVisuProperty::BluePoint) const;
+
+	void DisplayIsoNodes(const FString& Message, EGridSpace Space, const TArray<const FIsoNode*>& Nodes, EVisuProperty Property = EVisuProperty::BluePoint) const
+	{
+		if (!bDisplay)
+		{
+			return;
+		}
+		F3DDebugSession _(Message);
+		DisplayIsoNodes(Space, Nodes, Property);
+	}
+
+	void DisplayIsoPolyline(const FString& Message, EGridSpace Space, const TArray<const FIsoNode*>& Nodes, EVisuProperty Property = EVisuProperty::BluePoint) const
+	{
+		if (!bDisplay)
+		{
+			return;
+		}
+		F3DDebugSession _(Message);
+		const FIsoNode* LastNode = Nodes.Last();
+		for (const FIsoNode* Node : Nodes)
+		{
+			DisplayIsoSegment(Space, *LastNode, *Node, 0, (EVisuProperty)(Property + 1));
+			LastNode = Node;
+		}
+		DisplayIsoNodes(Space, Nodes, Property);
+	}
+
+	void DisplayIsoSegment(EGridSpace Space, const FIsoSegment& Segment, FIdent Ident = 0, EVisuProperty Property = EVisuProperty::BlueCurve, bool bDisplayOrientation = false) const;
+	void DisplayIsoSegment(EGridSpace Space, const FIsoNode& NodeA, const FIsoNode& NodeB, FIdent Ident = 0, EVisuProperty Property = EVisuProperty::BlueCurve) const;
+
+	void DisplayIsoSegments(EGridSpace Space, const TArray<FIsoSegment*>& Segments, bool bDisplayNode = false, bool bDisplayOrientation = false, bool bSplitBySegment = false, EVisuProperty Property = EVisuProperty::BlueCurve) const;
+	void DisplayIsoSegments(const FString& Message, EGridSpace Space, const TArray<FIsoSegment*>& Segments, bool bDisplayNode = false, bool bDisplayOrientation = false, bool bSplitBySegment = false, EVisuProperty Property = EVisuProperty::BlueCurve) const
+	{
+		if (!bDisplay)
+		{
+			return;
+		}
+		F3DDebugSession _(Message);
+		DisplayIsoSegments(Space, Segments, bDisplayNode, bDisplayOrientation, bSplitBySegment, Property);
+	}
+
+	void DisplayTriangle(EGridSpace Space, const FIsoNode& NodeA, const FIsoNode& NodeB, const FIsoNode& NodeC) const;
+	void DisplayTriangle3D(const FIsoNode& NodeA, const FIsoNode& NodeB, const FIsoNode& NodeC) const;
+
+	void DisplayGridPolyline(EGridSpace Space, const FLoopNode& StartNode, bool bDisplayNode, EVisuProperty Property = EVisuProperty::BlueCurve) const;
+	void DisplayGridPolyline(EGridSpace Space, const TArray<FIsoNode*>& Nodes, bool bDisplayNode, EVisuProperty Property = EVisuProperty::BlueCurve) const;
+	void DisplayGridPolyline(EGridSpace Space, const TArray<FLoopNode*>& Nodes, bool bDisplayNode, EVisuProperty Property = EVisuProperty::BlueCurve) const;
+	void DisplayGridPolyline(EGridSpace Space, const TArray<FLoopNode>& Nodes, bool bDisplayNode, EVisuProperty Property = EVisuProperty::BlueCurve) const;
+	void DisplayGridPolyline(const FString& Message, EGridSpace Space, const TArray<FLoopNode*>& Nodes, bool bDisplayNode, EVisuProperty Property = EVisuProperty::BlueCurve) const
+	{
+		if (!bDisplay)
+		{
+			return;
+		}
+		F3DDebugSession _(Message);
+		DisplayGridPolyline(Space, Nodes, bDisplayNode, Property);
+	}
+
+	void DisplayGridPolyline(const FString& Message, EGridSpace Space, const TArray<FLoopNode>& Nodes, bool bDisplayNode, EVisuProperty Property = EVisuProperty::BlueCurve) const
+	{
+		if (!bDisplay)
+		{
+			return;
+		}
+		F3DDebugSession _(Message);
+		DisplayGridPolyline(Space, Nodes, bDisplayNode, Property);
+	}
+
+	void DisplayNodes(const TCHAR* Message, EGridSpace Space, const TArray<const FIsoNode*>& Nodes, EVisuProperty Property) const;
+
+	virtual void DisplayGridPoints(EGridSpace DisplaySpace) const override;
+	void DisplayInnerPoints(TCHAR* Message, EGridSpace DisplaySpace) const;
+
+	void DisplayGridNormal() const;
+
+	template<typename TPoint>
+	void DisplayInnerDomainPoints(FString Message, const TArray<TPoint>& Points) const
+	{
+		if (!bDisplay)
+		{
+			return;
+		}
+
+		F3DDebugSession _(Message);
+		for (int32 Index = 0; Index < Points.Num(); ++Index)
+		{
+			if (IsNodeInsideAndCloseToLoop(Index))
+			{
+				DisplayPoint2DWithScale(Points[Index], EVisuProperty::OrangePoint, Index);
+			}
+			else if (IsNodeInsideAndMeshable(Index))
+			{
+				DisplayPoint2DWithScale(Points[Index], EVisuProperty::BluePoint, Index);
+			}
+			else if (IsNodeTooCloseToLoop(Index))
+			{
+				DisplayPoint2DWithScale(Points[Index], EVisuProperty::RedPoint, Index);
+			}
+			else
+			{
+				DisplayPoint2DWithScale(Points[Index], IsNodeCloseToLoop(Index) ? EVisuProperty::YellowPoint : EVisuProperty::GreenPoint, Index);
+			}
+		}
+	}
+
+	template<typename TPoint>
+	void DisplayGridLoop(FString Message, const TArray<TPoint>& Loop, bool bDisplayNodes, bool bMakeGroup, bool bDisplaySegments) const
+	{
+		if (!bDisplay)
+		{
+			return;
+		}
+
+		F3DDebugSession _(bMakeGroup, Message);
+
+		const TPoint* FirstSegmentPoint = &Loop[0];
+		if (bDisplayNodes)
+		{
+			F3DDebugSession Q(bDisplaySegments, FString::Printf(TEXT("FirstNode")));
+			DisplayPoint(*FirstSegmentPoint * DisplayScale, EVisuProperty::BluePoint, 0);
+		}
+
+		for (int32 Index = 1; Index < Loop.Num(); ++Index)
+		{
+			const TPoint* SecondSegmentPoint = &Loop[Index];
+			F3DDebugSession B(bDisplaySegments, FString::Printf(TEXT("Segment %d"), Index));
+			DisplaySegment<TPoint>(*FirstSegmentPoint * DisplayScale, *SecondSegmentPoint * DisplayScale);
+			if (bDisplayNodes)
+			{
+				DisplayPoint<TPoint>(*SecondSegmentPoint * DisplayScale, EVisuProperty::BluePoint, Index);
+			}
+			FirstSegmentPoint = SecondSegmentPoint;
+		}
+
+		{
+			F3DDebugSession C(bDisplaySegments, FString::Printf(TEXT("Segment %d"), Loop.Num()));
+			DisplaySegment(*FirstSegmentPoint * DisplayScale, *Loop.begin() * DisplayScale);
+		}
+	}
+
+	template<typename TPoint>
+	void DisplayGridLoops(FString Message, const TArray<TArray<TPoint>>& Loops, bool bDisplayNodes, bool bMakeGroup, bool bDisplaySegments) const
+	{
+		if (!bDisplay)
+		{
+			return;
+		}
+
+		F3DDebugSession _(*Message);
+		int32 LoopIndex = 0;
+		for (const TArray<TPoint>& Loop : Loops)
+		{
+			FString LoopName = FString::Printf(TEXT("Loop %d"), LoopIndex++);
+			DisplayGridLoop(LoopName, Loop, bDisplayNodes, bMakeGroup, bDisplaySegments);
+		}
+	}
+
+	void DisplayGridLoops(FString Message, EGridSpace DisplaySpace, bool bDisplayNodes, bool bMakeGroup, bool bDisplaySegments) const
+	{
+		DisplayGridLoops(Message, FaceLoops2D[DisplaySpace], bDisplayNodes, bMakeGroup, bDisplaySegments);
+	}
+
+#endif
+};
+
+}
diff --git a/Engine/Source/Runtime/Datasmith/CADKernel/Base/Private/Topo/Model.cpp b/Engine/Source/Runtime/Datasmith/CADKernel/Base/Private/Topo/Model.cpp
new file mode 100644
--- /dev/null
+++ b/Engine/Source/Runtime/Datasmith/CADKernel/Base/Private/Topo/Model.cpp
@@ -0,0 +1,201 @@
+// Copyright Epic Games, Inc. All Rights Reserved.
+#include "Topo/Model.h"
+
+#include "Core/EntityGeom.h"
+#include "Topo/Body.h"
+#include "Topo/TopologicalEdge.h"
+#include "Topo/TopologicalFace.h"
+#include "Topo/TopologicalVertex.h"
+
+#ifdef CADKERNEL_DEV
+#include "Topo/TopologyReport.h"
+#endif
+
+namespace UE::CADKernel
+{
+
+void FModel::AddEntity(TSharedRef<FTopologicalEntity> Entity)
+{
+	switch (Entity->GetEntityType())
+	{
+	case EEntity::Body:
+		Add(StaticCastSharedRef<FBody>(Entity));
+		break;
+	default:
+		break;
+	}
+}
+
+bool FModel::Contains(TSharedPtr<FTopologicalEntity> Entity)
+{
+	switch(Entity->GetEntityType())
+	{
+	case EEntity::Body:
+		return Bodies.Find(StaticCastSharedPtr<FBody>(Entity)) != INDEX_NONE;
+	default:
+		return false;
+	}
+	return false;
+}
+
+void FModel::PrintBodyAndShellCount()
+{
+	int32 NbBody = 0;
+	int32 NbShell = 0;
+
+	TArray<TSharedPtr<FBody>> NewBodies;
+	for (TSharedPtr<FBody> Body : Bodies)
+	{
+		NbShell += Body->GetShells().Num();
+		NbBody++;
+	}
+	FMessage::Printf(Log, TEXT("Body count %d shell count %d \n"), NbBody, NbShell);
+}
+
+int32 FModel::FaceCount() const
+{
+	int32 FaceCount = 0;
+	for (const TSharedPtr<FBody>& Body : Bodies)
+	{
+		FaceCount += Body->FaceCount();
+	}
+	return FaceCount;
+}
+
+void FModel::RemoveEmptyBodies()
+{
+	TArray<TSharedPtr<FBody>> NewBodies;
+	NewBodies.Reserve(Bodies.Num());
+
+	for (const TSharedPtr<FBody>& Body : Bodies)
+	{
+		if (!Body->IsDeleted() && Body->ShellCount())
+		{
+			NewBodies.Add(Body);
+		}
+		else
+		{
+			Body->Delete();
+			Body->ResetHost();
+		}
+	}
+
+	Bodies = MoveTemp(NewBodies);
+}
+
+void FModel::GetFaces(TArray<FTopologicalFace*>& OutFaces)
+{
+	for (const TSharedPtr<FBody>& Body : Bodies)
+	{
+		Body->GetFaces(OutFaces);
+	}
+}
+
+void FModel::PropagateBodyOrientation()
+{
+	for (TSharedPtr<FBody>& Body : Bodies)
+	{
+		Body->PropagateBodyOrientation();
+	}
+}
+
+void FModel::CompleteMetaData()
+{
+	for (TSharedPtr<FBody>& Body : Bodies)
+	{
+		Body->CompleteMetaData();
+	}
+}
+
+struct FBodyShell
+{
+	TSharedPtr<FBody> Body;
+	TSharedPtr<FShell> Shell;
+
+	FBodyShell(TSharedPtr<FBody> InBody, TSharedPtr<FShell> InShell)
+	: Body(InBody)
+	, Shell(InShell)
+	{
+	}
+
+};
+
+void FModel::CheckTopology() 
+{
+	TArray<FBodyShell> IsolatedBodies;
+	IsolatedBodies.Reserve(Bodies.Num()*2);
+
+	int32 ShellCount = 0;
+
+	for (TSharedPtr<FBody> Body : Bodies)
+	{
+		for (TSharedPtr<FShell> Shell : Body->GetShells())
+		{
+			ShellCount++;
+			TArray<FFaceSubset> SubShells;
+			Shell->CheckTopology(SubShells);
+
+			if (SubShells.Num() == 1 )
+			{
+				if (Shell->FaceCount() < 3 )
+				{
+					IsolatedBodies.Emplace(Body, Shell);
+				}
+				else
+				{
+					if (SubShells[0].BorderEdgeCount > 0 || SubShells[0].NonManifoldEdgeCount > 0)
+					{
+#ifdef CORETECHBRIDGE_DEBUG
+						FMessage::Printf(Log, TEXT("Body %d shell %d CADId %d is opened and has %d faces "), Body->GetKioId(), Shell->GetKioId(), Shell->GetId(), Shell->FaceCount());
+#else
+						FMessage::Printf(Log, TEXT("Body %d shell %d is opened and has %d faces "), Body->GetId(), Shell->GetId(), Shell->FaceCount());
+#endif
+						FMessage::Printf(Log, TEXT("and has %d border edges and %d nonManifold edges\n"), SubShells[0].BorderEdgeCount, SubShells[0].NonManifoldEdgeCount);
+					}
+					else
+					{
+#ifdef CORETECHBRIDGE_DEBUG
+						FMessage::Printf(Log, TEXT("Body %d shell %d CADId %d is closed and has %d faces\n"), Body->GetKioId(), Shell->GetKioId(), Shell->GetId(), Shell->FaceCount());
+#else
+						FMessage::Printf(Log, TEXT("Body %d shell %d is closed and has %d faces\n"), Body->GetId(), Shell->GetId(), Shell->FaceCount());
+#endif
+					}
+				}
+			}
+			else
+			{
+#ifdef CORETECHBRIDGE_DEBUG
+				FMessage::Printf(Log, TEXT("Body %d shell %d CADId %d has %d subshells\n"), Body->GetKioId(), Shell->GetKioId(), Shell->GetId(), SubShells.Num());
+#else
+				FMessage::Printf(Log, TEXT("Body %d shell %d has %d subshells\n"), Body->GetId(), Shell->GetId(), SubShells.Num());
+#endif
+for (const FFaceSubset& Subset : SubShells)
+				{
+					FMessage::Printf(Log, TEXT("     - Subshell of %d faces %d border edges and %d nonManifold edges\n"), Subset. Faces.Num(), Subset.BorderEdgeCount, Subset.NonManifoldEdgeCount);
+				}
+			}
+		}
+	}
+}
+
+#ifdef CADKERNEL_DEV
+void FModel::FillTopologyReport(FTopologyReport& Report) const
+{
+	for (TSharedPtr<FBody> Body : Bodies)
+	{
+		Body->FillTopologyReport(Report);
+	}
+
+	Report.ResetMarkers();
+}
+#endif
+
+void FModel::Orient()
+{
+	for (TSharedPtr<FBody> Body : Bodies)
+	{
+		Body->Orient();
+	}
+}
+
+}
\ No newline at end of file
diff --git a/Engine/Source/Runtime/Datasmith/CADKernel/Base/Private/UI/Display.cpp b/Engine/Source/Runtime/Datasmith/CADKernel/Base/Private/UI/Display.cpp
new file mode 100644
--- /dev/null
+++ b/Engine/Source/Runtime/Datasmith/CADKernel/Base/Private/UI/Display.cpp
@@ -0,0 +1,1337 @@
+// Copyright Epic Games, Inc. All Rights Reserved.
+
+#include "UI/Display.h"
+
+#include "Core/Chrono.h"
+#include "Core/Group.h"
+#include "Core/System.h"
+#include "Core/OrientedEntity.h"
+#include "Geo/Curves/RestrictionCurve.h"
+#include "Geo/Curves/BezierCurve.h"
+#include "Geo/Curves/NURBSCurve.h"
+#include "Geo/Curves/RestrictionCurve.h"
+#include "Geo/Curves/SplineCurve.h"
+#include "Geo/Sampler/SamplerOnChord.h"
+#include "Geo/Sampler/SamplerOnParam.h"
+#include "Geo/Surfaces/BezierSurface.h"
+#include "Geo/Surfaces/NURBSSurface.h"
+#include "Geo/Surfaces/Surface.h"
+#include "Math/Aabb.h"
+#include "Math/Boundary.h"
+#include "Mesh/Structure/EdgeMesh.h"
+#include "Mesh/Structure/ModelMesh.h"
+#include "Mesh/Structure/FaceMesh.h"
+#include "Mesh/Structure/VertexMesh.h"
+#include "Topo/Body.h"
+#include "Topo/TopologicalEdge.h"
+#include "Topo/Model.h"
+#include "Topo/Shell.h"
+#include "Topo/TopologicalFace.h"
+
+#include "Math/Plane.h"
+
+namespace UE::CADKernel
+{
+
+#ifdef CADKERNEL_DEBUG
+void Open3DDebugSession(FString name, const TArray<FIdent>& IdArray)
+{
+	FSystem::Get().GetVisu()->Open3DDebugSession(*name, IdArray);
+}
+
+void Close3DDebugSession(bool bIsDisplayed)
+{
+	if (bIsDisplayed)
+	{
+		FSystem::Get().GetVisu()->Close3DDebugSession();
+	}
+}
+
+void Wait(bool bMakeWait)
+{
+	if (bMakeWait)
+	{
+		system("PAUSE");
+		printf("Hello\n");
+	}
+}
+#endif
+
+void Open3DDebugSegment(FIdent Ident)
+{
+#ifdef CADKERNEL_DEBUG
+	FSystem::Get().GetVisu()->Open3DDebugSegment(Ident);
+#endif
+}
+
+void Close3DDebugSegment()
+{
+#ifdef CADKERNEL_DEBUG
+	FSystem::Get().GetVisu()->Close3DDebugSegment();
+#endif
+}
+
+void FlushVisu()
+{
+#ifdef CADKERNEL_DEBUG
+	FSystem::Get().GetVisu()->UpdateViewer();
+#endif
+}
+
+void DrawElement(int32 Dimension, TArray<FPoint>& Points, EVisuProperty Property)
+{
+#ifdef CADKERNEL_DEBUG
+	FSystem::Get().GetVisu()->DrawElement(Dimension, Points, Property);
+#endif
+}
+
+void DrawMesh(const FMesh& Mesh)
+{
+#ifdef CADKERNEL_DEBUG
+	FSystem::Get().GetVisu()->DrawMesh(Mesh.GetId());
+#endif
+}
+
+void Draw(const FLinearBoundary& Boundary, const FRestrictionCurve& Curve, EVisuProperty Property)
+{
+#ifdef CADKERNEL_DEBUG
+	if (Boundary.IsDegenerated())
+	{
+		return;
+	}
+
+	bool bShowOrientation = FSystem::Get().GetVisu()->GetParameters()->bDisplayCADOrient;
+
+	TArray<FPoint> Polyline;
+	Curve.GetDiscretizationPoints<FPoint>(Boundary, EOrientation::Front, Polyline);
+	Draw(Polyline, Property);
+
+	if (bShowOrientation)
+	{
+		double Length = 0;
+		for (int32 Index = 1; Index < Polyline.Num(); ++Index)
+		{
+			Length += Polyline[Index].Distance(Polyline[Index - 1]);
+		}
+
+		double Coordinate = (Boundary.Max + Boundary.Min) / 2.;
+		FCurvePoint Point;
+		Curve.EvaluatePoint(Coordinate, Point, 1);
+
+		double Height = Length / 20.0;
+		double Base = Height / 2;
+
+		DrawQuadripode(Height, Base, Point.Point, Point.Gradient, Property);
+	}
+#endif
+}
+
+void Draw2D(const FCurve& Curve, const FLinearBoundary& Boundary, EVisuProperty Property)
+{
+#ifdef CADKERNEL_DEBUG
+	if (Boundary.IsDegenerated())
+	{
+		return;
+	}
+
+	double DiscVisu = ((double)(FSystem::Get().GetVisu()->GetParameters()->ChordError)) / 10.;
+	bool bVoirOrientation = FSystem::Get().GetVisu()->GetParameters()->bDisplayCADOrient;
+
+	FPolyline2D Polyline;
+	FCurve2DSamplerOnChord Sampler(Curve, Boundary, DiscVisu, Polyline);
+	Sampler.Sample();
+
+	Draw(Polyline.GetPoints(), Property);
+
+	if (bVoirOrientation)
+	{
+		double Length = Polyline.GetLength(Boundary);
+
+		double Coordinate = (Boundary.Max + Boundary.Min) / 2.;
+		FCurvePoint Point;
+		Curve.EvaluatePoint(Coordinate, Point, 1);
+
+		double Height = Length / 20.0;
+		double Base = Height / 2;
+
+		DrawQuadripode(Height, Base, Point.Point, Point.Gradient, Property);
+	}
+#endif
+}
+
+void Draw3D(const FCurve& Curve, const FLinearBoundary& Boundary, EVisuProperty Property)
+{
+#ifdef CADKERNEL_DEBUG
+	if (Boundary.IsDegenerated())
+	{
+		return;
+	}
+
+	double DiscVisu = FSystem::Get().GetVisu()->GetParameters()->ChordError;
+	bool bVoirOrientation = FSystem::Get().GetVisu()->GetParameters()->bDisplayCADOrient;
+
+	FPolyline3D Polyline;
+	FCurveSamplerOnChord Sampler(Curve, Boundary, DiscVisu, Polyline);
+	Sampler.Sample();
+
+	Draw(Polyline.GetPoints(), Property);
+
+	if (bVoirOrientation)
+	{
+		double Length = Polyline.GetLength(Boundary);
+
+		double Coordinate = (Boundary.Max + Boundary.Min) / 2.;
+		FCurvePoint Point;
+		Curve.EvaluatePoint(Coordinate, Point, 1);
+
+		double Height = Length / 20.0;
+		double Base = Height / 2;
+
+		DrawQuadripode(Height, Base, Point.Point, Point.Gradient, Property);
+	}
+#endif
+}
+
+void Draw(const FCurve& Curve, const FLinearBoundary& Boundary, EVisuProperty Property)
+{
+#ifdef CADKERNEL_DEBUG
+	if (Curve.GetDimension() == 3)
+	{
+		Draw3D(Curve, Boundary, Property);
+	}
+	else
+	{
+		Draw2D(Curve, Boundary, Property);
+	}
+#endif
+}
+
+
+void Draw(const FCurve& Curve, EVisuProperty Property)
+{
+#ifdef CADKERNEL_DEBUG
+	const FLinearBoundary& Bounds = Curve.GetBoundary();
+	Draw(Curve, Bounds, Property);
+#endif
+}
+
+void DrawQuadripode(double Height, double Base, FPoint& Center, FPoint& InDirection, EVisuProperty Property)
+{
+#ifdef CADKERNEL_DEBUG
+	FPoint Direction = InDirection;
+	Direction.Normalize();
+
+	FPoint Normal(-Direction[1], Direction[0], 0.0);
+
+	double UVNorm = Normal.Length();
+	Normal /= UVNorm;
+
+	FPoint BiNormal(-Direction[2] * Normal[1], Direction[2] * Normal[0], UVNorm);
+
+	FPoint Point0 = Center;
+	FPoint PointBase = Point0 - Direction * Height;
+
+	Normal *= Base;
+	BiNormal *= Base;
+
+	FPoint Point1 = PointBase + Normal;
+	FPoint Point2 = PointBase + BiNormal;
+	FPoint Point3 = PointBase - Normal;
+	FPoint Point4 = PointBase - BiNormal;
+
+	TArray<FPoint> Polygone;
+	Polygone.Reserve(7);
+
+	Polygone.Add(Point1);
+	Polygone.Add(Point0);
+	Polygone.Add(Point2);
+	Polygone.Add(Point0);
+	Polygone.Add(Point3);
+	Polygone.Add(Point0);
+	Polygone.Add(Point4);
+
+	Draw(Polygone, Property);
+#endif
+}
+
+void DisplayEntity(const FEntity& Entity)
+{
+#ifdef CADKERNEL_DEBUG
+	if (Entity.IsDeleted())
+	{
+		return;
+	}
+
+	FTimePoint StartTime = FChrono::Now();
+	FProgress Progress;
+
+	F3DDebugSession GraphicSession(FString::Printf(TEXT("%s %d"), Entity.GetTypeName(), Entity.GetId()), { Entity.GetId() });
+
+	switch (Entity.GetEntityType())
+	{
+	case EEntity::TopologicalVertex:
+		Display((const FTopologicalVertex&)Entity);
+		break;
+	case EEntity::Curve:
+		Display((const FCurve&)Entity);
+		break;
+	case EEntity::Surface:
+		Display((const FSurface&)Entity);
+		break;
+	case EEntity::TopologicalFace:
+		Display((const FTopologicalFace&)Entity);
+		break;
+	case EEntity::TopologicalLoop:
+		Display((const FTopologicalLoop&)Entity);
+		break;
+	case EEntity::Shell:
+		Display((const FShell&)Entity);
+		break;
+	case EEntity::Body:
+		Display((const FBody&)Entity);
+		break;
+	case EEntity::TopologicalEdge:
+		Display((const FTopologicalEdge&)Entity);
+		break;
+	case EEntity::Model:
+		Display((const FModel&)Entity);
+		break;
+	case EEntity::MeshModel:
+		Display((const FModelMesh&)Entity);
+		break;
+	default:
+		FMessage::Printf(Log, TEXT("Unable to display Entity of type %s"), FEntity::GetTypeName(Entity.GetEntityType()));
+	}
+
+	FDuration DisplayDuration = FChrono::Elapse(StartTime);
+	FChrono::PrintClockElapse(Log, TEXT("  "), TEXT("Display "), DisplayDuration);
+
+#endif
+}
+
+void DisplayEntity2D(const FEntity& Entity)
+{
+#ifdef CADKERNEL_DEBUG
+	if (Entity.IsDeleted())
+	{
+		return;
+	}
+
+	FTimePoint StartTime = FChrono::Now();
+
+	FProgress Progress;
+
+	F3DDebugSession GraphicSession(FString::Printf(TEXT("%s %d"), Entity.GetTypeName(), Entity.GetId()), { Entity.GetId() });
+
+	switch (Entity.GetEntityType())
+	{
+	case EEntity::Surface:
+		Display2D((const FSurface&)Entity);
+		break;
+	case EEntity::TopologicalFace:
+		Display2D((const FTopologicalFace&)Entity);
+		break;
+	case EEntity::TopologicalLoop:
+		Display2D((const FTopologicalLoop&)Entity);
+		break;
+	case EEntity::TopologicalEdge:
+		Display2D((const FTopologicalEdge&)Entity);
+		break;
+	default:
+		FMessage::Printf(Log, TEXT("Unable to display Entity of type %s"), FEntity::GetTypeName(Entity.GetEntityType()));
+	}
+
+	FDuration DisplayDuration = FChrono::Elapse(StartTime);
+	FChrono::PrintClockElapse(Log, TEXT("  "), TEXT("Display "), DisplayDuration);
+
+#endif
+}
+
+void Display(const FPlane& Plane, FIdent Ident)
+{
+#ifdef CADKERNEL_DEBUG
+	FVector Normal = Plane.GetNormal();
+	if (Normal.Size() < DOUBLE_SMALL_NUMBER)
+	{
+		return;
+	}
+
+	FVector UAxis;
+	// Find the first frame axis non collinear to the plane normal
+	for (int32 Index = 0; Index < 3; ++Index)
+	{
+		FVector Axis = FVector::ZeroVector;
+		Axis[Index] = 1;
+		UAxis = Normal ^ Axis;
+		if (UAxis.Size() > DOUBLE_SMALL_NUMBER)
+		{
+			break;
+		}
+	}
+	FVector VAxis = UAxis ^ Normal;
+
+	UAxis.Normalize();
+	UAxis *= 10;
+	VAxis.Normalize();
+	VAxis *= 10;
+
+	TArray<FPoint> Points;
+
+	const FVector Point = Plane.GetOrigin();
+	
+	FVector VTemp;
+
+	VTemp = Point + UAxis + VAxis;
+	Points.Emplace(VTemp[0], VTemp[1], VTemp[2]);
+
+	VTemp = Point + UAxis - VAxis;
+	Points.Emplace(VTemp[0], VTemp[1], VTemp[2]);
+
+	VTemp = Point - UAxis - VAxis;
+	Points.Emplace(VTemp[0], VTemp[1], VTemp[2]);
+
+	VTemp = Point - UAxis + VAxis;
+	Points.Emplace(VTemp[0], VTemp[1], VTemp[2]);
+
+	VTemp = Point + UAxis + VAxis;
+	Points.Emplace(VTemp[0], VTemp[1], VTemp[2]);
+
+	F3DDebugSegment G(Ident);
+	DrawElement(2, Points);
+#endif
+}
+
+void DisplayEdgeCriteriaGrid(int32 EdgeId, const TArray<FPoint>& Points3D)
+{
+#ifdef CADKERNEL_DEBUG
+	FString Name = FString::Printf(TEXT("Edge Grid %d"), EdgeId);
+	F3DDebugSession G(Name);
+	{
+		F3DDebugSession _(TEXT("CriteriaGrid Point 3d"));
+		for (int32 Index = 0; Index < Points3D.Num(); Index += 2)
+		{
+			DisplayPoint(Points3D[Index]);
+		}
+	}
+	{
+		F3DDebugSession _(TEXT("CriteriaGrid IntermediateU"));
+		for (int32 Index = 1; Index < Points3D.Num(); Index += 2)
+		{
+			DisplayPoint(Points3D[Index], EVisuProperty::ControlPoint);
+		}
+	}
+#endif
+}
+
+void DisplayAABB(const FAABB& Aabb, FIdent Ident)
+{
+#ifdef CADKERNEL_DEBUG
+	TFunction<void(int32, FPoint&, FPoint&)> GetWire = [&](int32 Segment, FPoint& A, FPoint& B)
+	{
+		int32 Corner1;
+		int32 Corner2;
+
+		if (Segment < 4)
+		{
+			Corner1 = Segment * 2;
+			Corner2 = Corner1 + 1;
+		}
+		else if (Segment < 8)
+		{
+			if (Segment < 6)
+			{
+				Corner1 = Segment - 4;
+			}
+			else
+			{
+				Corner1 = Segment - 2;
+			}
+			Corner2 = Corner1 + 2;
+		}
+		else if (Segment < 12)
+		{
+			Corner1 = Segment - 8;
+			Corner2 = Corner1 + 4;
+		}
+		else
+		{
+			Corner1 = 0;
+			Corner2 = 1;
+			ERROR_NOT_EXPECTED;
+		}
+		A = Aabb.GetCorner(Corner1);
+		B = Aabb.GetCorner(Corner2);
+	};
+
+	F3DDebugSegment G(Ident);
+	for (int32 SegmentIndex = 0; SegmentIndex < 12; SegmentIndex++)
+	{
+		TArray<FPoint> Points;
+		Points.SetNum(2);
+		GetWire(SegmentIndex, Points[0], Points[1]);
+		Draw(Points);
+	}
+#endif
+}
+
+void DisplayAABB2D(const FAABB2D& aabb, FIdent Ident)
+{
+#ifdef CADKERNEL_DEBUG
+	FPoint A(aabb.GetMin().U, aabb.GetMin().V, 0.);
+	FPoint B(aabb.GetMin().U, aabb.GetMax().V, 0.);
+	FPoint C(aabb.GetMax().U, aabb.GetMax().V, 0.);
+	FPoint D(aabb.GetMax().U, aabb.GetMin().V, 0.);
+
+	TArray<FPoint> Points;
+	Points.Add(A);
+	Points.Add(B);
+	Points.Add(C);
+	Points.Add(D);
+	Points.Add(A);
+	Draw(Points);
+#endif
+}
+
+void Display(const FTopologicalVertex& Vertex, EVisuProperty Property)
+{
+#ifdef CADKERNEL_DEBUG
+	F3DDebugSegment G(Vertex.GetId());
+	DrawPoint(Vertex.GetCoordinates(), Property);
+#endif
+}
+
+void Display(const FCurve& Curve)
+{
+#ifdef CADKERNEL_DEBUG
+	F3DDebugSegment G(Curve.GetId());
+	Draw(Curve, Curve.GetBoundary());
+#endif
+}
+
+void Display(const FShell& Shell)
+{
+#ifdef CADKERNEL_DEBUG
+	FProgress Progress(Shell.GetFaces().Num(), TEXT("Display Shell"));
+	for (const FOrientedFace& Face : Shell.GetFaces())
+	{
+		if (Face.Entity->IsDeleted())
+		{
+			continue;
+		}
+		Display(*Face.Entity);
+	}
+#endif
+}
+
+void Draw(const FShell& Shell)
+{
+#ifdef CADKERNEL_DEBUG
+	FProgress Progress(Shell.GetFaces().Num(), TEXT("Display Shell"));
+	for (const FOrientedFace& Face : Shell.GetFaces())
+	{
+		if (Face.Entity->IsDeleted())
+		{
+			continue;
+		}
+		F3DDebugSegment _(Face.Entity->GetId());
+		Draw(*Face.Entity);
+	}
+#endif
+}
+
+void Display(const FBody& Body)
+{
+#ifdef CADKERNEL_DEBUG
+	FProgress Progress(Body.GetShells().Num(), TEXT("Display Body"));
+
+	F3DDebugSegment GraphicSegment(Body.GetId());
+	for (const TSharedPtr<FShell>& Shell : Body.GetShells())
+	{
+		if (!Shell.IsValid() || Shell->IsDeleted())
+		{
+			continue;
+		}
+		Draw(*Shell);
+	}
+#endif
+}
+
+void Display(const FSurface& Surface)
+{
+#ifdef CADKERNEL_DEBUG
+	FProgress Progress(TEXT("Display Surface"));
+
+	F3DDebugSegment GraphicSegment(Surface.GetId());
+
+	int32 IsoUNum = FSystem::Get().GetVisu()->GetParameters()->IsoUNumber;
+	int32 IsoVNum = FSystem::Get().GetVisu()->GetParameters()->IsoVNumber;
+
+	const double VisuSag = FSystem::Get().GetVisu()->GetParameters()->ChordError;
+
+	FPolyline3D Polyline;
+	FIsoCurve3DSamplerOnChord Sampler(Surface, VisuSag, Polyline);
+
+	TFunction <void(const int32, const EIso)> DrawIsos = [&](int32 IsoCount, const EIso IsoType)
+	{
+		const FLinearBoundary& CurveBounds = Surface.GetBoundary().Get(IsoType == EIso::IsoU ? EIso::IsoV : EIso::IsoU);
+		const FLinearBoundary& Bounds = Surface.GetBoundary().Get(IsoType);
+
+		double Coordinate = Bounds.Min;
+		IsoCount++;
+		const double Step = (Bounds.Max - Bounds.Min) / IsoCount;
+
+		for (int32 Index = 0; Index <= IsoCount; Index++)
+		{
+			Polyline.Empty();
+			Sampler.Set(IsoType, Coordinate, CurveBounds);
+			Sampler.Sample();
+
+			Draw(Polyline.GetPoints(), EVisuProperty::Iso);
+			Coordinate += Step;
+		}
+	};
+
+	DrawIsos(IsoVNum, EIso::IsoV);
+	DrawIsos(IsoUNum, EIso::IsoU);
+#endif
+}
+
+void Display2D(const FSurface& Surface)
+{
+#ifdef CADKERNEL_DEBUG
+	double StepU, StepV;
+
+	F3DDebugSegment G(Surface.GetId());
+
+	int32 IsoUCount = FSystem::Get().GetVisu()->GetParameters()->IsoUNumber;
+	int32 IsoVCount = FSystem::Get().GetVisu()->GetParameters()->IsoVNumber;
+	FSurfacicBoundary Bounds = Surface.GetBoundary();
+
+	StepU = (Bounds[EIso::IsoU].Max - Bounds[EIso::IsoU].Min) / (IsoUCount + 1);
+	StepV = (Bounds[EIso::IsoV].Max - Bounds[EIso::IsoV].Min) / (IsoVCount + 1);
+
+	for (int32 iIso = 0; iIso <= (IsoUCount + 1); iIso++)
+	{
+		FPoint2D StartPoint(Bounds[EIso::IsoU].Min + (iIso)*StepU, Bounds[EIso::IsoV].Min);
+		FPoint2D EndPoint(Bounds[EIso::IsoU].Min + (iIso)*StepU, Bounds[EIso::IsoV].Max);
+		DrawSegment(StartPoint, EndPoint, EVisuProperty::Iso);
+	}
+
+	for (int32 iIso = 0; iIso <= (IsoVCount + 1); iIso++)
+	{
+		FPoint2D StartPoint(Bounds[EIso::IsoU].Min, Bounds[EIso::IsoV].Min + (iIso)*StepV);
+		FPoint2D EndPoint(Bounds[EIso::IsoU].Max, Bounds[EIso::IsoV].Min + (iIso)*StepV);
+		DrawSegment(StartPoint, EndPoint, EVisuProperty::Iso);
+	}
+#endif
+}
+
+void DisplayIsoCurve(const FSurface& Surface, double Coordinate, EIso IsoType)
+{
+#ifdef CADKERNEL_DEBUG
+	const double VisuSag = FSystem::Get().GetVisu()->GetParameters()->ChordError;
+
+	F3DDebugSegment G(Surface.GetId());
+
+	const FLinearBoundary& CurveBounds = Surface.GetBoundary().Get(IsoType == EIso::IsoU ? EIso::IsoV : EIso::IsoU);
+
+	FPolyline3D Polyline;
+	FIsoCurve3DSamplerOnChord Sampler(Surface, VisuSag, Polyline);
+	Sampler.Set(IsoType, Coordinate, CurveBounds);
+	Draw(Polyline.GetPoints(), EVisuProperty::Iso);
+#endif
+}
+
+
+
+void DisplayControlPolygon(const FCurve& Curve)
+{
+#ifdef CADKERNEL_DEBUG
+	TFunction<void(const TArray<FPoint>&)> DisplayHull = [](const TArray<FPoint>& Poles)
+	{
+		for (const FPoint& Pole : Poles)
+		{
+			DisplayPoint(Pole, EVisuProperty::GreenPoint);
+		}
+
+		for (int32 Index = 1; Index < Poles.Num(); Index++)
+		{
+			FPoint Segment = Poles[Index] - Poles[Index - 1];
+			DisplaySegment(Poles[Index - 1] + Segment * 0.1, Poles[Index] - Segment * 0.1, 0, EVisuProperty::GreenCurve);
+		}
+	};
+
+	F3DDebugSegment GraphicSegment(Curve.GetId());
+	switch (Curve.GetCurveType())
+	{
+	case  ECurve::Bezier:
+	{
+		const FBezierCurve& Bezier = (const FBezierCurve&)Curve;
+		DisplayHull(Bezier.GetPoles());
+		return;
+	}
+	case ECurve::Nurbs:
+	{
+		const FNURBSCurve& Nurbs = (const FNURBSCurve&)Curve;
+		const TArray<FPoint>& Poles = Nurbs.GetPoles();
+		DisplayHull(Poles);
+		return;
+	}
+	case ECurve::Spline:
+	{
+		const FSplineCurve& Spline = (const FSplineCurve&)Curve;
+		const FInterpCurveFPoint& Poles = Spline.GetSplinePointsPosition();
+
+		for (const auto& Pole : Poles.Points)
+		{
+			DisplayPoint(Pole.OutVal, EVisuProperty::BluePoint);
+		}
+
+		for (const auto& Pole : Poles.Points)
+		{
+			DisplaySegment(Pole.OutVal - Pole.ArriveTangent / 2., Pole.OutVal, 0, EVisuProperty::GreenCurve);
+			DisplaySegment(Pole.OutVal, Pole.OutVal + Pole.LeaveTangent / 2., 0, EVisuProperty::GreenCurve);
+		}
+		return;
+	}
+
+	}
+
+#endif
+}
+
+void DisplayControlPolygon(const FSurface& Surface)
+{
+#ifdef CADKERNEL_DEBUG
+	bool bShowOrientation = FSystem::Get().GetVisu()->GetParameters()->bDisplayCADOrient;
+
+	TFunction<void(const TArray<FPoint>&, int32, int32)> DisplayHull = [&bShowOrientation](const TArray<FPoint>& Poles, int32 PoleUNum, int32 PoleVNum)
+	{
+		for (int32 Index = 0; Index < Poles.Num(); Index++)
+		{
+			DisplayPoint(Poles[Index], EVisuProperty::GreenPoint);
+		}
+
+		for (int32 IndexV = 0, Index = 0; IndexV < PoleVNum; IndexV++)
+		{
+			Index = Index + 1;
+			for (int32 IndexU = 1; IndexU < PoleUNum; IndexU++, Index++)
+			{
+				FPoint Segment = Poles[Index] - Poles[Index - 1];
+				DisplaySegment(Poles[Index - 1] + Segment * 0.1, Poles[Index] - Segment * 0.1, 0, EVisuProperty::YellowCurve);
+				if (bShowOrientation)
+				{
+					DrawSegmentOrientation(Poles[Index - 1], Poles[Index], EVisuProperty::YellowCurve);
+				}
+			}
+		}
+
+		for (int32 IndexU = 0; IndexU < PoleUNum; IndexU++)
+		{
+			for (int32 IndexV = 1, Index = PoleUNum + IndexU; IndexV < PoleVNum; IndexV++, Index += PoleUNum)
+			{
+				FPoint Segment = Poles[Index] - Poles[Index - PoleUNum];
+				DisplaySegment(Poles[Index - PoleUNum] + Segment * 0.1, Poles[Index] - Segment * 0.1, 0, EVisuProperty::GreenCurve);
+				if (bShowOrientation)
+				{
+					DrawSegmentOrientation(Poles[Index - PoleUNum], Poles[Index], EVisuProperty::GreenCurve);
+				}
+			}
+		}
+	};
+
+	F3DDebugSegment GraphicSegment(Surface.GetId());
+	if (Surface.GetSurfaceType() == ESurface::Bezier)
+	{
+		const FBezierSurface& Bezier = (const FBezierSurface&)Surface;
+		DisplayHull(Bezier.GetPoles(), Bezier.GetUDegre() + 1, Bezier.GetVDegre() + 1);
+		return;
+	}
+	if (Surface.GetSurfaceType() == ESurface::Nurbs)
+	{
+		const FNURBSSurface& Nurbs = (const FNURBSSurface&)Surface;
+		DisplayHull(Nurbs.GetPoles(), Nurbs.GetPoleCount(EIso::IsoU), Nurbs.GetPoleCount(EIso::IsoV));
+		return;
+	}
+#endif
+}
+
+void Display(const FTopologicalFace& Face)
+{
+#ifdef CADKERNEL_DEBUG
+	F3DDebugSegment GraphicSegment(Face.GetId());
+	Draw(Face);
+#endif
+}
+
+void Display2D(const FTopologicalFace& Face)
+{
+#ifdef CADKERNEL_DEBUG
+	F3DDebugSegment GraphicSegment(Face.GetId());
+	Draw2D(Face);
+#endif
+}
+
+void Display2DWithScale(const FTopologicalFace& Face)
+{
+#ifdef CADKERNEL_DEBUG
+	F3DDebugSegment GraphicSegment(Face.GetId());
+	Draw2D(Face);
+#endif
+}
+
+void Draw(const FTopologicalFace& Face)
+{
+#ifdef CADKERNEL_DEBUG
+	{
+		for (const TSharedPtr<FTopologicalLoop>& Loop : Face.GetLoops())
+		{
+			for (const FOrientedEdge& Edge : Loop->GetEdges())
+			{
+				EVisuProperty Property = EVisuProperty::BlueCurve;
+				switch (Edge.Entity->GetTwinEntityCount())
+				{
+				case 1:
+					Property =  Edge.Entity->IsDegenerated() ? EVisuProperty::OrangeCurve : EVisuProperty::BorderEdge;
+					break;
+				case 2:
+					Property = EVisuProperty::BlueCurve;
+					break;
+				default:
+					Property = EVisuProperty::NonManifoldEdge;
+				}
+				Draw(*Edge.Entity, Property);
+			}
+		}
+	}
+
+	DrawIsoCurves(Face);
+#endif
+}
+
+void Draw2D(const FTopologicalFace& Face)
+{
+#ifdef CADKERNEL_DEBUG
+	TArray<TArray<FPoint2D>> BoundaryApproximation;
+	Face.Get2DLoopSampling(BoundaryApproximation);
+
+	for (const TArray<FPoint2D>& Boudary : BoundaryApproximation)
+	{
+		Draw(Boudary, EVisuProperty::BlueCurve);
+	}
+
+	TFunction <void(const int32, const EIso)> DrawIsos = [&](int32 IsoCount, const EIso IsoType)
+	{
+		EIso OtherIso = IsoType == EIso::IsoU ? EIso::IsoV : EIso::IsoU;
+		const FLinearBoundary& Bounds = Face.GetBoundary().Get(IsoType);
+		const FLinearBoundary& CurveBounds = Face.GetBoundary().Get(OtherIso);
+
+
+		double Coordinate = Bounds.Min;
+		IsoCount++;
+		const double Step = (Bounds.Max - Bounds.Min) / IsoCount;
+
+		for (int32 iIso = 1; iIso < IsoCount; iIso++)
+		{
+			Coordinate += Step;
+
+			TArray<double> Intersections;
+			FindLoopIntersectionsWithIso(IsoType, Coordinate, BoundaryApproximation, Intersections);
+
+			FPoint2D Start;
+			FPoint2D End;
+
+			Start[IsoType] = Coordinate;
+			End[IsoType] = Coordinate;
+
+			if (Intersections.Num() % 2 != 0)
+			{
+				Start[OtherIso] = Intersections[Intersections.Num() - 1];
+				End[OtherIso] = CurveBounds.GetMax();
+				DrawSegment(Start, End, EVisuProperty::YellowCurve);
+				Intersections.Pop();
+			}
+
+			for (int32 ISection = 0; ISection < Intersections.Num(); ISection += 2)
+			{
+				Start[OtherIso] = Intersections[ISection];
+				End[OtherIso] = Intersections[ISection + 1];
+				DrawSegment(Start, End, EVisuProperty::Iso);
+			}
+		}
+	};
+
+	int32 IsoUCount = FSystem::Get().GetVisu()->GetParameters()->IsoUNumber;
+	int32 IsoVCount = FSystem::Get().GetVisu()->GetParameters()->IsoVNumber;
+
+	DrawIsos(IsoUCount, EIso::IsoU);
+	DrawIsos(IsoVCount, EIso::IsoV);
+#endif
+}
+
+void DrawIsoCurves(const FTopologicalFace& Face)
+{
+#ifdef CADKERNEL_DEBUG
+	TArray<TArray<FPoint2D>> BoundaryApproximation;
+	Face.Get2DLoopSampling(BoundaryApproximation);
+
+	const double VisuSag = FSystem::Get().GetVisu()->GetParameters()->ChordError;
+
+	FPolyline3D Polyline;
+	const FSurface& Surface = Face.GetCarrierSurface().Get();
+	FIsoCurve3DSamplerOnChord Sampler(Surface, VisuSag, Polyline);
+
+	TFunction <void(const int32, const EIso)> DrawIsos = [&](int32 IsoCount, const EIso IsoType)
+	{
+		const FLinearBoundary& Bounds = Face.GetBoundary().Get(IsoType);
+
+		EIso Other = IsoType == EIso::IsoU ? EIso::IsoV : EIso::IsoU;
+
+		double Coordinate = Bounds.Min;
+		IsoCount++;
+		const double Step = (Bounds.Max - Bounds.Min) / IsoCount;
+
+		for (int32 iIso = 1; iIso < IsoCount; iIso++)
+		{
+			Coordinate += Step;
+
+			TArray<double> Intersections;
+			FindLoopIntersectionsWithIso(IsoType, Coordinate, BoundaryApproximation, Intersections);
+			int32 IntersectionCount = Intersections.Num();
+			if (IntersectionCount == 0)
+			{
+				return;
+			}
+
+			FLinearBoundary CurveBounds(Intersections[0], Intersections.Last());
+
+			Polyline.Empty();
+			Sampler.Set(IsoType, Coordinate, CurveBounds);
+			Sampler.Sample();
+
+			if (IntersectionCount % 2 != 0)
+			{
+				TArray<FPoint> SubPolyline;
+				FLinearBoundary Boundary(Intersections[IntersectionCount - 1], CurveBounds.GetMax());
+				Polyline.GetSubPolyline(Boundary, EOrientation::Front, SubPolyline);
+				Draw(SubPolyline, EVisuProperty::YellowCurve);
+				Intersections.Pop();
+				IntersectionCount--;
+			}
+
+			for (int32 ISection = 0; ISection < IntersectionCount; ISection += 2)
+			{
+				TArray<FPoint> SubPolyline;
+				FLinearBoundary Boundary(Intersections[ISection], Intersections[ISection + 1]);
+				Polyline.GetSubPolyline(Boundary, EOrientation::Front, SubPolyline);
+				Draw(SubPolyline, EVisuProperty::Iso);
+			}
+		}
+	};
+
+	int32 IsoUCount = FSystem::Get().GetVisu()->GetParameters()->IsoUNumber;
+	int32 IsoVCount = FSystem::Get().GetVisu()->GetParameters()->IsoVNumber;
+
+	DrawIsos(IsoUCount, EIso::IsoU);
+	DrawIsos(IsoVCount, EIso::IsoV);
+#endif
+}
+
+void Display(const FTopologicalEdge& Edge, EVisuProperty Property)
+{
+#ifdef CADKERNEL_DEBUG
+	F3DDebugSegment GraphicSegment(Edge.GetId());
+	Draw(Edge, Property);
+#endif
+}
+
+void Display2D(const FTopologicalEdge& Edge, EVisuProperty Property)
+{
+#ifdef CADKERNEL_DEBUG
+	F3DDebugSegment GraphicSegment(Edge.GetId());
+	TArray<FPoint2D> Polyline;
+	Edge.GetCurve()->GetDiscretizationPoints(Edge.GetBoundary(), EOrientation::Front, Polyline);
+	DisplayPolyline(Polyline, Property);
+#endif
+}
+
+void Display2DWithScale(const FTopologicalEdge& Edge, EVisuProperty Property)
+{
+#ifdef CADKERNEL_DEBUG
+	F3DDebugSegment GraphicSegment(Edge.GetId());
+	TArray<FPoint2D> Polyline;
+	Edge.GetCurve()->GetDiscretizationPoints(Edge.GetBoundary(), EOrientation::Front, Polyline);
+	DisplayPolylineWithScale(Polyline, Property);
+#endif
+}
+
+void Draw(const FTopologicalEdge& Edge, EVisuProperty Property)
+{
+#ifdef CADKERNEL_DEBUG
+	const FLinearBoundary& Boundary = Edge.GetBoundary();
+	Draw(Boundary, *Edge.GetCurve(), Property);
+#endif
+}
+
+void Display(const FModel& Model)
+{
+#ifdef CADKERNEL_DEBUG
+	FProgress MainProgress(2, TEXT("Display model"));
+	{
+		const TArray<TSharedPtr<FBody>>& Bodies = Model.GetBodies();
+		FProgress BodyProgress((int32)Bodies.Num(), TEXT("Bodies"));
+		for (const TSharedPtr<FBody>& Body : Bodies)
+		{
+			Display(*Body);
+		}
+	}
+
+#endif
+}
+
+void DisplayProductTree(const FModel& Model)
+{
+#ifdef CADKERNEL_DEBUG
+	F3DDebugSession GraphicSession(FString::Printf(TEXT("%s %d"), Model.GetTypeName(), Model.GetId()), { Model.GetId() });
+
+	FProgress MainProgress(2, TEXT("Display Model"));
+	{
+		const TArray<TSharedPtr<FBody>>& Bodies = Model.GetBodies();
+		FProgress BodyProgress((int32)Bodies.Num(), TEXT("Bodies"));
+		for (const TSharedPtr<FBody>& Body : Bodies)
+		{
+			DisplayProductTree(*Body);
+		}
+	}
+
+#endif
+}
+
+void DisplayProductTree(const FBody& Body)
+{
+#ifdef CADKERNEL_DEBUG
+#ifdef CORETECHBRIDGE_DEBUG
+	F3DDebugSession GraphicSession(FString::Printf(TEXT("%s %s KioId: %d Id: %d"), Body.GetTypeName(), Body.GetName(), Body.GetKioId(), Body.GetId()), { Body.GetId() });
+#else
+	F3DDebugSession GraphicSession(FString::Printf(TEXT("%s %s Id: %d"), Body.GetTypeName(), Body.GetName(), Body.GetId()), { Body.GetId() });
+#endif
+	FProgress Progress(Body.GetShells().Num(), TEXT("Display Body"));
+
+	for (TSharedPtr<FShell> Shell : Body.GetShells())
+	{
+		if (!Shell.IsValid())
+		{
+			continue;
+		}
+		DisplayProductTree(*Shell);
+	}
+#endif
+}
+
+void DisplayProductTree(const FShell& Shell)
+{
+#ifdef CADKERNEL_DEBUG
+#ifdef CORETECHBRIDGE_DEBUG
+	F3DDebugSession GraphicSession(FString::Printf(TEXT("%s %s KioId: %d Id: %d"), Shell.GetTypeName(), Shell.GetName(), Shell.GetKioId(), Shell.GetId()), { Shell.GetId() });
+#else
+	F3DDebugSession GraphicSession(FString::Printf(TEXT("%s %s Id: %d"), Shell.GetTypeName(), Shell.GetName(), Shell.GetId()), { Shell.GetId() });
+#endif
+	Draw(Shell);
+#endif
+}
+
+void DisplayProductTree(const FEntity& Entity)
+{
+#ifdef CADKERNEL_DEBUG
+	FTimePoint StartTime = FChrono::Now();
+
+	switch (Entity.GetEntityType())
+	{
+	case EEntity::TopologicalFace:
+		Display((const FTopologicalFace&)Entity);
+		break;
+	case EEntity::Shell:
+		DisplayProductTree((const FShell&)Entity);
+		break;
+	case EEntity::Body:
+		DisplayProductTree((const FBody&)Entity);
+		break;
+	case EEntity::Model:
+		DisplayProductTree((const FModel&)Entity);
+		break;
+	default:
+		FMessage::Printf(Log, TEXT("Unable to display Entity of type %s"), FEntity::GetTypeName(Entity.GetEntityType()));
+	}
+
+	FDuration DisplayDuration = FChrono::Elapse(StartTime);
+	FChrono::PrintClockElapse(Log, TEXT("  "), TEXT("Display "), DisplayDuration);
+#endif
+}
+
+void Display(const FGroup& Group)
+{
+#ifdef CADKERNEL_DEBUG
+	TArray<TSharedPtr<FEntity>> Entities;
+
+	Group.GetValidEntities(Entities);
+
+	FProgress Progress(Entities.Num());
+	for (const TSharedPtr<FEntity>& Entity : Entities)
+	{
+		DisplayEntity(*Entity);
+	}
+#endif
+}
+
+void DisplayMesh(const FModel& Model)
+{
+	for (const TSharedPtr<FBody>& Body : Model.GetBodies())
+	{
+		if(Body)
+		{
+			DisplayMesh(*Body);
+		}
+	}
+}
+
+void DisplayMesh(const FBody& Body)
+{
+	for (const TSharedPtr<FShell>& Shell : Body.GetShells())
+	{
+		if (Shell)
+		{
+			DisplayMesh(*Shell);
+		}
+	}
+}
+
+void DisplayMesh(const FShell& Shell)
+{
+	for (const FOrientedFace& OrientedFace : Shell.GetFaces())
+	{
+		if (OrientedFace.Entity)
+		{
+			FTopologicalFace& Face = *OrientedFace.Entity;
+			if (!Face.IsDeletedOrDegenerated() && Face.IsMeshed())
+			{
+				const FFaceMesh* Mesh = Face.GetMesh();
+				if(Mesh)
+				{
+					DisplayMesh(*Mesh);
+				}
+			}
+		}
+	}
+}
+
+void DisplayMesh(const FFaceMesh& Mesh)
+{
+#ifdef CADKERNEL_DEBUG
+	TMap<int32, const FPoint*> NodeIdToCoordinates;
+	Mesh.GetNodeIdToCoordinates(NodeIdToCoordinates);
+
+	TArray<FPoint> Points;
+	Points.SetNum(3);
+	double EdgeMeanLength = 0;
+	const TArray<int32>& TriangleIndices = Mesh.TrianglesVerticesIndex;
+	const TArray<int32>& VertexIndices = Mesh.VerticesGlobalIndex;
+	for (int32 Index = 0; Index < TriangleIndices.Num();)
+	{
+		const FPoint** P0 = NodeIdToCoordinates.Find(VertexIndices[TriangleIndices[Index++]]);
+		const FPoint** P1 = NodeIdToCoordinates.Find(VertexIndices[TriangleIndices[Index++]]);
+		const FPoint** P2 = NodeIdToCoordinates.Find(VertexIndices[TriangleIndices[Index++]]);
+
+		if (!P0 || !P1 || !P2)
+		{
+			continue;
+		}
+
+		Points[0] = **P0;
+		Points[1] = **P1;
+		Points[2] = **P2;
+		DrawElement(2, Points, EVisuProperty::Element);
+		DrawSegment(Points[0], Points[1], EVisuProperty::EdgeMesh);
+		DrawSegment(Points[1], Points[2], EVisuProperty::EdgeMesh);
+		DrawSegment(Points[2], Points[0], EVisuProperty::EdgeMesh);
+		EdgeMeanLength += Points[0].Distance(Points[1]);
+		EdgeMeanLength += Points[1].Distance(Points[2]);
+		EdgeMeanLength += Points[2].Distance(Points[0]);
+	}
+	EdgeMeanLength /= 10 * TriangleIndices.Num();
+
+	for (const int32& Index : VertexIndices)
+	{
+		const FPoint** PointPtr = NodeIdToCoordinates.Find(Index);
+		if (!PointPtr)
+		{
+			continue;
+		}
+		F3DDebugSegment GraphicSegment(Index);
+		DrawPoint(**PointPtr, EVisuProperty::NodeMesh);
+	}
+
+	bool test = FSystem::Get().GetVisu()->GetParameters()->bDisplayNormals;
+	if (FSystem::Get().GetVisu()->GetParameters()->bDisplayNormals)
+	{
+		double NormalLength = FSystem::Get().GetVisu()->GetParameters()->NormalLength;
+		const TArray<FVector3f>& Normals = Mesh.Normals;
+		for (int32 Index = 0; Index < VertexIndices.Num(); ++Index)
+		{
+			const FPoint** PointPtr = NodeIdToCoordinates.Find(VertexIndices[Index]);
+			if (!PointPtr)
+			{
+				continue;
+			}
+
+			F3DDebugSegment GraphicSegment(Index);
+			FVector3f Normal = Normals[Index];
+			Normal.Normalize();
+			Normal *= NormalLength;
+			DrawSegment(**PointPtr, **PointPtr + Normal, EVisuProperty::EdgeMesh);
+		}
+	}
+#endif
+}
+
+void DisplayMesh(const FEdgeMesh& Mesh)
+{
+#ifdef CADKERNEL_DEBUG
+	const FModelMesh& MeshModel = Mesh.GetMeshModel();
+
+	const TArray<int32>& NodeIds = Mesh.EdgeVerticesIndex;
+	const TArray<FPoint>& NodeCoordinates = Mesh.GetNodeCoordinates();
+
+	int32 StartNodeId = NodeIds[0];
+	int32 LastNodeId = NodeIds.Last();
+
+	{
+		F3DDebugSegment GraphicSegment(Mesh.GetGeometricEntity().GetId());
+		if (NodeCoordinates.Num() == 0)
+		{
+			DrawSegment(MeshModel.GetMeshOfVertexNodeId(StartNodeId)->GetNodeCoordinates()[0], MeshModel.GetMeshOfVertexNodeId(LastNodeId)->GetNodeCoordinates()[0], EVisuProperty::EdgeMesh);
+		}
+		else
+		{
+			DrawSegment(MeshModel.GetMeshOfVertexNodeId(StartNodeId)->GetNodeCoordinates()[0], NodeCoordinates[0], EVisuProperty::EdgeMesh);
+			if (NodeCoordinates.Num() > 1)
+			{
+				for (int32 Index = 0; Index < NodeCoordinates.Num() - 1; Index++)
+				{
+					DrawSegment(NodeCoordinates[Index], NodeCoordinates[Index + 1], EVisuProperty::EdgeMesh);
+				}
+			}
+			DrawSegment(NodeCoordinates.Last(), MeshModel.GetMeshOfVertexNodeId(LastNodeId)->GetNodeCoordinates()[0], EVisuProperty::EdgeMesh);
+		}
+	}
+
+	{
+		F3DDebugSegment GraphicSegment(StartNodeId);
+		DrawPoint(MeshModel.GetMeshOfVertexNodeId(StartNodeId)->GetNodeCoordinates()[0], EVisuProperty::NodeMesh);
+	}
+
+	{
+		F3DDebugSegment GraphicSegment(LastNodeId);
+		DrawPoint(MeshModel.GetMeshOfVertexNodeId(LastNodeId)->GetNodeCoordinates()[0], EVisuProperty::NodeMesh);
+	}
+
+	if (NodeCoordinates.Num() > 1)
+	{
+		for (int32 Index = 0; Index < NodeCoordinates.Num(); Index++)
+		{
+			F3DDebugSegment GraphicSegment(NodeIds[Index]);
+			DrawPoint(NodeCoordinates[Index], EVisuProperty::NodeMesh);
+		}
+	}
+#endif
+}
+
+void DisplayMesh(const FVertexMesh& Mesh)
+{
+#ifdef CADKERNEL_DEBUG
+	DisplayPoint(Mesh.GetNodeCoordinates()[0], EVisuProperty::NodeMesh, Mesh.GetId());
+#endif
+}
+
+void Display(const FModelMesh& MeshModel)
+{
+#ifdef CADKERNEL_DEBUG
+	for (const FFaceMesh* Mesh : MeshModel.GetFaceMeshes())
+	{
+		if (Mesh)
+		{
+			DisplayMesh(*Mesh);
+		}
+	}
+#endif
+}
+
+void DisplaySegment(const FPoint& Point1, const FPoint& Point2, FIdent Ident, EVisuProperty Property)
+{
+#ifdef CADKERNEL_DEBUG
+	F3DDebugSegment G(Ident);
+	DrawSegment(Point1, Point2, Property);
+#endif
+}
+
+void DisplaySegment(const FPoint2D& Point1, const FPoint2D& Point2, FIdent Ident, EVisuProperty Property, bool bWithOrientation)
+{
+#ifdef CADKERNEL_DEBUG
+	F3DDebugSegment G(Ident);
+	if (bWithOrientation)
+	{
+		DrawSegmentOrientation(Point1, Point2, Property);
+	}
+	DrawSegment(Point1, Point2, Property);
+#endif
+}
+
+void DisplaySegment(const FPointH& Point1, const FPointH& Point2, FIdent Ident, EVisuProperty Property)
+{
+#ifdef CADKERNEL_DEBUG
+	F3DDebugSegment G(Ident);
+	DrawSegment(Point1, Point2, Property);
+#endif
+}
+
+void DisplayLoop(const FTopologicalFace& Surface)
+{
+#ifdef CADKERNEL_DEBUG
+	for (const TSharedPtr<FTopologicalLoop>& Loop : Surface.GetLoops())
+	{
+		Display(*Loop);
+	}
+#endif
+}
+
+
+void Display(const FTopologicalLoop& Loop)
+{
+#ifdef CADKERNEL_DEBUG
+	for (const FOrientedEdge& Edge : Loop.GetEdges())
+	{
+		Display(*Edge.Entity);
+	}
+#endif
+}
+
+void Display2D(const FTopologicalLoop& Loop)
+{
+#ifdef CADKERNEL_DEBUG
+	for (const FOrientedEdge& Edge : Loop.GetEdges())
+	{
+		Display2D(*Edge.Entity);
+	}
+#endif
+}
+
+void Display2DWithScale(const FTopologicalLoop& Loop)
+{
+#ifdef CADKERNEL_DEBUG
+	for (const FOrientedEdge& Edge : Loop.GetEdges())
+	{
+		Display2DWithScale(*Edge.Entity);
+	}
+#endif
+}
+
+} // namespace UE::CADKernel
diff --git a/Engine/Source/Runtime/Datasmith/CADKernel/Base/Private/UI/Progress.cpp b/Engine/Source/Runtime/Datasmith/CADKernel/Base/Private/UI/Progress.cpp
new file mode 100644
--- /dev/null
+++ b/Engine/Source/Runtime/Datasmith/CADKernel/Base/Private/UI/Progress.cpp
@@ -0,0 +1,84 @@
+// Copyright Epic Games, Inc. All Rights Reserved.
+
+#include "UI/Progress.h"
+#include "UI/Message.h"
+#include "Utils/Util.h"
+#include "Core/System.h"
+
+namespace UE::CADKernel
+{
+
+void FProgressManager::SetCurrentProgress(FProgress* InNewCurrent)
+{
+	if (!RootProgress)
+	{
+		RootProgress = InNewCurrent;
+		CurrentProgress = RootProgress;
+	}
+	else if (InNewCurrent == nullptr)
+	{
+		RootProgress = nullptr;
+		CurrentProgress = nullptr;
+	}
+
+	CurrentProgress = InNewCurrent;
+	Update();
+}
+
+double FProgressManager::GetProgression() const
+{
+	if (!RootProgress) 
+	{
+		return 0.;
+	}
+	return RootProgress->GetProgression();
+}
+
+FString FProgressManager::GetCurrentStep() const
+{
+	FString CurrentStepStr;
+
+	for (FProgress* Progress = GetCurrentProgress(); Progress; Progress = Progress->GetParent())
+	{
+		if (!Progress->GetName().IsEmpty())
+		{
+			CurrentStepStr = Progress->GetName() + TEXT(" -> ") + CurrentStepStr;
+		}
+	}
+	CurrentStepStr.RemoveFromEnd(TEXT(" -> "));
+	return CurrentStepStr;
+}
+
+
+FProgress::FProgress(int32 InStepCount, const FString& InStepName)
+: Name(InStepName)
+, StepCount(InStepCount)
+{
+	Parent = FSystem::Get().GetProgressManager().GetCurrentProgress();
+
+	if (Parent)
+	{
+		Parent->AddUnderlying(this);
+	}
+
+	FSystem::Get().GetProgressManager().SetCurrentProgress(this);
+}
+
+FProgress::~FProgress()
+{
+	if (Parent)
+	{
+		Parent->UnderlyingFinished(this);
+	}
+
+	FSystem::Get().GetProgressManager().SetCurrentProgress(Parent);
+	FSystem::Get().GetProgressManager().Update();
+}
+
+void FProgress::Increase(int32 StepSize)
+{
+	Progression += StepSize;
+	FSystem::Get().GetProgressManager().Update();
+}
+
+}
\ No newline at end of file
diff --git a/Engine/Source/Runtime/Datasmith/CADKernel/Base/Private/UI/Visu.cpp b/Engine/Source/Runtime/Datasmith/CADKernel/Base/Private/UI/Visu.cpp
new file mode 100644
--- /dev/null
+++ b/Engine/Source/Runtime/Datasmith/CADKernel/Base/Private/UI/Visu.cpp
@@ -0,0 +1,40 @@
+// Copyright Epic Games, Inc. All Rights Reserved.
+
+#include "UI/Visu.h"
+
+namespace UE::CADKernel
+{
+const TCHAR* VisuPropertyNames[] =
+{
+	TEXT("Bounding Box"),
+	TEXT("Points"),
+	TEXT("Elements"),
+	TEXT("Curves"),
+	TEXT("Iso"),
+	TEXT("Control Lines"),
+	TEXT("Edges"),
+	TEXT("Nodes"),
+	TEXT("Border Edges"),
+	TEXT("Multiple Edges"),
+	TEXT("Geometrical Borders"),
+	TEXT("Control Points"),
+
+	TEXT("YellowPoint"),
+	TEXT("YellowCurve"),
+	TEXT("BluePoint"),
+	TEXT("BlueCurve"),
+	TEXT("RedPoint"),
+	TEXT("RedCurve"),
+	TEXT("PurplePoint"),
+	TEXT("PurpleCurve"),
+	TEXT("GreenPoint"),
+	TEXT("GreenCurve"),
+	TEXT("PinkPoint"),
+	TEXT("PinkCurve"),
+	TEXT("OrangePoint"),
+	TEXT("OrangeCurve"),
+	nullptr
+};
+
+
+}
diff --git a/Engine/Source/Runtime/Datasmith/CADKernel/Base/Public/Core/System.h b/Engine/Source/Runtime/Datasmith/CADKernel/Base/Public/Core/System.h
new file mode 100644
--- /dev/null
+++ b/Engine/Source/Runtime/Datasmith/CADKernel/Base/Public/Core/System.h
@@ -0,0 +1,156 @@
+// Copyright Epic Games, Inc. All Rights Reserved.
+#pragma once
+
+#include "Core/Types.h"
+#include "UI/Console.h"
+#include "UI/Progress.h"
+#include "UI/Visu.h"
+
+class FArchive;
+
+namespace UE::CADKernel
+{
+class FConsole;
+class FKernelParameters;
+class FVisu;
+
+class FSystem
+{
+	friend class FDatabase;
+
+protected:
+
+	FString ProductName = FString(TEXT("CADKernel"));
+
+	FVisu DefaultVisu;
+	FVisu* Viewer;
+
+	FConsole DefaultConsole;
+	FConsole* Console;
+
+	FProgressManager DefaultProgressManager;
+	FProgressManager* ProgressManager;
+
+	static TUniquePtr<FSystem> Instance;
+
+	TSharedPtr<FArchive> LogFile;
+	EVerboseLevel LogLevel;
+	TSharedPtr<FArchive> SpyFile;
+
+#if defined(CADKERNEL_DEV) || defined(CADKERNEL_STDA)
+	TSharedRef<FKernelParameters> Parameters;
+	TSharedPtr<FArchive> ReportFile;
+	TSharedPtr<FArchive> ReportHeaderFile;
+#endif
+
+	EVerboseLevel VerboseLevel;
+
+public:
+
+	FSystem();
+
+	void Initialize(bool bIsDll = true, const FString& LogFilePath = TEXT(""), const FString& SpyFilePath = TEXT(""));
+
+	void Shutdown();
+
+	void CloseLogFiles();
+
+	static FString GetToolkitVersion();
+	static FString GetCompilationDate();
+
+	FVisu* GetVisu() const
+	{
+		return Viewer;
+	}
+
+	FConsole& GetConsole() const
+	{
+		ensureCADKernel(Console);
+		return *Console;
+	}
+
+	void SetViewer(FVisu* NewViewer)
+	{
+		Viewer = NewViewer;
+	}
+
+	void SetConsole(FConsole* InConsole)
+	{
+		Console = InConsole;
+	}
+
+	FProgressManager& GetProgressManager()
+	{
+		ensureCADKernel(ProgressManager);
+		return *ProgressManager;
+	}
+
+	void SetProgressManager(FProgressManager* InProgressManager)
+	{
+		ProgressManager = InProgressManager;
+	}
+
+	EVerboseLevel GetVerboseLevel() const
+	{
+		return VerboseLevel;
+	}
+
+	void SetVerboseLevel(EVerboseLevel Level)
+	{
+		VerboseLevel = Level;
+	}
+
+	void InitializeCADKernel();
+
+	EVerboseLevel GetLogLevel() const
+	{
+		return LogLevel;
+	}
+
+	void DefineLogFile(const FString& LogFilePath, EVerboseLevel Level = Log);
+	TSharedPtr<FArchive> GetLogFile() const
+	{
+		return LogFile;
+	}
+
+	void DefineSpyFile(const FString& SpyFilePath);
+	TSharedPtr<FArchive> GetSpyFile() const
+	{
+		return SpyFile;
+	}
+
+#if defined(CADKERNEL_DEV) || defined(CADKERNEL_STDA)
+	TSharedRef<FKernelParameters> GetParameters()
+	{
+		return Parameters;
+	}
+
+	void DefineReportFile(const FString& InLogFilePath);
+	TSharedPtr<FArchive> GetReportFile() const
+	{
+		return ReportFile;
+	}
+
+	TSharedPtr<FArchive> GetReportHeaderFile() const
+	{
+		return ReportHeaderFile;
+	}
+#endif
+
+	static FSystem& Get()
+	{
+		if (!Instance.IsValid())
+		{
+			Instance = MakeUnique<FSystem>();
+		}
+		return *Instance;
+	}
+
+protected:
+
+	void PrintHeader();
+
+};
+
+} // namespace UE::CADKernel
+
diff --git a/Engine/Source/Runtime/Datasmith/CADKernel/Base/Public/Geo/GeoPoint.h b/Engine/Source/Runtime/Datasmith/CADKernel/Base/Public/Geo/GeoPoint.h
new file mode 100644
--- /dev/null
+++ b/Engine/Source/Runtime/Datasmith/CADKernel/Base/Public/Geo/GeoPoint.h
@@ -0,0 +1,158 @@
+// Copyright Epic Games, Inc. All Rights Reserved.
+#pragma once
+
+#include "Core/Types.h"
+#include "Geo/GeoEnum.h"
+#include "Math/Point.h"
+
+namespace UE::CADKernel
+{
+struct FCurvePoint2D
+{
+	int32 DerivativeOrder = -1;
+	FPoint2D Point = FPoint2D::ZeroPoint;
+	FPoint2D Gradient = FPoint2D::ZeroPoint;
+	FPoint2D Laplacian = FPoint2D::ZeroPoint;
+};
+
+struct FSurfacicPoint
+{
+	int32 DerivativeOrder = -1;
+	FPoint Point = FPoint::ZeroPoint;
+	FPoint GradientU = FPoint::ZeroPoint;
+	FPoint GradientV = FPoint::ZeroPoint;
+	FPoint LaplacianU = FPoint::ZeroPoint;
+	FPoint LaplacianV = FPoint::ZeroPoint;
+	FPoint LaplacianUV = FPoint::ZeroPoint;
+};
+
+struct FSurfacicCurvePoint
+{
+	bool bWithNormals;
+	bool bWithTangent;
+
+	FPoint2D Point2D = FPoint2D::ZeroPoint;
+	FPoint Point = FPoint::ZeroPoint;
+	FPoint Normal = FPoint::ZeroPoint;
+	FPoint Tangent = FPoint::ZeroPoint;
+};
+
+struct FSurfacicCurvePointWithTolerance
+{
+	FPoint2D Point2D = FPoint2D::ZeroPoint;
+	FPoint Point = FPoint::ZeroPoint;
+	FSurfacicTolerance Tolerance = FPoint2D::FarawayPoint;
+};
+
+typedef FSurfacicCurvePointWithTolerance FSurfacicCurveExtremities[2];
+
+struct FCurvePoint
+{
+	int32 DerivativeOrder = -1;
+	FPoint Point = FPoint::ZeroPoint;
+	FPoint Gradient = FPoint::ZeroPoint;
+	FPoint Laplacian = FPoint::ZeroPoint;
+
+	FCurvePoint() = default;
+
+	FCurvePoint(FPoint InPoint)
+		: Point(InPoint)
+	{
+		Gradient = FPoint::ZeroPoint;
+		Laplacian = FPoint::ZeroPoint;
+	}
+
+	void Init()
+	{
+		Point = FPoint::ZeroPoint;
+		Gradient = FPoint::ZeroPoint;
+		Laplacian = FPoint::ZeroPoint;
+	}
+
+	/**
+	 * Compute the 3D surface curve point property (3D Coordinate, Gradient, Laplacian) according to
+	 * its 2D curve point property and the 3D surface point property
+	 */
+	void Combine(const FCurvePoint2D& Point2D, const FSurfacicPoint& SurfacicPoint)
+	{
+		ensureCADKernel(Point2D.DerivativeOrder >= 0);
+		ensureCADKernel(SurfacicPoint.DerivativeOrder >= 0);
+
+		ensureCADKernel(Point2D.DerivativeOrder <= SurfacicPoint.DerivativeOrder);
+
+		DerivativeOrder = Point2D.DerivativeOrder;
+		Point = SurfacicPoint.Point;
+
+		if (DerivativeOrder > 0)
+		{
+			Gradient = SurfacicPoint.GradientU * Point2D.Gradient.U + SurfacicPoint.GradientV * Point2D.Gradient.V;
+		}
+
+		if (DerivativeOrder > 1)
+		{
+			Laplacian = SurfacicPoint.LaplacianU * FMath::Square(Point2D.Gradient.U)
+				+ 2.0 * SurfacicPoint.LaplacianUV * Point2D.Gradient.U * Point2D.Gradient.V
+				+ SurfacicPoint.LaplacianV * FMath::Square(Point2D.Gradient.V)
+				+ SurfacicPoint.GradientU * Point2D.Laplacian.U
+				+ SurfacicPoint.GradientV * Point2D.Laplacian.V;
+		}
+	}
+};
+
+struct FCoordinateGrid
+{
+	TArray<double> Coordinates[2];
+
+	FCoordinateGrid()
+	{
+	}
+
+	FCoordinateGrid(const TArray<double>& InUCoordinates, const TArray<double>& InVCoordinates)
+	{
+		Coordinates[EIso::IsoU] = InUCoordinates;
+		Coordinates[EIso::IsoV] = InVCoordinates;
+	}
+
+	void Swap(TArray<double>& InUCoordinates, TArray<double>& InVCoordinates)
+	{
+		::Swap(Coordinates[EIso::IsoU], InUCoordinates);
+		::Swap(Coordinates[EIso::IsoV], InVCoordinates);
+	}
+
+	int32 Count() const
+	{
+		return Coordinates[EIso::IsoU].Num() * Coordinates[EIso::IsoV].Num();
+	}
+
+	int32 IsoCount(EIso Iso) const
+	{
+		return Coordinates[Iso].Num();
+	}
+
+	void SetNum(int32 UNumber, int32 VNumber)
+	{
+		Coordinates[0].SetNum(UNumber);
+		Coordinates[1].SetNum(VNumber);
+	}
+
+	void Empty(int32 UNumber = 0, int32 VNumber = 0)
+	{
+		Coordinates[0].Empty(UNumber);
+		Coordinates[1].Empty(VNumber);
+	}
+
+	constexpr TArray<double>& operator[](EIso Iso)
+	{
+		ensureCADKernel(Iso == 0 || Iso == 1);
+		return Coordinates[Iso];
+	}
+
+	constexpr const TArray<double>& operator[](EIso Iso) const
+	{
+		ensureCADKernel(Iso == 0 || Iso == 1);
+		return Coordinates[Iso];
+	}
+};
+
+
+}
\ No newline at end of file
diff --git a/Engine/Source/Runtime/Datasmith/CADKernel/Base/Public/Math/Point.h b/Engine/Source/Runtime/Datasmith/CADKernel/Base/Public/Math/Point.h
new file mode 100644
--- /dev/null
+++ b/Engine/Source/Runtime/Datasmith/CADKernel/Base/Public/Math/Point.h
@@ -0,0 +1,843 @@
+// Copyright Epic Games, Inc. All Rights Reserved.
+#pragma once
+
+#include "Core/Types.h"
+#include "Math/MathConst.h"
+
+// this will eventually be reconciled with what the Geometry team is using.
+namespace UE::CADKernel
+{
+
+class FPoint2D;
+
+class CADKERNEL_API FPoint
+{
+public:
+	union
+	{
+		struct
+		{
+			double X;
+			double Y;
+			double Z;
+		};
+
+		UE_DEPRECATED(all, "For internal use only")
+		double XYZ[3];
+	};
+
+public:
+
+	/** A zero point (0,0,0)*/
+	static const FPoint ZeroPoint;
+	static const FPoint UnitPoint;
+	static const FPoint FarawayPoint;
+	static const int32 Dimension;
+
+	explicit FPoint(double InCoordX = 0., double InCoordY = 0., double InCoordZ = 0.)
+	{
+		X = InCoordX; Y = InCoordY; Z = InCoordZ;
+	}
+
+	FPoint(const FPoint2D& InPoint2D);
+
+	FPoint(const double* const InPoint)
+		: X(InPoint[0])
+		, Y(InPoint[1])
+		, Z(InPoint[2])
+	{
+	}
+
+	FPoint(const FPoint& Point)
+		: X(Point.X)
+		, Y(Point.Y)
+		, Z(Point.Z)
+	{
+	}
+
+	constexpr double& operator[](int32 Index)
+	{
+PRAGMA_DISABLE_DEPRECATION_WARNINGS
+		return XYZ[Index];
+PRAGMA_ENABLE_DEPRECATION_WARNINGS
+	}
+
+	constexpr double operator[](int32 Index) const
+	{
+PRAGMA_DISABLE_DEPRECATION_WARNINGS
+		return XYZ[Index];
+PRAGMA_ENABLE_DEPRECATION_WARNINGS
+	}
+
+	bool operator==(const FPoint& Point) const
+	{
+		return FMath::IsNearlyEqual(X, Point.X) && FMath::IsNearlyEqual(Y, Point.Y) && FMath::IsNearlyEqual(Z, Point.Z);
+	}
+
+	bool operator!=(const FPoint& Point) const
+	{
+		return !(*this == Point);
+	};
+
+	FPoint operator+(const FPoint& Point) const
+	{
+		return FPoint(X + Point.X, Y + Point.Y, Z + Point.Z);
+	}
+
+	FPoint operator+(const FVector& Point) const
+	{
+		return FPoint(X + Point.X, Y + Point.Y, Z + Point.Z);
+	}
+
+	FPoint operator+(const FVector3f& Point) const
+	{
+		return FPoint(X + Point.X, Y + Point.Y, Z + Point.Z);
+	}
+
+	FPoint operator-(const FPoint& Point) const
+	{
+		return FPoint(X - Point.X, Y - Point.Y, Z - Point.Z);
+	}
+
+	FPoint operator^(const FPoint& Point) const
+	{
+		return FPoint(Y * Point.Z - Z * Point.Y, Z * Point.X - X * Point.Z, X * Point.Y - Y * Point.X);
+	}
+
+	double operator*(const FPoint& Point) const
+	{
+		return X * Point.X + Y * Point.Y + Z * Point.Z;
+	}
+
+	FPoint operator*(double Scale) const
+	{
+		return FPoint(X * Scale, Y * Scale, Z * Scale);
+	}
+
+	FPoint operator/(double InvScale) const
+	{
+		if (FMath::IsNearlyZero(InvScale))
+		{
+			return *this;
+		}
+		return FPoint(X / InvScale, Y / InvScale, Z / InvScale);
+	}
+
+	FPoint operator-() const
+	{
+		return FPoint(-X, -Y, -Z);
+	}
+
+	FPoint& operator=(const FPoint& Point)
+	{
+		X = Point.X;
+		Y = Point.Y;
+		Z = Point.Z;
+		return *this;
+	}
+
+	FPoint& operator+=(const FPoint& Point)
+	{
+		X += Point.X;
+		Y += Point.Y;
+		Z += Point.Z;
+		return *this;
+	}
+
+	FPoint& operator+=(const FVector& Point)
+	{
+		X += Point.X;
+		Y += Point.Y;
+		Z += Point.Z;
+		return *this;
+	}
+
+	FPoint& operator+=(const FVector3f& Point)
+	{
+		X += Point.X;
+		Y += Point.Y;
+		Z += Point.Z;
+		return *this;
+	}
+
+	FPoint& operator-=(const FPoint& Point)
+	{
+		X -= Point.X;
+		Y -= Point.Y;
+		Z -= Point.Z;
+		return *this;
+	}
+
+	FPoint& operator*=(double Scale)
+	{
+		X *= Scale;
+		Y *= Scale;
+		Z *= Scale;
+		return *this;
+	}
+
+	void NonUniformScale(const FPoint& Scale)
+	{
+		X *= Scale.X;
+		Y *= Scale.Y;
+		Z *= Scale.Z;
+	}
+
+	FPoint& operator/=(double InvScale)
+	{
+		if (FMath::IsNearlyZero(InvScale))
+		{
+			return *this;
+		}
+
+		X /= InvScale;
+		Y /= InvScale;
+		Z /= InvScale;
+		return *this;
+	}
+
+	bool operator<(const FPoint& Other) const
+	{
+		int32 Compare = RealCompare(X, Other.X);
+		if (Compare < 0)
+		{
+			return true;
+		}
+		if (Compare > 0)
+		{
+			return false;
+		}
+
+		Compare = RealCompare(Y, Other.Y);
+		if (Compare < 0)
+		{
+			return true;
+		}
+		if (Compare > 0)
+		{
+			return false;
+		}
+
+		Compare = RealCompare(Z, Other.Z);
+		return (Compare < 0);
+	}
+
+	operator FVector() const
+	{
+		return FVector(X, Y, Z);
+	}
+
+	operator FVector3f() const
+	{
+		return FVector3f((float)X, (float)Y, (float)Z);
+	}
+
+	void SetMin(const FPoint& Point)
+	{
+		X = FMath::Min(X, Point.X);
+		Y = FMath::Min(Y, Point.Y);
+		Z = FMath::Min(Z, Point.Z);
+	}
+
+	void SetMax(const FPoint& Point)
+	{
+		X = FMath::Max(X, Point.X);
+		Y = FMath::Max(Y, Point.Y);
+		Z = FMath::Max(Z, Point.Z);
+	}
+
+	void Set(double InCoordX = 0., double InCoordY = 0., double InCoordZ = 0.)
+	{
+		X = InCoordX;
+		Y = InCoordY;
+		Z = InCoordZ;
+	}
+
+	void Set(double* InCoordinates)
+	{
+		X = InCoordinates[0];
+		Y = InCoordinates[1];
+		Z = InCoordinates[2];
+	}
+
+	double Length() const
+	{
+		return FMath::Sqrt((X * X) + (Y * Y) + (Z * Z));
+	}
+
+	double SquareLength() const
+	{
+		return (X * X) + (Y * Y) + (Z * Z);
+	}
+
+	FPoint& Normalize()
+	{
+		double Norm = Length();
+		if (FMath::IsNearlyZero(Norm))
+		{
+			*this = FPoint::UnitPoint.Normalize();
+		}
+		else
+		{
+			*this /= Norm;
+		}
+		return *this;
+	}
+
+	FPoint Normalize() const 
+	{
+		double Norm = Length();
+		if (FMath::IsNearlyZero(Norm))
+		{
+			return FPoint::UnitPoint.Normalize();
+		}
+		else
+		{
+			return *this / Norm;
+		}
+	}
+
+	FPoint Middle(const FPoint& Point) const
+	{
+		return FPoint(0.5 * (X + Point.X), 0.5 * (Y + Point.Y), 0.5 * (Z + Point.Z));
+	}
+
+	double Distance(const FPoint& Point) const
+	{
+		return FMath::Sqrt(FMath::Square(Point.X - X) + FMath::Square(Point.Y - Y) + FMath::Square(Point.Z - Z));
+	}
+
+	double SquareDistance(const FPoint& Point) const
+	{
+		return FMath::Square(Point.X - X) + FMath::Square(Point.Y - Y) + FMath::Square(Point.Z - Z);
+	}
+
+	double ComputeCosinus(const FPoint& OtherVector) const
+	{
+		FPoint ThisNormalized = *this;
+		FPoint OtherNormalized = OtherVector;
+
+		ThisNormalized.Normalize();
+		OtherNormalized.Normalize();
+
+		double Cosinus = ThisNormalized * OtherNormalized;
+
+		return FMath::Max(-1.0, FMath::Min(Cosinus, 1.0));
+	}
+
+	double ComputeSinus(const FPoint& OtherVector) const
+	{
+		FPoint ThisNormalized = *this;
+		FPoint OtherNormalized = OtherVector;
+
+		ThisNormalized.Normalize();
+		OtherNormalized.Normalize();
+
+		FPoint SinusPoint = ThisNormalized ^ OtherNormalized;
+		double Sinus = SinusPoint.Length();
+		return FMath::Min(Sinus, 1.0);
+	}
+
+	double ComputeAngle(const FPoint& OtherVector) const
+	{
+		double CosAngle = ComputeCosinus(OtherVector);
+		return acos(CosAngle);
+	}
+
+	double SignedAngle(const FPoint& Vector2, const FPoint& Normal) const;
+
+	/**
+	 * Return the projection of the point on the diagonal axis (of vector (1,1,1))
+	 * i.e. return X + Y + Z
+	 */
+	double DiagonalAxisCoordinate() const
+	{
+		return X + Y + Z;
+	}
+
+	static double MixedTripleProduct(const FPoint& VectorA, const FPoint& VectorB, const FPoint& VectorC)
+	{
+		return VectorA * (VectorB ^ VectorC);
+	}
+
+	friend FArchive& operator<<(FArchive& Ar, FPoint& Point)
+	{
+		Ar.Serialize(&Point, sizeof(FPoint));
+		return Ar;
+	}
+
+	friend FPoint Abs(const FPoint& Point)
+	{
+		double PointX = FMath::Abs(Point.X);
+		double PointY = FMath::Abs(Point.Y);
+		double PointZ = FMath::Abs(Point.Z);
+		return FPoint(PointX, PointY, PointZ);
+	}
+
+};
+
+class CADKERNEL_API FPointH
+{
+public:
+	union
+	{
+		struct
+		{
+			double X;
+			double Y;
+			double Z;
+			double W;
+		};
+
+		UE_DEPRECATED(all, "For internal use only")
+		double XYZW[4];
+	};
+
+public:
+	/** A zero point (0,0,0)*/
+	static const FPointH ZeroPoint;
+	static const FPointH FarawayPoint;
+	static const int32 Dimension;
+
+	FPointH()
+		: X(0.)
+		, Y(0.)
+		, Z(0.)
+		, W(1.)
+	{
+	}
+
+	FPointH(const FPoint& Point, double Weight)
+		: X(Point.X* Weight)
+		, Y(Point.Y* Weight)
+		, Z(Point.Z* Weight)
+		, W(Weight)
+	{
+	}
+
+	FPointH(const double InX, const double InY, const double InZ, const double InW)
+		: X(InX)
+		, Y(InY)
+		, Z(InZ)
+		, W(InW)
+	{
+	}
+
+	constexpr double& operator[](int32 Index)
+	{
+PRAGMA_DISABLE_DEPRECATION_WARNINGS
+		return XYZW[Index];
+PRAGMA_ENABLE_DEPRECATION_WARNINGS
+	}
+
+	constexpr double operator[](int32 Index) const
+	{
+PRAGMA_DISABLE_DEPRECATION_WARNINGS
+		return XYZW[Index];
+PRAGMA_ENABLE_DEPRECATION_WARNINGS
+	}
+
+	FPointH operator+(const FPointH& Point) const
+	{
+		return FPointH(X + Point.X, Y + Point.Y, Z + Point.Z, W + Point.W);
+	}
+
+	FPointH operator-(const FPointH& Point) const
+	{
+		return FPointH(X - Point.X, Y - Point.Y, Z - Point.Z, W - Point.W);
+	}
+
+	FPointH operator*(double Factor) const
+	{
+		return FPointH(X * Factor, Y * Factor, Z * Factor, W * Factor);
+	}
+
+	operator FPoint() const
+	{
+		return FPoint(X / W, Y / W, Z / W);
+	}
+
+	friend FArchive& operator<<(FArchive& Ar, FPointH& Point)
+	{
+		Ar.Serialize(&Point, sizeof(FPointH));
+		return Ar;
+	}
+};
+
+class CADKERNEL_API FPoint2D
+{
+public:
+	union
+	{
+		struct
+		{
+			double U;
+			double V;
+		};
+
+		UE_DEPRECATED(all, "For internal use only")
+		double UV[2];
+	};
+
+
+public:
+
+	/** A zero point (0,0,0)*/
+	static const FPoint2D ZeroPoint;
+	static const FPoint2D FarawayPoint;
+	static const int32 Dimension;
+
+	FPoint2D()
+	{
+		U = V = 0.0;
+	}
+
+	FPoint2D(double NewU, double NewV)
+		: U(NewU)
+		, V(NewV)
+	{
+	}
+
+	FPoint2D(const FPoint& Point)
+		: U(Point.X)
+		, V(Point.Y)
+	{
+	}
+
+	FPoint2D(const FPoint2D& Point)
+		: U(Point.U)
+		, V(Point.V)
+	{
+	}
+
+	FPoint2D& operator=(const FPoint2D& Point)
+	{
+		U = Point.U;
+		V = Point.V;
+		return *this;
+	}
+
+	FPoint2D Rotate(double Theta)
+	{
+		return FPoint2D(U * cos(Theta) - V * sin(Theta), U * sin(Theta) + V * cos(Theta));
+	}
+
+	FPoint2D Middle(const FPoint2D& Point) const
+	{
+		return FPoint(0.5 * (U + Point.U), 0.5 * (V + Point.V));
+	}
+
+	double Distance(const FPoint2D& Point) const
+	{
+		return FMath::Sqrt(SquareDistance(Point));
+	}
+
+	double SquareDistance(const FPoint2D& Point) const
+	{
+		return FMath::Square(Point.U - U) + FMath::Square(Point.V - V);
+	}
+
+	double Length() const
+	{
+		return FMath::Sqrt(SquareLength());
+	}
+
+	double SquareLength() const
+	{
+		return FMath::Square(U) + FMath::Square(V);
+	}
+
+	FPoint2D& Normalize(double& Norm)
+	{
+		Norm = Length();
+		if (FMath::IsNearlyZero(Norm))
+		{
+			U = V = 0.0;
+		}
+		else
+		{
+			U /= Norm;
+			V /= Norm;
+		}
+		return *this;
+	}
+
+	FPoint2D& Normalize()
+	{
+		double Norm = 0;
+		return Normalize(Norm);
+	}
+
+	constexpr double& operator[](int32 Index)
+	{
+PRAGMA_DISABLE_DEPRECATION_WARNINGS
+		return UV[Index];
+PRAGMA_ENABLE_DEPRECATION_WARNINGS
+	}
+
+	constexpr double operator[](int32 Index) const
+	{
+PRAGMA_DISABLE_DEPRECATION_WARNINGS
+		return UV[Index];
+PRAGMA_ENABLE_DEPRECATION_WARNINGS
+	}
+
+	bool operator==(const FPoint2D& Point) const
+	{
+		return FMath::IsNearlyEqual(U, Point.U) && FMath::IsNearlyEqual(V, Point.V);
+	}
+
+	bool operator<(const FPoint2D& other) const
+	{
+		int32 cmp = RealCompare(U, other.U);
+		if (cmp < 0)
+		{
+			return true;
+		}
+		if (cmp > 0)
+		{
+			return false;
+		}
+
+		cmp = RealCompare(V, other.V);
+		return (cmp < 0);
+	}
+
+	void Set(double NewU, double NewV = 0.)
+	{
+		U = NewU;
+		V = NewV;
+	}
+
+	void Set(double* NewCoordinate)
+	{
+		U = NewCoordinate[0];
+		V = NewCoordinate[1];
+	}
+
+	FPoint2D operator+(const FPoint2D& Point) const
+	{
+		return FPoint2D(U + Point.U, V + Point.V);
+	}
+
+	FPoint2D operator-(const FPoint2D& Point) const
+	{
+		return FPoint2D(U - Point.U, V - Point.V);
+	}
+
+	double operator*(const FPoint2D& Point) const
+	{
+		return U * Point.U + V * Point.V;
+	}
+
+	double operator^(const FPoint2D& Point) const
+	{
+		return (U * Point.V - V * Point.U);
+	}
+
+	FPoint2D operator*(double Scale) const
+	{
+		return FPoint2D(U * Scale, V * Scale);
+	}
+
+	FPoint2D operator/(double InvScale) const
+	{
+		if (FMath::IsNearlyZero(InvScale))
+		{
+			return *this;
+		}
+
+		return FPoint2D(U / InvScale, V / InvScale);
+	}
+
+	FPoint2D operator-() const
+	{
+		return FPoint2D(-U, -V);
+	}
+
+	FPoint2D& operator+=(const FPoint2D& Point)
+	{
+		U += Point.U;
+		V += Point.V;
+		return *this;
+	}
+
+	FPoint2D& operator-=(const FPoint2D& Point)
+	{
+		U -= Point.U;
+		V -= Point.V;
+		return *this;
+	}
+
+	FPoint2D& operator*=(double Scale)
+	{
+		U *= Scale;
+		V *= Scale;
+		return *this;
+	}
+
+	FPoint2D& operator/=(double InvScale)
+	{
+		if (FMath::IsNearlyZero(InvScale))
+		{
+			return *this;
+		}
+
+		U /= InvScale;
+		V /= InvScale;
+		return *this;
+	}
+
+	double ComputeCosinus(const FPoint2D& OtherVector) const
+	{
+		FPoint2D ThisNormalized = *this;
+		FPoint2D OtherNormalized = OtherVector;
+
+		ThisNormalized.Normalize();
+		OtherNormalized.Normalize();
+
+		double Cosinus = ThisNormalized * OtherNormalized;
+
+		return FMath::Max(-1.0, FMath::Min(Cosinus, 1.0));
+	}
+
+	FPoint2D GetPerpendicularVector()
+	{
+		return FPoint2D(-V, U);
+	}
+
+	/**
+	 * Return the projection of the point on the diagonal axis (of vector (1,1))
+	 * i.e. return U + V
+	 */
+	double DiagonalAxisCoordinate() const
+	{
+		return U + V;
+	}
+
+	friend FArchive& operator<<(FArchive& Ar, FPoint2D& Point)
+	{
+		Ar.Serialize(&Point, sizeof(FPoint2D));
+		return Ar;
+	}
+
+	friend FPoint2D Abs(const FPoint2D& Point)
+	{
+		double PointU = FMath::Abs(Point.U);
+		double PointV = FMath::Abs(Point.V);
+		return { PointU, PointV };
+	}
+
+	friend FPoint2D Max(const FPoint2D& PointA, const FPoint2D& PointB)
+	{
+		return { FMath::Max(PointA.U, PointB.U), FMath::Max(PointA.V, PointB.V) };
+	}
+};
+
+inline FPoint::FPoint(const FPoint2D& Point)
+{
+	X = Point.U; Y = Point.V; Z = 0.0;
+}
+
+typedef FPoint2D FSurfacicTolerance;
+
+class CADKERNEL_API FFPoint
+{
+public:
+	union
+	{
+		struct
+		{
+			float X;
+			float Y;
+			float Z;
+		};
+
+		UE_DEPRECATED(all, "For internal use only")
+		float XYZ[3];
+	};
+
+	/** A zero point (0,0,0)*/
+	static const FFPoint ZeroPoint;
+	static const FFPoint FarawayPoint;
+	static const int32 Dimension;
+
+	FFPoint(float InX = 0., float InY = 0., float InZ = 0.)
+	{
+		X = InX;
+		Y = InY;
+		Z = InZ;
+	}
+
+	FFPoint(const FPoint& Point)
+	{
+		X = (float)Point.X;
+		Y = (float)Point.Y;
+		Z = (float)Point.Z;
+	}
+
+	constexpr float& operator[](int32 Index)
+	{
+PRAGMA_DISABLE_DEPRECATION_WARNINGS
+		return XYZ[Index];
+PRAGMA_ENABLE_DEPRECATION_WARNINGS
+	}
+
+	constexpr float operator[](int32 Index) const
+	{
+PRAGMA_DISABLE_DEPRECATION_WARNINGS
+		return XYZ[Index];
+PRAGMA_ENABLE_DEPRECATION_WARNINGS
+	}
+
+	FFPoint& operator=(const FPoint& Point)
+	{
+		X = (float)Point.X; Y = (float)Point.Y; Z = (float)Point.Z; return *this;
+	}
+
+	friend FArchive& operator<<(FArchive& Ar, FFPoint& Point)
+	{
+		Ar.Serialize(&Point, sizeof(FFPoint));
+		return Ar;
+	}
+};
+
+inline FPoint operator*(double Scale, const FPoint& Point)
+{
+	return FPoint(Point.X * Scale, Point.Y * Scale, Point.Z * Scale);
+}
+
+inline FPoint operator*(float Scale, const FPoint& Point)
+{
+	double ScaleD = static_cast<double>(Scale);
+	return FPoint(Point.X * ScaleD, Point.Y * ScaleD, Point.Z * ScaleD);
+}
+
+inline FPoint2D operator*(double Scale, const FPoint2D& Point)
+{
+	return FPoint2D(Point.U * Scale, Point.V * Scale);
+}
+
+inline FPoint2D operator*(float Scale, const FPoint2D& Point)
+{
+	double ScaleD = static_cast<double>(Scale);
+	return FPoint2D(Point.U * ScaleD, Point.V * ScaleD);
+}
+
+}
+
+inline uint32 GetTypeHash(const UE::CADKernel::FPoint& Point)
+{
+	return HashCombine(GetTypeHash(Point.X), HashCombine(GetTypeHash(Point.Y), GetTypeHash(Point.Z)));
+}
+
+inline uint32 GetTypeHash(const UE::CADKernel::FPoint2D& Point)
+{
+	return HashCombine(GetTypeHash(Point.U), GetTypeHash(Point.V));
+}
+
+inline uint32 GetTypeHash(const UE::CADKernel::FPointH& Point)
+{
+	return HashCombine(GetTypeHash(Point.X), HashCombine(GetTypeHash(Point.Y), HashCombine(GetTypeHash(Point.Z), GetTypeHash(Point.W))));
+}
+
diff --git a/Engine/Source/Runtime/Datasmith/CADKernel/Base/Public/UI/Console.h b/Engine/Source/Runtime/Datasmith/CADKernel/Base/Public/UI/Console.h
new file mode 100644
--- /dev/null
+++ b/Engine/Source/Runtime/Datasmith/CADKernel/Base/Public/UI/Console.h
@@ -0,0 +1,24 @@
+// Copyright Epic Games, Inc. All Rights Reserved.
+#pragma once
+
+#include "Core/Types.h"
+
+#include "UI/Message.h"
+
+namespace UE::CADKernel
+{
+class CADKERNEL_API FConsole
+{
+public:
+
+	virtual ~FConsole()
+	{
+	}
+
+	virtual void Print(const TCHAR* Texte, EVerboseLevel VerboseLevel = EVerboseLevel::NoVerbose)
+	{}
+
+};
+
+} // namespace UE::CADKernel
+
diff --git a/Engine/Source/Runtime/Datasmith/CADKernel/Base/Public/UI/Display.h b/Engine/Source/Runtime/Datasmith/CADKernel/Base/Public/UI/Display.h
new file mode 100644
--- /dev/null
+++ b/Engine/Source/Runtime/Datasmith/CADKernel/Base/Public/UI/Display.h
@@ -0,0 +1,359 @@
+// Copyright Epic Games, Inc. All Rights Reserved.
+#pragma once
+
+#include "Core/System.h"
+#include "Core/Types.h"
+#include "Geo/GeoEnum.h"
+#include "Math/Point.h"
+#include "UI/Visu.h"
+
+#ifdef CADKERNEL_DEV
+#include "UI/DefineForDebug.h"
+#endif
+
+class FString;
+
+namespace UE::CADKernel
+{
+
+class FAABB2D;
+class FAABB;
+class FBody;
+class FCurve;
+class FEdgeMesh;
+class FEdgeMesh;
+class FEntity;
+class FFaceMesh;
+class FFaceMesh;
+class FGroup;
+class FMesh;
+class FModel;
+class FModelMesh;
+class FRestrictionCurve;
+class FShell;
+class FSurface;
+class FTopologicalEdge;
+class FTopologicalFace;
+class FTopologicalLoop;
+class FTopologicalVertex;
+class FVertexMesh;
+class FVertexMesh;
+struct FLinearBoundary;
+
+#ifdef CADKERNEL_DEBUG
+void Wait(bool bMakeWait = true);
+void Open3DDebugSession(FString name, const TArray<FIdent>& idList = TArray<FIdent>());
+inline void Open3DDebugSession(bool bIsDisplayed, FString name, const TArray<FIdent>& idList = TArray<FIdent>())
+{
+	if (bIsDisplayed)
+	{
+		Open3DDebugSession(name, idList);
+	}
+};
+void Close3DDebugSession(bool bIsDisplayed = true);
+#else
+inline void Wait(bool bMakeWait = true) {};
+inline void Open3DDebugSession(bool bIsDisplayed, FString name, const TArray<FIdent>& idList = TArray<FIdent>()) {};
+inline void Open3DDebugSession(FString name, const TArray<FIdent>& idList = TArray<FIdent>()) {};
+inline void Close3DDebugSession(bool bIsDisplayed = true) {};
+#endif
+
+class CADKERNEL_API F3DDebugSession
+{
+private:
+	bool bDisplay = true;
+
+public:
+	F3DDebugSession(FString Name, const TArray<FIdent>& Idents = TArray<FIdent>())
+	{
+#ifdef CADKERNEL_DEBUG
+		Open3DDebugSession(Name, Idents);
+#endif
+	}
+
+	F3DDebugSession(bool bInDisplay, FString Name, const TArray<FIdent>& Idents = TArray<FIdent>())
+		: bDisplay(bInDisplay)
+	{
+#ifdef CADKERNEL_DEBUG
+		if (bDisplay)
+		{
+			Open3DDebugSession(Name, Idents);
+		}
+#endif
+	}
+
+	~F3DDebugSession()
+	{
+#ifdef CADKERNEL_DEBUG
+		if (bDisplay)
+		{
+			Close3DDebugSession();
+		}
+#endif
+	}
+};
+
+CADKERNEL_API void Open3DDebugSegment(FIdent Ident);
+CADKERNEL_API void Close3DDebugSegment();
+
+class CADKERNEL_API F3DDebugSegment
+{
+public:
+	F3DDebugSegment(FIdent Ident)
+	{
+#ifdef CADKERNEL_DEBUG
+		Open3DDebugSegment(Ident);
+#endif
+	}
+
+	~F3DDebugSegment()
+	{
+#ifdef CADKERNEL_DEBUG
+		Close3DDebugSegment();
+#endif
+	}
+};
+
+CADKERNEL_API void FlushVisu();
+
+template<typename TPoint>
+void DrawPoint(const TPoint& InPoint, EVisuProperty Property = EVisuProperty::BluePoint)
+{
+#ifdef CADKERNEL_DEBUG
+	FSystem::Get().GetVisu()->DrawPoint(InPoint, Property);
+#endif
+}
+
+/**
+ * Draw a mesh element i.e. Element of dimension 1 is an edge, of dimension 2 is a triangle or quadrangle, of dimension 3 is tetrahedron, pyramid, hexahedron, ...
+ */
+CADKERNEL_API void DrawElement(int32 Dimension, TArray<FPoint>& Points, EVisuProperty Property = EVisuProperty::Element);
+
+template<typename TPoint>
+void Draw(const TArray<TPoint>& Points, EVisuProperty Property = EVisuProperty::BlueCurve)
+{
+#ifdef CADKERNEL_DEBUG
+	FSystem::Get().GetVisu()->DrawPolyline(Points, Property);
+#endif
+}
+
+CADKERNEL_API void DrawMesh(const TSharedPtr<FMesh>& mesh);
+
+CADKERNEL_API void DisplayEdgeCriteriaGrid(int32 EdgeId, const TArray<FPoint>& Points3D);
+
+template<typename TPoint>
+void DisplayPoint(const TPoint& Point, FIdent Ident)
+{
+#ifdef CADKERNEL_DEBUG
+	F3DDebugSegment G(Ident);
+	DrawPoint(Point);
+#endif
+}
+
+template<typename TPoint>
+void DisplayPoint(const TPoint& Point, EVisuProperty Property = EVisuProperty::BluePoint)
+{
+#ifdef CADKERNEL_DEBUG
+	DrawPoint(Point, Property);
+#endif
+}
+
+template<typename TPoint>
+void DisplayPoint(const TPoint& Point, EVisuProperty Property, FIdent Ident)
+{
+#ifdef CADKERNEL_DEBUG
+	F3DDebugSegment G(Ident);
+	DrawPoint(Point, Property);
+#endif
+}
+
+static void DisplayPoint2DWithScale(const FPoint2D& Point, EVisuProperty Property = EVisuProperty::BluePoint)
+{
+#ifdef CADKERNEL_DEBUG
+	DrawPoint(Point * DisplayScale, Property);
+#endif
+}
+
+static void DisplayPoint2DWithScale(const FPoint2D& Point, EVisuProperty Property, FIdent Ident)
+{
+#ifdef CADKERNEL_DEBUG
+	F3DDebugSegment G(Ident);
+	DrawPoint(Point * DisplayScale, Property);
+#endif
+}
+
+template<typename TPoint>
+void DisplayPoints(FString Message, const TArray<TPoint>& Points, EVisuProperty Property = EVisuProperty::BluePoint, bool bDisplay = true)
+{
+#ifdef CADKERNEL_DEBUG
+	if (!bDisplay)
+	{
+		return;
+	}
+
+	Open3DDebugSession(Message);
+	for (int32 Index = 0; Index < Points.Num(); ++Index)
+	{
+		DisplayPoint(Points[Index], Property);
+	}
+	Close3DDebugSession();
+#endif
+}
+
+CADKERNEL_API void DisplayProductTree(const FEntity& RootId);
+CADKERNEL_API void DisplayProductTree(const FModel& Model);
+CADKERNEL_API void DisplayProductTree(const FBody& Body);
+CADKERNEL_API void DisplayProductTree(const FShell& Shell);
+
+CADKERNEL_API void DisplayAABB(const FAABB& aabb, FIdent Ident = 0);
+CADKERNEL_API void DisplayAABB2D(const FAABB2D& aabb, FIdent Ident = 0);
+
+CADKERNEL_API void DisplayEntity(const FEntity& Entity);
+CADKERNEL_API void DisplayEntity2D(const FEntity& Entity);
+
+CADKERNEL_API void DisplayLoop(const FTopologicalFace& Entity);
+CADKERNEL_API void DisplayIsoCurve(const FSurface& CarrierSurface, double Coordinate, EIso Type);
+
+CADKERNEL_API void Display(const FPlane& plane, FIdent Ident = 0);
+
+CADKERNEL_API void Display(const FCurve& Curve);
+CADKERNEL_API void Display(const FSurface& CarrierSurface);
+
+CADKERNEL_API void Display(const FGroup& Group);
+CADKERNEL_API void Display(const FModel& Model);
+CADKERNEL_API void Display(const FBody& Body);
+CADKERNEL_API void Display(const FShell& Shell);
+CADKERNEL_API void Display(const FTopologicalEdge& Edge, EVisuProperty Property = EVisuProperty::BlueCurve);
+CADKERNEL_API void Display(const FTopologicalFace& Face);
+CADKERNEL_API void Display(const FTopologicalLoop& Loop);
+CADKERNEL_API void Display(const FTopologicalVertex& Vertex, EVisuProperty Property = EVisuProperty::BluePoint);
+
+CADKERNEL_API void Display2D(const FTopologicalEdge& Edge, EVisuProperty Property = EVisuProperty::BlueCurve);
+CADKERNEL_API void Display2D(const FTopologicalFace& Face);
+CADKERNEL_API void Display2D(const FTopologicalLoop& Loop);
+CADKERNEL_API void Display2D(const FSurface& CarrierSurface);
+
+CADKERNEL_API void Display2DWithScale(const FTopologicalEdge& Edge, EVisuProperty Property = EVisuProperty::BlueCurve);
+CADKERNEL_API void Display2DWithScale(const FTopologicalFace& Face);
+CADKERNEL_API void Display2DWithScale(const FTopologicalLoop& Loop);
+CADKERNEL_API void Display2DWithScale(const FSurface& CarrierSurface);
+
+CADKERNEL_API void DisplayMesh(const FFaceMesh& Mesh);
+CADKERNEL_API void DisplayMesh(const FEdgeMesh& Mesh);
+CADKERNEL_API void DisplayMesh(const FVertexMesh& Mesh);
+
+CADKERNEL_API void DisplayMesh(const FModel& Model);
+CADKERNEL_API void DisplayMesh(const FBody& Body);
+CADKERNEL_API void DisplayMesh(const FShell &Shel);
+
+CADKERNEL_API void Display(const FModelMesh& MeshModel);
+
+CADKERNEL_API void DisplayControlPolygon(const FCurve& Entity);
+CADKERNEL_API void DisplayControlPolygon(const FSurface& Entity);
+
+template<typename TPoint>
+void DisplaySegment(const TPoint& Point1, const TPoint& Point2, FIdent Ident = 0, EVisuProperty Property = EVisuProperty::BlueCurve, bool bWithOrientation = false)
+{
+#ifdef CADKERNEL_DEBUG
+	F3DDebugSegment G(Ident);
+	if (bWithOrientation)
+	{
+		DrawSegmentOrientation(Point1, Point2, Property);
+	}
+	DrawSegment(Point1, Point2, Property);
+ #endif
+};
+
+template<typename TPoint>
+void DisplaySegmentWithScale(const TPoint& Point1, const TPoint& Point2, FIdent Ident = 0, EVisuProperty Property = EVisuProperty::BlueCurve, bool bWithOrientation = false)
+{
+#ifdef CADKERNEL_DEBUG
+	F3DDebugSegment G(Ident);
+	if (bWithOrientation)
+	{
+		DrawSegmentOrientation(Point1 * DisplayScale, Point2 * DisplayScale, Property);
+	}
+	DrawSegment(Point1 * DisplayScale, Point2 * DisplayScale, Property);
+#endif
+};
+
+template<typename TPoint>
+void DisplayPolyline(const TArray<TPoint>& Points, EVisuProperty Property)
+{
+#ifdef CADKERNEL_DEBUG
+	Open3DDebugSegment(0);
+	Draw(Points, Property);
+	Close3DDebugSegment();
+#endif
+}
+
+static void DisplayPolylineWithScale(const TArray<FPoint2D>& Points, EVisuProperty Property)
+{
+#ifdef CADKERNEL_DEBUG
+	TArray<FPoint2D> PointsWithScale = Points;
+	for (FPoint2D& Point : PointsWithScale)
+	{
+		Point *= DisplayScale;
+	}
+
+	Open3DDebugSegment(0);
+	Draw(PointsWithScale, Property);
+	Close3DDebugSegment();
+#endif
+}
+
+template<typename TPoint>
+void DisplayOrientedPolyline(const TArray<TPoint>& Points, EVisuProperty Property)
+{
+#ifdef CADKERNEL_DEBUG
+	Open3DDebugSegment(0);
+	for (int32 Index = 1; Index < Points.Num(); ++Index)
+	{
+		DisplaySegment(Points[Index - 1], Points[Index], Index, Property, true);
+	}
+	Close3DDebugSegment();
+#endif
+}
+
+CADKERNEL_API void DrawQuadripode(double Height, double Base, FPoint& Centre, FPoint& Direction, EVisuProperty Property);
+
+CADKERNEL_API void Draw(const FTopologicalEdge& Edge, EVisuProperty Property = EVisuProperty::BlueCurve);
+CADKERNEL_API void Draw(const FTopologicalFace& Face);
+CADKERNEL_API void Draw2D(const FTopologicalFace& Face);
+CADKERNEL_API void Draw2DWithScale(const FTopologicalFace& Face);
+CADKERNEL_API void Draw(const FShell& Shell);
+
+CADKERNEL_API void Draw(const FCurve& Curve, EVisuProperty Property = EVisuProperty::BlueCurve);
+CADKERNEL_API void Draw(const FCurve& Curve, const FLinearBoundary& Boundary, EVisuProperty Property = EVisuProperty::BlueCurve);
+CADKERNEL_API void Draw(const FLinearBoundary& Boundary, const FRestrictionCurve& Curve, EVisuProperty Property = EVisuProperty::BlueCurve);
+
+template<typename TPoint>
+void DrawSegment(const TPoint& Point1, const TPoint& Point2, EVisuProperty Property = EVisuProperty::Element)
+{
+#ifdef CADKERNEL_DEBUG
+	TArray<FPoint> Points;
+	Points.Add(Point1);
+	Points.Add(Point2);
+	Draw(Points, Property);
+#endif
+}
+
+template<typename TPoint>
+void DrawSegmentOrientation(const TPoint& Point1, const TPoint& Point2, EVisuProperty Property = EVisuProperty::Element)
+{
+#ifdef CADKERNEL_DEBUG
+	double Length = Point1.Distance(Point2);
+	double Height = Length / 10.0;
+	double Base = Height / 2;
+
+	FPoint Middle = (Point1 + Point2) / 2.;
+	FPoint Tangent = Point2 - Point1;
+	DrawQuadripode(Height, Base, Middle, Tangent, Property);
+#endif
+}
+
+CADKERNEL_API void DrawIsoCurves(const FTopologicalFace& Face);
+
+
+} // namespace UE::CADKernel
+
