Moved CADKernel library code from /Engine/Source/Runtime/Datasmith/CADKernel to /Engine/Source/Runtime/Datasmith/CADKernel/Base

diff --git a/Engine/Source/Runtime/Datasmith/CADKernel/Base/Public/UI/Progress.h b/Engine/Source/Runtime/Datasmith/CADKernel/Base/Public/UI/Progress.h
new file mode 100644
--- /dev/null
+++ b/Engine/Source/Runtime/Datasmith/CADKernel/Base/Public/UI/Progress.h
@@ -0,0 +1,122 @@
+// Copyright Epic Games, Inc. All Rights Reserved.
+#pragma once
+
+#include "Core/Types.h"
+
+namespace UE::CADKernel
+{
+class FProgress;
+class FProgressBase;
+
+class CADKERNEL_API FProgressManager
+{
+	friend class FProgress;
+
+protected:
+
+	FProgress* RootProgress;
+	FProgress* CurrentProgress;
+
+	FProgress* GetRoot() const
+	{
+		return RootProgress;
+	}
+
+	virtual void ActivateProgressBar(bool bActive)
+	{}
+
+	/**
+	 * Update the progress bar
+	 */
+	virtual void Update()
+	{}
+
+	double GetProgression() const;
+
+	/**
+	 * @return a string describing the current step hierarchy
+	 */
+	virtual FString GetCurrentStep() const;
+
+	virtual void SetCurrentProgress(FProgress* InProgress);
+
+	FProgress* GetCurrentProgress() const
+	{
+		return CurrentProgress;
+	}
+
+public:
+	FProgressManager()
+		: RootProgress(nullptr)
+		, CurrentProgress(nullptr)
+	{}
+};
+
+class CADKERNEL_API FProgress
+{
+	friend class FProgressManager;
+private:
+	FString Name;
+
+	FProgress* Parent = nullptr;
+	FProgress* UnderlyingProgress = nullptr;
+
+	int32 StepCount;
+	int32 Progression = 0;
+
+public:
+	FProgress(int32 InStepCount, const FString& InStepName = FString());
+
+	FProgress(const FString& StepName = FString())
+		: FProgress(1, StepName)
+	{
+	}
+
+	virtual ~FProgress();
+
+	const FString& GetName()
+	{
+		return Name;
+	}
+
+	FProgress* GetParent() const
+	{
+		return Parent;
+	}
+
+	/**
+	 * Increase the advancement of of one step
+	 */
+	void Increase(int32 StepSize = 1);
+
+protected:
+
+	bool IsRoot() const
+	{
+		return Parent == nullptr;
+	}
+
+	void AddUnderlying(FProgress* Progress)
+	{
+		UnderlyingProgress = Progress;
+	}
+
+	void UnderlyingFinished(FProgress* Progress)
+	{
+		Progression++;
+		UnderlyingProgress = nullptr;
+	}
+
+	double GetProgression() const
+	{
+		double Ratio = ((double)Progression / (double)StepCount);
+
+		if (UnderlyingProgress)
+		{
+			Ratio += UnderlyingProgress->GetProgression() / (double)StepCount;
+		}
+		return Ratio;
+	}
+};
+}
+
diff --git a/Engine/Source/Runtime/Datasmith/CADKernel/Base/Public/UI/Visu.h b/Engine/Source/Runtime/Datasmith/CADKernel/Base/Public/UI/Visu.h
new file mode 100644
--- /dev/null
+++ b/Engine/Source/Runtime/Datasmith/CADKernel/Base/Public/UI/Visu.h
@@ -0,0 +1,177 @@
+// Copyright Epic Games, Inc. All Rights Reserved.
+#pragma once
+
+#ifdef CADKERNEL_DEV
+#include "Core/Parameter.h"
+#include "Core/Parameters.h"
+#endif
+
+#include "Core/Types.h"
+#include "Math/Point.h"
+
+namespace UE::CADKernel
+{
+enum EVisuProperty
+{
+	BBox = 0,          // Green see CADKernelUI::properties    
+	Point,             // Blue
+	Element,           // Brown
+	Curve,             // Blue
+	Iso,               // Red
+	ControlLine,       // Purple
+	EdgeMesh,          // blue
+	NodeMesh,          // blue
+	BorderEdge,        // Green
+	NonManifoldEdge,   // Red          
+	BorderTopo,        // Dark brown
+	ControlPoint,      // Purple		 
+	YellowPoint,
+	YellowCurve,
+	BluePoint,
+	BlueCurve,
+	RedPoint,
+	RedCurve,
+	PurplePoint,
+	PurpleCurve,
+	GreenPoint,
+	GreenCurve,
+	PinkPoint,
+	PinkCurve,
+	OrangePoint,
+	OrangeCurve,
+	Last,
+};
+
+extern const TCHAR* VisuPropertyNames[];
+
+class FMesh;
+
+#ifdef CADKERNEL_DEV
+
+class CADKERNEL_API FVisuParameters : public FParameters
+{
+public:
+
+	FParameter ChordError;
+
+	/**
+	 *	Number of isos along U to display on a surface or topological face
+	 */
+	FParameter IsoUNumber;
+
+	/**
+	 *	Number of isos along V to display on a surface or topological face
+	 */
+	FParameter IsoVNumber;
+
+	FParameter bDisplayCADOrient;
+	FParameter bDisplayMeshOrient;
+	FParameter bDisplayAxis;
+	FParameter bDisplayNormals;
+
+	FParameter NormalLength;
+
+	FVisuParameters()
+		: FParameters(10)
+		, ChordError(TEXT("DiscretizationError"), 0.02, *this)
+		, IsoUNumber(TEXT("IsoUNumber"), 3, *this)
+		, IsoVNumber(TEXT("IsoVNumber"), 3, *this)
+		, bDisplayCADOrient(TEXT("DisplayCadOrient"), false, *this)
+		, bDisplayMeshOrient(TEXT("DisplayMeshOrient"), true, *this)
+		, bDisplayAxis(TEXT("DisplayAxis"), false, *this)
+		, bDisplayNormals(TEXT("DisplayNormals"), false, *this)
+		, NormalLength(TEXT("NormalLength"), 10., *this)
+	{}
+};
+#endif
+
+class CADKERNEL_API FVisu
+{
+protected:
+#ifdef CADKERNEL_DEV
+	FVisuParameters Parameters;
+#endif
+
+	/**
+	 * Number of graphic session opened
+	 */
+	int32 SessionNum;
+
+	void NewSession()
+	{
+		ensureCADKernel(SessionNum >= 0);
+		++SessionNum;
+	}
+
+	void EndSession()
+	{
+		--SessionNum;
+		if (SessionNum < 0)
+		{
+			SessionNum = 0;
+		}
+	}
+
+public:
+
+	FVisu()
+	{
+		SessionNum = 0;
+	}
+
+	virtual ~FVisu() = default;
+
+	int32 GetPropertyCount()
+	{
+		return EVisuProperty::Last;
+	}
+
+#ifdef CADKERNEL_DEV
+	FVisuParameters* GetParameters()
+	{
+		return &Parameters;
+	}
+#endif
+
+	virtual void NewDB(const TCHAR* InName)
+	{}
+
+	virtual void Open3DDebugSession(const TCHAR* SessionName, const TArray<FIdent>& Ids)
+	{}
+
+	virtual void Close3DDebugSession()
+	{}
+
+	virtual void Open3DDebugSegment(FIdent Ident)
+	{}
+
+	virtual void Close3DDebugSegment()
+	{}
+
+	virtual void UpdateViewer()
+	{}
+
+	virtual void DrawPoint(const FPoint& Point, EVisuProperty InProperty = EVisuProperty::BluePoint)
+	{}
+
+	virtual void DrawPoint(const FPoint2D& Point, EVisuProperty InProperty = EVisuProperty::BluePoint)
+	{}
+
+	virtual void DrawPoint(const FPointH& Point, EVisuProperty InProperty = EVisuProperty::BluePoint)
+	{}
+
+	virtual void DrawElement(int32 Dimension, TArray<FPoint>& Points, EVisuProperty InProperty = EVisuProperty::Element)
+	{}
+
+	virtual void DrawPolyline(const TArray<FPoint>& Points, EVisuProperty InProperty = EVisuProperty::BlueCurve)
+	{}
+
+	virtual void DrawPolyline(const TArray<FPoint2D>& Points, EVisuProperty InProperty = EVisuProperty::BlueCurve)
+	{}
+
+	virtual void DrawMesh(FIdent MeshId)
+	{}
+};
+
+} // namespace UE::CADKernel
+
