Removed ensure causing pogram to crash

#rnx

diff --git a/Engine/Source/Runtime/Datasmith/CADKernel/Base/Public/Geo/Curves/RestrictionCurve.h b/Engine/Source/Runtime/Datasmith/CADKernel/Base/Public/Geo/Curves/RestrictionCurve.h
--- a/Engine/Source/Runtime/Datasmith/CADKernel/Base/Public/Geo/Curves/RestrictionCurve.h
+++ b/Engine/Source/Runtime/Datasmith/CADKernel/Base/Public/Geo/Curves/RestrictionCurve.h
@@ -69,6 +69,11 @@ public:
 		return ECurve::Restriction;
 	}
 
+	const FSurfacicPolyline& GetPolyline() const
+	{
+		return Polyline;
+	}
+
 	const TSharedRef<FCurve> Get2DCurve() const
 	{
 		return Curve2D.ToSharedRef();
@@ -169,7 +174,7 @@ public:
 	 */
 	void CheckIfDegenerated(const FLinearBoundary& InBoundary, bool& bDegeneration2D, bool& bDegeneration3D, double& Length3D) const
 	{
-		if (!ensure(Polyline.Coordinates.Num() > 1))
+		if (Polyline.Coordinates.Num() <  2)
 		{
 			bDegeneration2D = true;
 			bDegeneration3D = true;
