Renamed variable to prevent hiding globally declared LogLevel name

#rnx

diff --git a/Engine/Source/Runtime/Datasmith/CADKernel/Base/Private/UI/Message.cpp b/Engine/Source/Runtime/Datasmith/CADKernel/Base/Private/UI/Message.cpp
--- a/Engine/Source/Runtime/Datasmith/CADKernel/Base/Private/UI/Message.cpp
+++ b/Engine/Source/Runtime/Datasmith/CADKernel/Base/Private/UI/Message.cpp
@@ -34,7 +34,7 @@ void FMessage::VPrintf(EVerboseLevel Level, const TCHAR* Text, ...)
 {
 	EVerboseLevel VerboseLevel = (EVerboseLevel)(int32)FSystem::Get().GetVerboseLevel();
 	TSharedPtr<FArchive> LogFile = FSystem::Get().GetLogFile();
-	EVerboseLevel LogLevel = FSystem::Get().GetLogLevel();
+	EVerboseLevel SystemLogLevel = FSystem::Get().GetLogLevel();
 	TSharedPtr<FArchive> SpyFile = FSystem::Get().GetSpyFile();
 
 	FString Indentation;
@@ -67,11 +67,11 @@ void FMessage::VPrintf(EVerboseLevel Level, const TCHAR* Text, ...)
 		FSystem::Get().GetConsole().Print(Buffer, Level);
 	}
 
-	if ((LogFile.IsValid() && Level <= LogLevel) || (SpyFile.IsValid() && Level <= Spy))
+	if ((LogFile.IsValid() && Level <= SystemLogLevel) || (SpyFile.IsValid() && Level <= Spy))
 	{
 		FTCHARToUTF8 UTF8Indentation(*Indentation, Indentation.Len());
 		FTCHARToUTF8 UTF8String(Buffer, Result);
-		if (LogFile.IsValid() && Level <= LogLevel)
+		if (LogFile.IsValid() && Level <= SystemLogLevel)
 		{
 			LogFile->Serialize((UTF8CHAR*)UTF8Indentation.Get(), UTF8Indentation.Length() * sizeof(UTF8CHAR));
 			LogFile->Serialize((UTF8CHAR*)UTF8String.Get(), UTF8String.Length() * sizeof(UTF8CHAR));
