Fixed issue with wrong trimming on surfaces with narrow knotspaces

#rnx

diff --git a/Engine/Source/Runtime/Datasmith/CADKernel/Base/Private/Geo/Surfaces/NURBSSurface.h b/Engine/Source/Runtime/Datasmith/CADKernel/Base/Private/Geo/Surfaces/NURBSSurface.h
--- a/Engine/Source/Runtime/Datasmith/CADKernel/Base/Private/Geo/Surfaces/NURBSSurface.h
+++ b/Engine/Source/Runtime/Datasmith/CADKernel/Base/Private/Geo/Surfaces/NURBSSurface.h
@@ -111,9 +111,12 @@ public:
 
 	virtual void ValidateUVPoints(TArray<FPoint2D>& UVPoints) const
 	{
-		auto ValidateValue = [](double& Value, const double& Reference)
+		const double ToleranceU = GetIsoTolerance(EIso::IsoU);
+		const double ToleranceV = GetIsoTolerance(EIso::IsoV);
+
+		auto ValidateValue = [](double& Value, const double& Reference, const double& Tolerance)
 			{
-				if (FMath::IsNearlyEqual(Value, Reference, UE_DOUBLE_KINDA_SMALL_NUMBER))
+				if (FMath::IsNearlyEqual(Value, Reference, Tolerance))
 				{
 					Value = Reference;
 					return true;
@@ -126,7 +129,7 @@ public:
 		{
 			for (const double& NodalValue : UNodalVector)
 			{
-				if (ValidateValue(UVPoint.U, NodalValue))
+				if (ValidateValue(UVPoint.U, NodalValue, ToleranceU))
 				{
 					break;
 				}
@@ -134,7 +137,7 @@ public:
 
 			for (const double& NodalValue : VNodalVector)
 			{
-				if (ValidateValue(UVPoint.V, NodalValue))
+				if (ValidateValue(UVPoint.V, NodalValue, ToleranceV))
 				{
 					break;
 				}
diff --git a/Engine/Source/Runtime/Datasmith/CADKernel/Base/Private/Math/Geometry.cpp b/Engine/Source/Runtime/Datasmith/CADKernel/Base/Private/Math/Geometry.cpp
--- a/Engine/Source/Runtime/Datasmith/CADKernel/Base/Private/Math/Geometry.cpp
+++ b/Engine/Source/Runtime/Datasmith/CADKernel/Base/Private/Math/Geometry.cpp
@@ -145,6 +145,7 @@ void FindLoopIntersectionsWithIso(const EIso Iso, const double IsoParameter, con
 
 	if (LocalIntersections.Num() == 0)
 	{
+		OutIntersections.Empty();
 		return;
 	}
 
