Fixed issue with tessellation of planar surfaces
Added CVar "CADKernel.FaceMesher.DetectPlanarFace" to enable/disable specific tessellation for planar faces
Added GeometryCore module to use FDelaunay class to tessellate planar surfaces
Added utility functions for Curve, Surface and Face entities
Replaced use of NURBS curves of degree 1 with use of polylines
Made necessary fixes on code related to polylines
Fixed miscellanous inconsistencies in the API
Cleanup code - WIP

#rnx

diff --git a/Engine/Source/Runtime/Datasmith/CADKernel/Base/CADKernel.Build.cs b/Engine/Source/Runtime/Datasmith/CADKernel/Base/CADKernel.Build.cs
--- a/Engine/Source/Runtime/Datasmith/CADKernel/Base/CADKernel.Build.cs
+++ b/Engine/Source/Runtime/Datasmith/CADKernel/Base/CADKernel.Build.cs
@@ -17,8 +17,11 @@ namespace UnrealBuildTool.Rules
 				new string[]
 				{
 					"Core",
+					"GeometryCore",
 				}
 			);
+
+			PublicDefinitions.Add(Target.Type == TargetType.Program ? "CADKERNEL_DO_ENSURE=0" : "CADKERNEL_DO_ENSURE=1");
 		}
 	}
 }
\ No newline at end of file
diff --git a/Engine/Source/Runtime/Datasmith/CADKernel/Base/Private/Core/Database.cpp b/Engine/Source/Runtime/Datasmith/CADKernel/Base/Private/Core/Database.cpp
--- a/Engine/Source/Runtime/Datasmith/CADKernel/Base/Private/Core/Database.cpp
+++ b/Engine/Source/Runtime/Datasmith/CADKernel/Base/Private/Core/Database.cpp
@@ -32,14 +32,14 @@ FDatabase::FDatabase()
 	AddEntity(Model.ToSharedRef());
 }
 
-FModel& FDatabase::GetModel()
+TSharedPtr<FModel> FDatabase::GetModelAsShared()
 {
 	if (!Model.IsValid())
 	{
 		Model = FEntity::MakeShared<FModel>();
 		AddEntity(Model.ToSharedRef());
 	}
-	return *Model;
+	return Model;
 }
 
 FIdent FDatabase::CreateId()
@@ -236,7 +236,7 @@ void FDatabase::SerializeSelection(FCADKernelArchive& Ar, const TArray<FIdent>&
 {
 	TRACE_CPUPROFILER_EVENT_SCOPE(FDatabase::Serialize);
 
-	ensure(Ar.IsSaving());
+	ensureCADKernel(Ar.IsSaving());
 
 	bIsRecursiveSerialization = true;
 
@@ -304,7 +304,7 @@ void FDatabase::Serialize(FCADKernelArchive& Ar)
 {
 	TRACE_CPUPROFILER_EVENT_SCOPE(FDatabase::Serialize);
 
-	ensure(Ar.IsSaving());
+	ensureCADKernel(Ar.IsSaving());
 
 	SpawnEntityIdent(GetModel(), true);
 
@@ -338,7 +338,7 @@ void FDatabase::Deserialize(FCADKernelArchive& Ar)
 {
 	TRACE_CPUPROFILER_EVENT_SCOPE(FDatabase::Deserialize);
 
-	ensure(Ar.IsLoading());
+	ensureCADKernel(Ar.IsLoading());
 
 	Ar.Session.Serialize(Ar);
 
@@ -397,7 +397,7 @@ void FDatabase::CleanArchiveEntities()
 				TSharedPtr<FVertexLink> VertexLink = StaticCastSharedPtr<FVertexLink>(Entity);
 				VertexLink->CleanLink();
 #ifdef CADKERNEL_DEV
- 				ensureCADKernel(VertexLink->GetTwinEntityNum());
+ 				ensureCADKernel(VertexLink->GetTwinEntityCount());
 #endif
 				break;
 			}
diff --git a/Engine/Source/Runtime/Datasmith/CADKernel/Base/Private/Core/Session.cpp b/Engine/Source/Runtime/Datasmith/CADKernel/Base/Private/Core/Session.cpp
--- a/Engine/Source/Runtime/Datasmith/CADKernel/Base/Private/Core/Session.cpp
+++ b/Engine/Source/Runtime/Datasmith/CADKernel/Base/Private/Core/Session.cpp
@@ -23,9 +23,9 @@ void FSession::SaveDatabase(const TCHAR* FileName)
 	Archive->Close();
 }
 
-FModel& FSession::GetModel()
+TSharedPtr<FModel> FSession::GetModelAsShared()
 {
-	return Database.GetModel();
+	return Database.GetModelAsShared();
 }
 
 void FSession::SaveDatabase(const TCHAR* FileName, const TArray<FEntity*>& SelectedEntities)
@@ -66,6 +66,7 @@ void FSession::SaveDatabase(const TCHAR* FileName, const TArray<TSharedPtr<FEnti
 
 void FSession::LoadDatabase(const TCHAR* FilePath)
 {
+	TGuardValue<double> GeometricToleranceGuard(GeometricTolerance, 0.01);
 	TSharedPtr<FCADKernelArchive> Archive = FCADKernelArchive::CreateArchiveReader(*this, FilePath);
 	if (!Archive.IsValid())
 	{
@@ -88,13 +89,14 @@ void FSession::LoadDatabase(const TCHAR* FilePath)
 
 void FSession::AddDatabase(const TArray<uint8>& InRawData)
 {
+	TGuardValue<double> GeometricToleranceGuard(GeometricTolerance, 0.01);
 	FCADKernelArchive Archive = FCADKernelArchive(*this, InRawData);
 	Database.Deserialize(Archive);
 }
 
 void FSession::SetGeometricTolerance(double NewTolerance)
 {
-	ensure(Database.GetModel().EntityCount() == 0);
+	ensureCADKernel(Database.GetModel().EntityCount() == 0);
 	GeometricTolerance = NewTolerance;
 	IntersectionTool::SetTolerance(NewTolerance);
 }
diff --git a/Engine/Source/Runtime/Datasmith/CADKernel/Base/Private/Geo/Curves/BezierCurve.cpp b/Engine/Source/Runtime/Datasmith/CADKernel/Base/Private/Geo/Curves/BezierCurve.cpp
--- a/Engine/Source/Runtime/Datasmith/CADKernel/Base/Private/Geo/Curves/BezierCurve.cpp
+++ b/Engine/Source/Runtime/Datasmith/CADKernel/Base/Private/Geo/Curves/BezierCurve.cpp
@@ -37,7 +37,7 @@ bool FBezierCurve::IsBezier(const FNurbsCurveData& NurbsCurveData)
 
 FBezierCurve::FBezierCurve(const FNurbsCurveData& NurbsCurveData)
 {
-	ensure(IsBezier(NurbsCurveData));
+	ensureCADKernel(IsBezier(NurbsCurveData));
 
 	Degree = NurbsCurveData.Degree;
 	Dimension = (int8)NurbsCurveData.Dimension;
@@ -91,12 +91,12 @@ FBezierCurve::FBezierCurve(const FNurbsCurveData& NurbsCurveData)
 void FBezierCurve::EvaluatePoint(double Coordinate, FCurvePoint& OutPoint, int32 DerivativeOrder) const
 {
 	// Adjust coordinate value to curve's boundaries
-	if (!ensure(Coordinate >= (NodalVector[0] - UE_DOUBLE_SMALL_NUMBER)))
+	if (!ensureCADKernel(Coordinate >= (NodalVector[0] - UE_DOUBLE_SMALL_NUMBER)))
 	{
 		Coordinate = NodalVector[0];
 	}
 
-	if (!ensure(Coordinate <= (NodalVector[NumSegments] + UE_DOUBLE_SMALL_NUMBER)))
+	if (!ensureCADKernel(Coordinate <= (NodalVector[NumSegments] + UE_DOUBLE_SMALL_NUMBER)))
 	{
 		Coordinate = NodalVector[NumSegments];
 	}
@@ -117,7 +117,7 @@ void FBezierCurve::EvaluatePoint(double Coordinate, FCurvePoint& OutPoint, int32
 			}
 		}
 	}
-	ensure(SegmentIndex < NumSegments);
+	ensureCADKernel(SegmentIndex < NumSegments);
 
 	OutPoint.DerivativeOrder = DerivativeOrder;
 	OutPoint.Init();
diff --git a/Engine/Source/Runtime/Datasmith/CADKernel/Base/Private/Geo/Curves/BezierCurve.h b/Engine/Source/Runtime/Datasmith/CADKernel/Base/Private/Geo/Curves/BezierCurve.h
--- a/Engine/Source/Runtime/Datasmith/CADKernel/Base/Private/Geo/Curves/BezierCurve.h
+++ b/Engine/Source/Runtime/Datasmith/CADKernel/Base/Private/Geo/Curves/BezierCurve.h
@@ -66,7 +66,7 @@ public:
 
 	int32 GetDegre() const
 	{
-		return Poles.Num() - 1;
+		return Degree;
 	}
 
 	const TArray<FPoint>& GetPoles() const
diff --git a/Engine/Source/Runtime/Datasmith/CADKernel/Base/Private/Geo/Curves/BoundedCurve.cpp b/Engine/Source/Runtime/Datasmith/CADKernel/Base/Private/Geo/Curves/BoundedCurve.cpp
--- a/Engine/Source/Runtime/Datasmith/CADKernel/Base/Private/Geo/Curves/BoundedCurve.cpp
+++ b/Engine/Source/Runtime/Datasmith/CADKernel/Base/Private/Geo/Curves/BoundedCurve.cpp
@@ -6,13 +6,13 @@ namespace UE::CADKernel
 
 void FBoundedCurve::EvaluatePoint(double Coordinate, FCurvePoint& OutPoint, int32 DerivativeOrder) const
 {
-	ensure(Boundary.Contains(Coordinate));
+	ensureCADKernel(Boundary.Contains(Coordinate));
 	return Curve->EvaluatePoint(Coordinate, OutPoint, DerivativeOrder);
 }
 
 void FBoundedCurve::Evaluate2DPoint(double Coordinate, FCurvePoint2D& OutPoint, int32 DerivativeOrder) const
 {
-	ensure(Boundary.Contains(Coordinate));
+	ensureCADKernel(Boundary.Contains(Coordinate));
 	return Curve->Evaluate2DPoint(Coordinate, OutPoint, DerivativeOrder);
 }
 
diff --git a/Engine/Source/Runtime/Datasmith/CADKernel/Base/Private/Geo/Curves/CompositeCurve.cpp b/Engine/Source/Runtime/Datasmith/CADKernel/Base/Private/Geo/Curves/CompositeCurve.cpp
--- a/Engine/Source/Runtime/Datasmith/CADKernel/Base/Private/Geo/Curves/CompositeCurve.cpp
+++ b/Engine/Source/Runtime/Datasmith/CADKernel/Base/Private/Geo/Curves/CompositeCurve.cpp
@@ -55,7 +55,7 @@ double FCompositeCurve::GlobalToLocalCoordinate(int32 CurveIndex, double Coordin
 
 void FCompositeCurve::EvaluatePoint(double Coordinate, FCurvePoint& OutPoint, int32 DerivativeOrder) const
 {
-	ensure(Dimension == 3);
+	ensureCADKernel(Dimension == 3);
 
 	OutPoint.DerivativeOrder = DerivativeOrder;
 
@@ -88,7 +88,7 @@ void FCompositeCurve::EvaluatePoint(double Coordinate, FCurvePoint& OutPoint, in
 
 void FCompositeCurve::Evaluate2DPoint(double Coordinate, FCurvePoint2D& OutPoint, int32 DerivativeOrder) const
 {
-	ensure(Dimension == 2);
+	ensureCADKernel(Dimension == 2);
 
 	OutPoint.DerivativeOrder = DerivativeOrder;
 
diff --git a/Engine/Source/Runtime/Datasmith/CADKernel/Base/Private/Geo/Curves/Curve.cpp b/Engine/Source/Runtime/Datasmith/CADKernel/Base/Private/Geo/Curves/Curve.cpp
--- a/Engine/Source/Runtime/Datasmith/CADKernel/Base/Private/Geo/Curves/Curve.cpp
+++ b/Engine/Source/Runtime/Datasmith/CADKernel/Base/Private/Geo/Curves/Curve.cpp
@@ -74,22 +74,6 @@ void FCurve::Evaluate2DPoints(const TArray<double>& Coordinates, TArray<FPoint2D
 	}
 }
 
-#ifdef CADKERNEL_DEV
-FInfoEntity& FCurve::GetInfo(FInfoEntity& Info) const
-{
-	return FEntity::GetInfo(Info)
-		.Add(TEXT("Curve type"), CurvesTypesNames[(uint8)GetCurveType()])
-		.Add(TEXT("Dimension"), (int32)Dimension)
-		.Add(TEXT("Boundary"), Boundary)
-		.Add(TEXT("Length"), GetLength(0.01));
-}
-#endif
-
-void FCurve::FindNotDerivableCoordinates(int32 DerivativeOrder, TArray<double>& OutNotDerivableCoordinates) const
-{
-	FindNotDerivableCoordinates(Boundary, DerivativeOrder, OutNotDerivableCoordinates);
-}
-
 void FCurve::FindNotDerivableCoordinates(const FLinearBoundary& InBoundary, int32 DerivativeOrder, TArray<double>& OutNotDerivableCoordinates) const
 {
 }
diff --git a/Engine/Source/Runtime/Datasmith/CADKernel/Base/Private/Geo/Curves/CurveMaker.cpp b/Engine/Source/Runtime/Datasmith/CADKernel/Base/Private/Geo/Curves/CurveMaker.cpp
--- a/Engine/Source/Runtime/Datasmith/CADKernel/Base/Private/Geo/Curves/CurveMaker.cpp
+++ b/Engine/Source/Runtime/Datasmith/CADKernel/Base/Private/Geo/Curves/CurveMaker.cpp
@@ -3,6 +3,7 @@
 #include "Geo/Curves/BezierCurve.h"
 #include "Geo/Curves/Curve.h"
 #include "Geo/Curves/NURBSCurve.h"
+#include "Geo/Curves/PolylineCurve.h"
 #include "Geo/Curves/SplineCurve.h"
 
 namespace UE::CADKernel
@@ -10,11 +11,75 @@ namespace UE::CADKernel
 
 TSharedPtr<FCurve> FCurve::MakeNurbsCurve(FNurbsCurveData& InNurbsData)
 {
-	if (FBezierCurve::IsBezier(InNurbsData))
+	if (InNurbsData.Degree == 1)
 	{
-		return FEntity::MakeShared<UE::CADKernel::FBezierCurve>(InNurbsData);
+		ensureCADKernel(InNurbsData.Dimension > 1);
+		TArray<double> Coordinates(InNurbsData.NodalVector.GetData() + 1, InNurbsData.NodalVector.Num() - 2);
+		ensureCADKernel(Coordinates.Num() > 1);
+		
+		// Make sure Coordinates contains rising or equal values
+		if (Coordinates.Num() == 2)
+		{
+			if (!ensureCADKernel(Coordinates[0] < Coordinates[1]))
+			{
+				Coordinates[1] = Coordinates[0];
+			}
+		}
+		else
+		{
+			double LastValue = Coordinates[0];
+			for (int32 Index = 1; Index < Coordinates.Num(); ++Index)
+			{
+				if (Coordinates[Index] < LastValue)
+				{
+					Coordinates[Index] = LastValue;
+				}
+				else
+				{
+					LastValue = Coordinates[Index];
+				}
+			}
+		}
+
+		if (InNurbsData.Dimension == 2)
+		{
+			TArray<FPoint2D> NewPoles;
+			NewPoles.Reserve(InNurbsData.Poles.Num());
+
+			if (InNurbsData.bIsRational && InNurbsData.Poles.Num() == InNurbsData.Weights.Num())
+			{
+				for (int32 Index = 0; Index < InNurbsData.Poles.Num(); ++Index)
+				{
+					NewPoles.Emplace(InNurbsData.Poles[Index].X, InNurbsData.Poles[Index].Y);
+					NewPoles.Last() /= InNurbsData.Weights[Index];
+				}
+			}
+			else
+			{
+				for (const FPoint& Pole : InNurbsData.Poles)
+				{
+					NewPoles.Emplace(Pole.X, Pole.Y);
+				}
+			}
+
+			return FEntity::MakeShared<FPolyline2DCurve>(NewPoles, Coordinates);
+		}
+
+		if (InNurbsData.bIsRational && InNurbsData.Poles.Num() == InNurbsData.Weights.Num())
+		{
+			TArray<FPoint> NewPoles(InNurbsData.Poles);
+			for (int32 Index = 0; Index < NewPoles.Num(); ++Index)
+			{
+				NewPoles[Index] /= InNurbsData.Weights[Index];
+			}
+
+			return FEntity::MakeShared<FPolylineCurve>(NewPoles, Coordinates);
+		}
+
+		return FEntity::MakeShared<FPolylineCurve>(InNurbsData.Poles, Coordinates);
 	}
-	return FEntity::MakeShared<UE::CADKernel::FNURBSCurve>(InNurbsData);
+
+	return FEntity::MakeShared<FNURBSCurve>(InNurbsData);
 }
 
 TSharedPtr<FCurve> FCurve::MakeBezierCurve(const TArray<FPoint>& InPoles)
diff --git a/Engine/Source/Runtime/Datasmith/CADKernel/Base/Private/Geo/Curves/CurveUtilities.cpp b/Engine/Source/Runtime/Datasmith/CADKernel/Base/Private/Geo/Curves/CurveUtilities.cpp
new file mode 100644
--- /dev/null
+++ b/Engine/Source/Runtime/Datasmith/CADKernel/Base/Private/Geo/Curves/CurveUtilities.cpp
@@ -0,0 +1,278 @@
+// Copyright Epic Games, Inc. All Rights Reserved.
+
+#include "Geo/Curves/CurveUtilities.h"
+
+#include "Geo/Curves/BezierCurve.h"
+#include "Geo/Curves/CompositeCurve.h"
+#include "Geo/Curves/NURBSCurve.h"
+#include "Geo/Curves/PolylineCurve.h"
+#include "Geo/Curves/SegmentCurve.h"
+#include "Geo/GeoEnum.h"
+#include "Geo/Sampling/SurfacicPolyline.h"
+
+namespace UE::CADKernel::CurveUtilities
+{
+	template<class T>
+	bool GetCurvePoles(const T& Curve, TArray<FVector>& PolesOut)
+	{
+		const TArray<FPoint>& Poles = Curve.GetPoles();
+		const int32 PoleCount = Poles.Num();
+
+		PolesOut.Reserve(PoleCount);
+
+		if (Curve.GetDimension() == 2)
+		{
+			for (int32 Index = 0; Index < PoleCount; ++Index)
+			{
+				const FPoint& Pole = Poles[Index];
+				PolesOut.Add(FVector(Pole.X, Pole.Y, 0.f));
+			}
+		}
+		else
+		{
+			for (int32 Index = 0; Index < PoleCount; ++Index)
+			{
+				const FPoint& Pole = Poles[Index];
+				PolesOut.Add(FVector(Pole.X, Pole.Y, Pole.Z));
+			}
+		}
+
+		return PolesOut.Num() > 0;
+	}
+
+	template<class PointType, class PointCurveType>
+	bool GetPolylinePoints(const TPolylineCurve<PointType, PointCurveType>& Curve, TArray<FVector>& Poles)
+	{
+		const TArray<PointType>& Points = Curve.GetPolylinePoints();
+		const int32 PointCount = Points.Num();
+
+		Poles.Reserve(PointCount);
+
+		if (Curve.GetDimension() == 2)
+		{
+			for (int32 Index = 0; Index < PointCount; ++Index)
+			{
+				const PointType Point = Points[Index];
+				Poles.Add(FVector(Point[0], Point[1], 0.f));
+			}
+		}
+		else
+		{
+			for (int32 Index = 0; Index < PointCount; ++Index)
+			{
+				const PointType Point = Points[Index];
+				Poles.Add(FVector(Point[0], Point[1], Point[3]));
+			}
+		}
+
+		return Poles.Num() > 0;
+	}
+
+	TArray<FVector> GetPoles(const UE::CADKernel::FCurve& Curve)
+	{
+		using namespace UE::CADKernel;
+
+		TArray<FVector> Poles;
+
+		switch (Curve.GetCurveType())
+		{
+		case ECurve::Bezier:
+		{
+			GetCurvePoles(static_cast<const FBezierCurve&>(Curve), Poles);
+		}
+		break;
+
+		case ECurve::Nurbs:
+		{
+			GetCurvePoles(static_cast<const FNURBSCurve&>(Curve), Poles);
+		}
+		break;
+
+		case ECurve::Restriction:
+		{
+			return GetPoles(*static_cast<const FRestrictionCurve&>(Curve).Get2DCurve());
+		}
+
+		case ECurve::Segment:
+		{
+			const FSegmentCurve& Segment = static_cast<const FSegmentCurve&>(Curve);
+			Poles.Reserve(2);
+			{
+				const FPoint& Point = Segment.GetStartPoint();
+				Poles.Emplace(Point.X, Point.Y, Point.Z);
+			}
+			{
+				const FPoint& Point = Segment.GetEndPoint();
+				Poles.Emplace(Point.X, Point.Y, Point.Z);
+			}
+		}
+		break;
+
+		case ECurve::Polyline3D:
+		{
+			GetPolylinePoints<FPoint, FCurvePoint>(static_cast<const TPolylineCurve<FPoint, FCurvePoint>&>(Curve), Poles);
+		}
+
+		case ECurve::Polyline2D:
+		{
+			GetPolylinePoints<FPoint2D, FCurvePoint2D>(static_cast<const TPolylineCurve<FPoint2D, FCurvePoint2D>&>(Curve), Poles);
+		}
+		break;
+
+		default:
+			ensureCADKernel(false);
+			break;
+		}
+
+		return Poles;
+	}
+
+	TArray<FVector2d> GetPoles(const UE::CADKernel::FRestrictionCurve& Curve)
+	{
+		TArray<FVector> Poles = GetPoles(*Curve.Get2DCurve());
+
+		TArray<FVector2d> SurfacicPoints;
+		SurfacicPoints.Reserve(Poles.Num());
+
+		for (const FVector& Pole : Poles)
+		{
+			SurfacicPoints.Emplace(Pole.X, Pole.Y);
+		}
+
+		return SurfacicPoints;
+	}
+
+	TArray<FVector2d> Get2DPolyline(const UE::CADKernel::FRestrictionCurve& Curve)
+	{
+		using namespace UE::CADKernel;
+
+		const FSurfacicPolyline& Polyline = Curve.GetPolyline();
+
+		const TArray<FPoint2D>& Point2Ds = Polyline.Get2DPoints();
+
+		TArray<FVector2d> SurfacicPoints;
+		SurfacicPoints.Reserve(Point2Ds.Num());
+
+		for (const FPoint2D& Point2D : Point2Ds)
+		{
+			SurfacicPoints.Emplace(Point2D.U, Point2D.V);
+		}
+
+		return SurfacicPoints;
+	}
+
+	TArray<FVector> Get3DPolyline(const UE::CADKernel::FRestrictionCurve& Curve)
+	{
+		using namespace UE::CADKernel;
+
+		const FSurfacicPolyline& Polyline = Curve.GetPolyline();
+
+		const TArray<FPoint>& Point3Ds = Polyline.GetPoints();
+
+		TArray<FVector> Points;
+		Points.Reserve(Point3Ds.Num());
+
+		for (const FPoint& Point : Point3Ds)
+		{
+			Points.Emplace(Point.X, Point.Y, Point.Z);
+		}
+
+		return Points;
+	}
+
+	int32 GetPoleCount(const UE::CADKernel::FCurve& Curve)
+	{
+		using namespace UE::CADKernel;
+
+		switch (Curve.GetCurveType())
+		{
+		case ECurve::Bezier:
+		{
+			return static_cast<const FBezierCurve&>(Curve).GetPoles().Num();
+		}
+		break;
+
+		case ECurve::Nurbs:
+		{
+			return static_cast<const FNURBSCurve&>(Curve).GetPoles().Num();
+		}
+
+		case ECurve::Restriction:
+		{
+			return GetPoleCount(*static_cast<const FRestrictionCurve&>(Curve).Get2DCurve());
+		}
+
+		case ECurve::Segment:
+		{
+			return 2;
+		}
+
+		case ECurve::Polyline3D:
+		{
+			return static_cast<const TPolylineCurve<FPoint, FCurvePoint>&>(Curve).GetPolylinePoints().Num();
+		}
+
+		case ECurve::Polyline2D:
+		{
+			return static_cast<const TPolylineCurve<FPoint2D, FCurvePoint2D>&>(Curve).GetPolylinePoints().Num();
+		}
+
+		default:
+			ensureCADKernel(false);
+			break;
+		}
+
+		return 0;
+	}
+
+	int32 GetPoleCount(const UE::CADKernel::FRestrictionCurve& Curve)
+	{
+		return GetPoleCount(*Curve.Get2DCurve());
+	}
+
+	int32 GetDegree(const UE::CADKernel::FCurve& Curve)
+	{
+		using namespace UE::CADKernel;
+
+		TArray<FVector> Poles;
+
+		switch (Curve.GetCurveType())
+		{
+		case ECurve::Bezier:
+		{
+			return static_cast<const FBezierCurve&>(Curve).GetDegre();
+		}
+		break;
+
+		case ECurve::Nurbs:
+		{
+			return static_cast<const FNURBSCurve&>(Curve).GetDegree();
+		}
+		break;
+
+		case ECurve::Restriction:
+		{
+			return GetDegree(*static_cast<const FRestrictionCurve&>(Curve).Get2DCurve());
+		}
+
+		case ECurve::Polyline2D:
+		case ECurve::Polyline3D:
+		case ECurve::Segment:
+		{
+			return 1;
+		}
+		break;
+
+		default:
+			ensureCADKernel(false);
+			break;
+		}
+
+		return -1;
+	}
+
+	int32 GetDegree(const UE::CADKernel::FRestrictionCurve& Curve)
+	{
+		return GetDegree(*Curve.Get2DCurve());
+	}
+}
diff --git a/Engine/Source/Runtime/Datasmith/CADKernel/Base/Private/Geo/Curves/EllipseCurve.h b/Engine/Source/Runtime/Datasmith/CADKernel/Base/Private/Geo/Curves/EllipseCurve.h
--- a/Engine/Source/Runtime/Datasmith/CADKernel/Base/Private/Geo/Curves/EllipseCurve.h
+++ b/Engine/Source/Runtime/Datasmith/CADKernel/Base/Private/Geo/Curves/EllipseCurve.h
@@ -69,13 +69,13 @@ public:
 
 	virtual void EvaluatePoint(double Coordinate, FCurvePoint& OutPoint, int32 DerivativeOrder = 0) const override
 	{
-		ensure(Dimension == 3);
+		ensureCADKernel(Dimension == 3);
 		Evaluate<FCurvePoint, FPoint>(Coordinate, OutPoint, DerivativeOrder);
 	}
 
 	virtual void Evaluate2DPoint(double Coordinate, FCurvePoint2D& OutPoint, int32 DerivativeOrder = 0) const override
 	{
-		ensure(Dimension == 2);
+		ensureCADKernel(Dimension == 2);
 		Evaluate<FCurvePoint2D, FPoint2D>(Coordinate, OutPoint, DerivativeOrder);
 	}
 
diff --git a/Engine/Source/Runtime/Datasmith/CADKernel/Base/Private/Geo/Curves/HyperbolaCurve.h b/Engine/Source/Runtime/Datasmith/CADKernel/Base/Private/Geo/Curves/HyperbolaCurve.h
--- a/Engine/Source/Runtime/Datasmith/CADKernel/Base/Private/Geo/Curves/HyperbolaCurve.h
+++ b/Engine/Source/Runtime/Datasmith/CADKernel/Base/Private/Geo/Curves/HyperbolaCurve.h
@@ -56,13 +56,13 @@ public:
 
 	virtual void EvaluatePoint(double Coordinate, FCurvePoint& OutPoint, int32 DerivativeOrder = 0) const override
 	{
-		ensure(Dimension == 3);
+		ensureCADKernel(Dimension == 3);
 		Evaluate<FCurvePoint, FPoint>(Coordinate, OutPoint, DerivativeOrder);
 	}
 
 	virtual void Evaluate2DPoint(double Coordinate, FCurvePoint2D& OutPoint, int32 DerivativeOrder = 0) const override
 	{
-		ensure(Dimension == 2);
+		ensureCADKernel(Dimension == 2);
 		Evaluate<FCurvePoint2D, FPoint2D>(Coordinate, OutPoint, DerivativeOrder);
 	}
 
diff --git a/Engine/Source/Runtime/Datasmith/CADKernel/Base/Private/Geo/Curves/ParabolaCurve.h b/Engine/Source/Runtime/Datasmith/CADKernel/Base/Private/Geo/Curves/ParabolaCurve.h
--- a/Engine/Source/Runtime/Datasmith/CADKernel/Base/Private/Geo/Curves/ParabolaCurve.h
+++ b/Engine/Source/Runtime/Datasmith/CADKernel/Base/Private/Geo/Curves/ParabolaCurve.h
@@ -52,13 +52,13 @@ public:
 
 	virtual void EvaluatePoint(double Coordinate, FCurvePoint& OutPoint, int32 DerivativeOrder = 0) const override
 	{
-		ensure(Dimension == 3);
+		ensureCADKernel(Dimension == 3);
 		Evaluate<FCurvePoint, FPoint>(Coordinate, OutPoint, DerivativeOrder);
 	}
 
 	virtual void Evaluate2DPoint(double Coordinate, FCurvePoint2D& OutPoint, int32 DerivativeOrder = 0) const override
 	{
-		ensure(Dimension == 2);
+		ensureCADKernel(Dimension == 2);
 		Evaluate<FCurvePoint2D, FPoint2D>(Coordinate, OutPoint, DerivativeOrder);
 	}
 
diff --git a/Engine/Source/Runtime/Datasmith/CADKernel/Base/Private/Geo/Curves/PolylineCurve.h b/Engine/Source/Runtime/Datasmith/CADKernel/Base/Private/Geo/Curves/PolylineCurve.h
--- a/Engine/Source/Runtime/Datasmith/CADKernel/Base/Private/Geo/Curves/PolylineCurve.h
+++ b/Engine/Source/Runtime/Datasmith/CADKernel/Base/Private/Geo/Curves/PolylineCurve.h
@@ -93,8 +93,18 @@ public:
 
 	virtual void FindNotDerivableCoordinates(const FLinearBoundary& InBoundary, int32 DerivativeOrder, TArray<double>& OutNotDerivableCoordinates) const override
 	{
-		//TO DO
-		ensureCADKernel(false);
+		const int32 CoordinateCount = this->Coordinates.Num();
+		if (CoordinateCount > 2)
+		{
+			OutNotDerivableCoordinates.Reserve(CoordinateCount - 2);
+			int32 Index = 1;
+			for (; Index < CoordinateCount - 1 && this->Coordinates[Index] <= InBoundary.GetMin(); ++Index);
+
+			for (; Index < CoordinateCount - 1 && this->Coordinates[Index] <= InBoundary.GetMax(); ++Index)
+			{
+				OutNotDerivableCoordinates.Emplace(this->Coordinates[Index]);
+			}
+		}
 	}
 
 	void SetPoints(const TArray<PointType>& InPoints)
@@ -177,12 +187,12 @@ class CADKERNEL_API FPolyline2DCurve : public TPolylineCurve<FPoint2D, FCurvePoi
 
 protected:
 	FPolyline2DCurve(const TArray<FPoint2D>& InPoints, const TArray<double>& InCoordinates)
-		: TPolylineCurve<FPoint2D, FCurvePoint2D>(InPoints, InCoordinates, 3)
+		: TPolylineCurve<FPoint2D, FCurvePoint2D>(InPoints, InCoordinates, 2)
 	{
 	}
 
 	FPolyline2DCurve(const TArray<FPoint2D>& InPoints)
-		: TPolylineCurve<FPoint2D, FCurvePoint2D>(InPoints, 3)
+		: TPolylineCurve<FPoint2D, FCurvePoint2D>(InPoints, 2)
 	{
 	}
 
@@ -192,7 +202,7 @@ public:
 
 	virtual ECurve GetCurveType() const override
 	{
-		return ECurve::Polyline3D;
+		return ECurve::Polyline2D;
 	}
 
 	virtual TSharedPtr<FEntityGeom> ApplyMatrix(const FMatrixH& InMatrix) const override
diff --git a/Engine/Source/Runtime/Datasmith/CADKernel/Base/Private/Geo/Curves/SegmentCurve.h b/Engine/Source/Runtime/Datasmith/CADKernel/Base/Private/Geo/Curves/SegmentCurve.h
--- a/Engine/Source/Runtime/Datasmith/CADKernel/Base/Private/Geo/Curves/SegmentCurve.h
+++ b/Engine/Source/Runtime/Datasmith/CADKernel/Base/Private/Geo/Curves/SegmentCurve.h
@@ -47,13 +47,13 @@ public:
 
 	virtual void EvaluatePoint(double Coordinate, FCurvePoint& OutPoint, int32 DerivativeOrder = 0) const override
 	{
-		ensure(Dimension == 3);
+		ensureCADKernel(Dimension == 3);
 		Evaluate<FCurvePoint, FPoint>(Coordinate, OutPoint, DerivativeOrder);
 	}
 
 	virtual void Evaluate2DPoint(double Coordinate, FCurvePoint2D& OutPoint, int32 DerivativeOrder = 0) const override
 	{
-		ensure(Dimension == 2);
+		ensureCADKernel(Dimension == 2);
 		Evaluate<FCurvePoint2D, FPoint2D>(Coordinate, OutPoint, DerivativeOrder);
 	}
 
diff --git a/Engine/Source/Runtime/Datasmith/CADKernel/Base/Private/Geo/Curves/SplineCurve.cpp b/Engine/Source/Runtime/Datasmith/CADKernel/Base/Private/Geo/Curves/SplineCurve.cpp
--- a/Engine/Source/Runtime/Datasmith/CADKernel/Base/Private/Geo/Curves/SplineCurve.cpp
+++ b/Engine/Source/Runtime/Datasmith/CADKernel/Base/Private/Geo/Curves/SplineCurve.cpp
@@ -10,7 +10,7 @@ namespace UE::CADKernel
 
 TSharedPtr<FEntityGeom> FSplineCurve::ApplyMatrix(const FMatrixH& InMatrix) const
 {
-	ensure(false);
+	ensureCADKernel(false);
 	//TArray<FPoint> TransformedPoles;
 	//TransformedPoles.Reserve(Poles.Num());
 
@@ -24,7 +24,7 @@ TSharedPtr<FEntityGeom> FSplineCurve::ApplyMatrix(const FMatrixH& InMatrix) cons
 
 void FSplineCurve::Offset(const FPoint& OffsetDirection)
 {
-	ensure(false);
+	ensureCADKernel(false);
 	//for (FPoint& Pole : Poles)
 	//{
 	//	Pole += OffsetDirection;
@@ -42,7 +42,7 @@ FInfoEntity& FSplineCurve::GetInfo(FInfoEntity& Info) const
 
 void FSplineCurve::ExtendTo(const FPoint& Point)
 {
-	ensure(false);
+	ensureCADKernel(false);
 	//PolylineTools::ExtendTo(Poles, Point);
 }
 
@@ -71,7 +71,7 @@ void FSplineCurve::SetSplinePoints(const TArray<FPoint>& Points)
 void FSplineCurve::SetSplinePoints(const TArray<FPoint>& Points, const TArray<FPoint>& Tangents)
 {
 	const int32 NumPoints = Points.Num();
-	ensure(Tangents.Num() == NumPoints);
+	ensureCADKernel(Tangents.Num() == NumPoints);
 
 	Position.Points.Reset(NumPoints);
 
@@ -91,7 +91,7 @@ void FSplineCurve::SetSplinePoints(const TArray<FPoint>& Points, const TArray<FP
 {
 	const int32 NumPoints = Points.Num();
 	Position.Points.Reset(NumPoints);
-	ensure(false);
+	ensureCADKernel(false);
 
 	double InputKey = 0.0;
 	for (int32 Index = 0; Index < Points.Num(); ++Index)
diff --git a/Engine/Source/Runtime/Datasmith/CADKernel/Base/Private/Geo/Sampler/SamplerAbstract.h b/Engine/Source/Runtime/Datasmith/CADKernel/Base/Private/Geo/Sampler/SamplerAbstract.h
--- a/Engine/Source/Runtime/Datasmith/CADKernel/Base/Private/Geo/Sampler/SamplerAbstract.h
+++ b/Engine/Source/Runtime/Datasmith/CADKernel/Base/Private/Geo/Sampler/SamplerAbstract.h
@@ -91,10 +91,6 @@ protected:
 	 */
 	virtual void SamplingInitalizing()
 	{
-#ifdef DEBUG_CURVE_SAMPLING
-		++CurveIndex;
-		CurveToDisplay = CurveIndex;
-#endif
 		IsOptimalSegments.Empty(100);
 		Sampling.Empty(100);
 
@@ -116,8 +112,10 @@ protected:
 
 			for (int32 Index = 0; Index < NotDerivableCoordinates.Num() - 1; ++Index)
 			{
-				AddIntermediateCoordinates(NotDerivableCoordinates[Index], NotDerivableCoordinates[Index + 1], ComplementaryPointOffset);
-				NextCoordinates.Add(NotDerivableCoordinates[Index + 1]);
+				if (AddIntermediateCoordinates(NotDerivableCoordinates[Index], NotDerivableCoordinates[Index + 1], ComplementaryPointOffset))
+				{
+					NextCoordinates.Add(NotDerivableCoordinates[Index + 1]);
+				}
 			}
 		}
 
@@ -125,23 +123,17 @@ protected:
 		CandidatePoints.SwapCoordinates(NextCoordinates);
 		EvaluatesNewCandidatePoints();
 
-		// Not Derivable point initialize the sampling
+		// Not Derivable points and boundaries initialize the sampling
 		ComplementaryPointOffset++;
 		for (int32 Index = 0, ISampling = 0; Index < CandidatePoints.Size(); Index += ComplementaryPointOffset, ++ISampling)
 		{
 			Sampling.EmplaceAt(ISampling, CandidatePoints, Index);
 			IsOptimalSegments.Add(false);
 		}
+		// Remove last entered entry to IsOptimalSegments as this is to track segments
 		IsOptimalSegments.Pop();
 
 		CandidatePoints.RemoveComplementaryPoints(ComplementaryPointOffset);
-
-		// first segment is not optimal
-
-#ifdef DEBUG_CURVE_SAMPLING
-		int32 Step = 0;
-		DisplaySampling(CurveIndex == CurveToDisplay, Step);
-#endif
 	}
 
 	/**
@@ -255,11 +247,11 @@ protected:
 	}
 
 	/** Adds coordinates of the next candidate points in NextCoordinates array */
-	void AddIntermediateCoordinates(double UMin, double UMax, int32 PointNum)
+	bool AddIntermediateCoordinates(double UMin, double UMax, int32 PointNum)
 	{
 		if (FMath::IsNearlyEqual(UMin, UMax, UE_DOUBLE_SMALL_NUMBER))
 		{
-			return;
+			return false;
 		}
 		double DeltaCoord = (UMax - UMin) / (PointNum + 1);
 		double UCoord = UMin;
@@ -268,6 +260,8 @@ protected:
 			UCoord += DeltaCoord;
 			NextCoordinates.Add(UCoord);
 		}
+
+		return true;
 	};
 
 	/**
diff --git a/Engine/Source/Runtime/Datasmith/CADKernel/Base/Private/Geo/Sampler/SamplerOnParam.h b/Engine/Source/Runtime/Datasmith/CADKernel/Base/Private/Geo/Sampler/SamplerOnParam.h
--- a/Engine/Source/Runtime/Datasmith/CADKernel/Base/Private/Geo/Sampler/SamplerOnParam.h
+++ b/Engine/Source/Runtime/Datasmith/CADKernel/Base/Private/Geo/Sampler/SamplerOnParam.h
@@ -311,8 +311,10 @@ public:
 
 			for (int32 Index = 0; Index < LocalNotDerivableCoordinates.Num() - 1; ++Index)
 			{
-				AddIntermediateCoordinates(LocalNotDerivableCoordinates[Index], LocalNotDerivableCoordinates[Index + 1], ComplementaryPointOffset);
-				NextCoordinates.Add(LocalNotDerivableCoordinates[Index + 1]);
+				if (AddIntermediateCoordinates(LocalNotDerivableCoordinates[Index], LocalNotDerivableCoordinates[Index + 1], ComplementaryPointOffset))
+				{
+					NextCoordinates.Add(LocalNotDerivableCoordinates[Index + 1]);
+				}
 			}
 		}
 
@@ -329,13 +331,6 @@ public:
 		IsOptimalSegments.Pop();
 
 		CandidatePoints.RemoveComplementaryPoints(ComplementaryPointOffset);
-
-		// first segment is not optimal
-
-#ifdef DEBUG_CURVE_SAMPLING
-		int32 Step = 0;
-		DisplaySampling(CurveIndex == CurveToDisplay, Step);
-#endif
 	}
 
 protected:
diff --git a/Engine/Source/Runtime/Datasmith/CADKernel/Base/Private/Geo/Sampling/SurfacicPolyline.cpp b/Engine/Source/Runtime/Datasmith/CADKernel/Base/Private/Geo/Sampling/SurfacicPolyline.cpp
--- a/Engine/Source/Runtime/Datasmith/CADKernel/Base/Private/Geo/Sampling/SurfacicPolyline.cpp
+++ b/Engine/Source/Runtime/Datasmith/CADKernel/Base/Private/Geo/Sampling/SurfacicPolyline.cpp
@@ -9,6 +9,8 @@
 #include "Geo/Sampling/PolylineTools.h"
 #include "Math/SlopeUtils.h"
 
+#include "Algo/BinarySearch.h"
+
 namespace UE::CADKernel
 {
 
@@ -88,125 +90,132 @@ void FSurfacicPolyline::GetExtremities(const FLinearBoundary& InBoundary, const
 	}
 }
 
-void FSurfacicPolyline::ComputeIntersectionsWithIsos(const FLinearBoundary& InBoundary, const TArray<double>& IsoCoordinates, const EIso TypeIso, const FSurfacicTolerance& ToleranceIso, TArray<double>& Intersection) const
+void FSurfacicPolyline::ComputeIntersectionsWithIsos(const FLinearBoundary& InBoundary, const TArray<double>& IsoCoordinates, const EIso IsoType, const FSurfacicTolerance& ToleranceIso, TArray<double>& Intersection) const
 {
-	if (BoundingBox.Length(TypeIso) < ToleranceIso[TypeIso])
+	TRACE_CPUPROFILER_EVENT_SCOPE(UE::CADKernel::FSurfacicPolyline::ComputeIntersectionsWithIsos)
+
+	const double SurfaceIsoTolerance = ToleranceIso[IsoType];
+
+	if (BoundingBox.Length(IsoType) < SurfaceIsoTolerance)
 	{
 		// the edge is on an IsoCurve on Iso axis
 		return;
 	}
 
-#ifdef DEBUG_COMPUTEINTERSECTIONSWITHISOS
-	F3DDebugSession _(TEXT("IntersectEdgeIsos"));
-#endif
- 
-	Intersection.Reserve(IsoCoordinates.Num());
+	const int32 IsoCoordinateCount = IsoCoordinates.Num();
+	Intersection.Reserve(IsoCoordinateCount + Coordinates.Num());
 
-	double UMin;
-	double UMax;
+	double LastIntersection = UE_MAX_FLT;
 
-	int32 IsoCoordinateIndex = 0;
-	double IsoCoordinate = IsoCoordinates[0];
-
-	for (int32 Index = 0; Index < Points2D.Num() - 1; ++Index)
-	{
-		// if the segment is outside edge boundary, go to the next one
-		if (Coordinates[Index + 1] < InBoundary.GetMin())
+	TFunction<void(double, double, int32, double)> InsertIntersection;
+	InsertIntersection = [this, &Intersection, &LastIntersection](double Start, double End, int32 SegmentIndex, double IsoCoordinate) -> void
 		{
-			continue;
-		}
+			const double LocalCoord = (IsoCoordinate - Start) / (End - Start);
+			double EdgeCoord = this->Coordinates[SegmentIndex] + (LocalCoord * (this->Coordinates[SegmentIndex + 1] - this->Coordinates[SegmentIndex]));
 
-		if (Coordinates[Index] > InBoundary.GetMax())
-		{
-			break;
-		}
-
-		// if the segment is too parallel to Iso, go to the next one
-		double EdgeLocalSlope = ComputeUnorientedSlope(Points2D[Index], Points2D[Index + 1], 0);
-		if (TypeIso == EIso::IsoV)
-		{
-			if (EdgeLocalSlope < 0.1 || EdgeLocalSlope > 3.9)
+			if (!FMath::IsNearlyEqual(EdgeCoord, LastIntersection))
 			{
-				continue;
+				Intersection.Add(EdgeCoord);
 			}
-		}
-		else
+		};
+
+	// Find the largest surface's iso coordinate that is less or equal to the start iso value,
+	// within the surface's tolerance on the given iso type
+	TFunction<int32(double)> FindIsoCoordinateIndex;
+	FindIsoCoordinateIndex = [&IsoCoordinates, &SurfaceIsoTolerance, &IsoCoordinateCount](double Value) -> int32
 		{
-			if (EdgeLocalSlope < 2.1 && EdgeLocalSlope > 1.9)
-			{
-				continue;
-			}
-		}
+			// #cad_kernel: Algo::LowerBound ???
+			int32 Index = 0;
+			const double IsoValueStartMinusTol = Value - SurfaceIsoTolerance;
+			for (; Index < IsoCoordinateCount - 1 && IsoValueStartMinusTol > IsoCoordinates[Index]; ++Index);
 
-		GetMinMax(Points2D[Index][TypeIso], Points2D[Index + 1][TypeIso], UMin, UMax);
-		UMin -= DOUBLE_SMALL_NUMBER;
-		UMax += DOUBLE_SMALL_NUMBER;
+			return Index;
+		};
 
-		if (IsoCoordinate < UMin || IsoCoordinate > UMax)
+	for (int32 Index = 0, NextIndex = 1; NextIndex < Points2D.Num(); ++NextIndex, ++Index)
+	{
+		// Check that the segment to consider is within the curve's boundaries
+		// used for trimming
 		{
-			if (IsoCoordinateIndex > 0 && IsoCoordinateIndex < IsoCoordinates.Num() - 1 && IsoCoordinates[IsoCoordinateIndex - 1] < UMin && UMax < IsoCoordinates[IsoCoordinateIndex + 1])
+			// if the segment is outside the minimum edge boundary, go to the next one
+			if (Coordinates[NextIndex] < InBoundary.GetMin())
 			{
 				continue;
 			}
-		}
 
-		if (IsoCoordinate > UMin)
-		{
-			while (IsoCoordinateIndex > 0 && IsoCoordinates[IsoCoordinateIndex - 1] > UMin)
+			// if the segment is outside the maximum edge boundary, no need to go any further
+			if (Coordinates[Index] > InBoundary.GetMax())
 			{
-				IsoCoordinateIndex--;
+				break;
 			}
-			IsoCoordinate = IsoCoordinates[IsoCoordinateIndex];
 		}
 
-		while (IsoCoordinateIndex < IsoCoordinates.Num() - 1 && IsoCoordinates[IsoCoordinateIndex] < UMin)
+		const double IsoValueStart = Points2D[Index][IsoType];
+		const double IsoValueEnd = Points2D[NextIndex][IsoType];
+
+		// Skip this segment if it is degenerated along the given iso type
+		if (FMath::IsNearlyEqual(IsoValueStart, IsoValueEnd, SurfaceIsoTolerance))
 		{
-			IsoCoordinateIndex++;
+			continue;
 		}
-		IsoCoordinate = IsoCoordinates[IsoCoordinateIndex];
 
-		if (UMax < IsoCoordinate)
+		int32 IsoCoordinateIndex = FindIsoCoordinateIndex(IsoValueStart);
+
+		// If the start point is equal to the iso coordinate within the surface's tolerance
+		// on the given iso direction, add the polyline's coordinate and continue
+		if (FMath::IsNearlyEqual(IsoValueStart, IsoCoordinates[IsoCoordinateIndex], SurfaceIsoTolerance))
 		{
-			continue;
+			Intersection.Add(Coordinates[Index]);
+			LastIntersection = Intersection.Last();
 		}
 
-		for (;; ++IsoCoordinateIndex)
+		// Segment is forward in surface's given iso type's direction
+		if (IsoValueStart < IsoValueEnd)
 		{
-			if (IsoCoordinateIndex == IsoCoordinates.Num() || IsoCoordinates[IsoCoordinateIndex] > UMax)
+			// Insert an intersection as long as the iso values are less than the end value
+			// Intentionally do not check whether the end iso value is within surface's tolerance
+			// on an iso type's coordinate. It will be checked on the next iteration
+			const double IsoValueEndMinusTol = IsoValueEnd - SurfaceIsoTolerance;
+			for (++IsoCoordinateIndex; IsoCoordinateIndex < IsoCoordinateCount && IsoCoordinates[IsoCoordinateIndex] < IsoValueEndMinusTol; ++IsoCoordinateIndex)
 			{
-				--IsoCoordinateIndex;
-				break;
+				InsertIntersection(IsoValueStart, IsoValueEnd, Index, IsoCoordinates[IsoCoordinateIndex]);
+				LastIntersection = Intersection.Last();
 			}
-			IsoCoordinate = IsoCoordinates[IsoCoordinateIndex];
-
-			if ((IsoCoordinate > UMin) && (IsoCoordinate < UMax))
+		}
+		// Segment is backward in surface's given iso type's direction
+		else if(IsoCoordinateIndex > 0)
+		{
+			// Insert an intersection as long as the iso values are more than the end value
+			// Intentionally do not check whether the end iso value is within surface's tolerance
+			// on an iso type's coordinate. It will be checked on the next iteration
+			const double IsoValueEndMinusTol = IsoValueEnd + SurfaceIsoTolerance;
+			for (--IsoCoordinateIndex; IsoCoordinateIndex >= 0 && IsoCoordinates[IsoCoordinateIndex] > IsoValueEndMinusTol; --IsoCoordinateIndex)
 			{
-				double EdgeCoord = 0.;
-				const double Delta = (Points2D[Index + 1][TypeIso] - Points2D[Index][TypeIso]);
-				if (FMath::IsNearlyZero(Delta))
-				{
-					EdgeCoord = Coordinates[Index];
-				}
-				else
-				{
-					const double LocalCoord = (IsoCoordinate - Points2D[Index][TypeIso]) / Delta;
-					EdgeCoord = PolylineTools::LinearInterpolation(Coordinates, Index, LocalCoord);
-				}
-
-#ifdef DEBUG_COMPUTEINTERSECTIONSWITHISOS
-				::DisplayPoint(Points2D[Index + 1], EVisuProperty::RedPoint);
-				::DisplayPoint(Points2D[Index], EVisuProperty::RedPoint);
-				FPoint2D IntersectionPoint = PolylineTools::LinearInterpolation(Points2D, Index, LocalCoord);
-				::DisplayPoint(IntersectionPoint);
-				Wait();
-#endif // #ifdef DEBUG_COMPUTEINTERSECTIONSWITHISOS
-
-				if (!(Intersection.Num() && FMath::IsNearlyEqual(EdgeCoord, Intersection.Last())))
-				{
-					Intersection.Add(EdgeCoord);
-				}
+				InsertIntersection(IsoValueStart, IsoValueEnd, Index, IsoCoordinates[IsoCoordinateIndex]);
+				LastIntersection = Intersection.Last();
 			}
 		}
+		// Lowest surface's iso coordinate is strictly between the end value and the start value
+		// Insert an intersection there
+		else
+		{
+			InsertIntersection(IsoValueStart, IsoValueEnd, Index, IsoCoordinates[IsoCoordinateIndex]);
+			LastIntersection = Intersection.Last();
+		}
+	}
+
+	// Process last point
+	{
+		const double IsoValue = Points2D.Last()[IsoType];
+
+		int32 IsoCoordinateIndex = FindIsoCoordinateIndex(IsoValue);
+
+		// If the start point is equal to the iso coordinate within the surface's tolerance
+		// on the given iso direction, add the polyline's coordinate and continue
+		if (FMath::IsNearlyEqual(IsoValue, IsoCoordinates[IsoCoordinateIndex], SurfaceIsoTolerance))
+		{
+			Intersection.Add(Coordinates.Last());
+		}
 	}
 }
 
diff --git a/Engine/Source/Runtime/Datasmith/CADKernel/Base/Private/Geo/Surfaces/BezierSurface.h b/Engine/Source/Runtime/Datasmith/CADKernel/Base/Private/Geo/Surfaces/BezierSurface.h
--- a/Engine/Source/Runtime/Datasmith/CADKernel/Base/Private/Geo/Surfaces/BezierSurface.h
+++ b/Engine/Source/Runtime/Datasmith/CADKernel/Base/Private/Geo/Surfaces/BezierSurface.h
@@ -52,12 +52,12 @@ public:
 		return ESurface::Bezier;
 	}
 
-	int32 GetUDegre() const
+	int32 GetUDegree() const
 	{
 		return UDegre;
 	}
 
-	int32 GetVDegre() const
+	int32 GetVDegree() const
 	{
 		return VDegre;
 	}
diff --git a/Engine/Source/Runtime/Datasmith/CADKernel/Base/Private/Geo/Surfaces/SurfaceUtilities.cpp b/Engine/Source/Runtime/Datasmith/CADKernel/Base/Private/Geo/Surfaces/SurfaceUtilities.cpp
new file mode 100644
--- /dev/null
+++ b/Engine/Source/Runtime/Datasmith/CADKernel/Base/Private/Geo/Surfaces/SurfaceUtilities.cpp
@@ -0,0 +1,106 @@
+// Copyright Epic Games, Inc. All Rights Reserved.
+
+#include "Geo/Surfaces/SurfaceUtilities.h"
+
+#include "Geo/Surfaces/BezierSurface.h"
+#include "Geo/Surfaces/NURBSSurface.h"
+#include "Geo/Surfaces/OffsetSurface.h"
+#include "Geo/Surfaces/Surface.h"
+
+#include "Geo/GeoEnum.h"
+
+namespace UE::CADKernel::SurfaceUtilities
+{
+	TArray<TArray<FVector3f>> GetPoles(const UE::CADKernel::FSurface& Surface)
+	{
+		using namespace UE::CADKernel;
+
+		TArray<TArray<FVector3f>> Poles;
+		
+		return Poles;
+	}
+
+	bool AreControlPointsPlanar(const TArray<FPoint>& Poles, int32 UPoleCount)
+	{
+		using namespace UE::CADKernel;
+
+		const FVector Origin(Poles[0].X, Poles[0].Y, Poles[0].Z);
+		const FVector UDir(Poles[1].X, Poles[1].Y, Poles[1].Z);
+		const FVector VDir(Poles[UPoleCount].X, Poles[UPoleCount].Y, Poles[UPoleCount].Z);
+		const FVector Normal = ((UDir - Origin) ^ (VDir - Origin)).GetSafeNormal();
+
+		for (int32 Index = 1; Index < Poles.Num(); ++Index)
+		{
+			const double Dot = Normal | (FVector(Poles[Index].X, Poles[Index].Y, Poles[Index].Z) - Origin).GetSafeNormal();
+			if (FMath::Abs(Dot) > 0.01745241 /* Angle between Normal and segment is more than 89 degrees */)
+			{
+				return false;
+			}
+		}
+
+		return true;
+	}
+
+	bool IsSurfacePlanar(const UE::CADKernel::FBezierSurface& Surface)
+	{
+		if (Surface.GetUDegree() == 1 && Surface.GetVDegree() == 1)
+		{
+			return true;
+		}
+
+		return AreControlPointsPlanar(Surface.GetPoles(), Surface.GetUDegree() + 1);
+	}
+
+	bool IsSurfacePlanar(const UE::CADKernel::FNURBSSurface& Surface)
+	{
+		if (Surface.GetDegree(EIso::IsoU) == 1 && Surface.GetDegree(EIso::IsoV) == 1)
+		{
+			return true;
+		}
+
+		// #cad_kernel: Might be rational but still planar
+		if (Surface.IsRational())
+		{
+			return false;
+		}
+
+		return AreControlPointsPlanar(Surface.GetPoles(), Surface.GetPoleCount(EIso::IsoU));
+	}
+
+	bool IsPlanar(const UE::CADKernel::FSurface& Surface)
+	{
+		using namespace UE::CADKernel;
+
+		switch (Surface.GetSurfaceType())
+		{
+		case ESurface::Nurbs:
+			return IsSurfacePlanar(static_cast<const FNURBSSurface&>(Surface));
+		
+		case ESurface::Bezier:
+			return IsSurfacePlanar(static_cast<const FBezierSurface&>(Surface));
+
+		case ESurface::Offset:
+			return IsPlanar(*static_cast<const FOffsetSurface&>(Surface).GetBaseSurface());
+
+		case ESurface::Composite:
+		case ESurface::Cone:
+		case ESurface::Coons:
+		case ESurface::Cylinder:
+		case ESurface::Revolution:
+		case ESurface::Ruled:
+		case ESurface::Sphere:
+		case ESurface::TabulatedCylinder:
+		case ESurface::Torus:
+			return false;
+
+		case ESurface::Plane:
+			return true;
+
+		default:
+			ensureCADKernel(false);
+			break;
+		}
+
+		return false;
+	}
+}
diff --git a/Engine/Source/Runtime/Datasmith/CADKernel/Base/Private/Math/BSpline.cpp b/Engine/Source/Runtime/Datasmith/CADKernel/Base/Private/Math/BSpline.cpp
--- a/Engine/Source/Runtime/Datasmith/CADKernel/Base/Private/Math/BSpline.cpp
+++ b/Engine/Source/Runtime/Datasmith/CADKernel/Base/Private/Math/BSpline.cpp
@@ -286,10 +286,10 @@ void InterpolerBSpline2D(const FNURBSSurface& Nurbs, FGrid& Grid)
 		}
 
 		// pour chaque Reseau Pj(U), on calcule les points P(UV)
-		//TBSplineData<FPointH> Data(FMath::Max(Nurbs.GetUDegre(), Nurbs.GetVDegre()) + 1);
+		//TBSplineData<FPointH> Data(FMath::Max(Nurbs.GetUDegree(), Nurbs.GetVDegree()) + 1);
 
 		//FTimePoint StartTime2 = FChrono::Now();
-		//TBSplineData<FPointH> Data2(FMath::Max(Nurbs.GetUDegre() , Grid.GetNumV()) + 1);
+		//TBSplineData<FPointH> Data2(FMath::Max(Nurbs.GetUDegree() , Grid.GetNumV()) + 1);
 		//FChrono::PrintClockElapse(LOG, "    ", "Alloc", FChrono::Elapse(StartTime2), ETimeUnit::NanoSeconds);
 
 		//FTimePoint StartTime = FChrono::Now();
@@ -299,7 +299,7 @@ void InterpolerBSpline2D(const FNURBSSurface& Nurbs, FGrid& Grid)
 
 		{
 			FPointH OutGradient;
-			TBSplineSurface<FPointH> Surface(Nurbs.GetPoleNumU(), Nurbs.GetUDegre(), Nurbs.GetNodalVectorU(), Nurbs.GetHomogeneousPoles());
+			TBSplineSurface<FPointH> Surface(Nurbs.GetPoleNumU(), Nurbs.GetUDegree(), Nurbs.GetNodalVectorU(), Nurbs.GetHomogeneousPoles());
 			//string Message = "Iteration IndexV " + Utils::ToString(IndexV);
 			Surface.Display(TEXT("in"));
 			Wait();
@@ -310,15 +310,15 @@ void InterpolerBSpline2D(const FNURBSSurface& Nurbs, FGrid& Grid)
 			}
 		}
 
-		TBSplineSurface<FPointH> OutSurface(Nurbs.GetPoleNumV(), Nurbs.GetVDegre(), Nurbs.GetNodalVectorV(), PointUCurveIndexVArray);
+		TBSplineSurface<FPointH> OutSurface(Nurbs.GetPoleNumV(), Nurbs.GetVDegree(), Nurbs.GetNodalVectorV(), PointUCurveIndexVArray);
 		OutSurface.Display(TEXT("Hello"));
 		Wait();
 
 		return;
 		//{
-		//TBSplineSurface<FPointH> OutSurface(Nurbs.GetPoleNumV(), Nurbs.GetVDegre(), Nurbs.GetVNodalVector(), PointUCurveIndexVArray);
+		//TBSplineSurface<FPointH> OutSurface(Nurbs.GetPoleNumV(), Nurbs.GetVDegree(), Nurbs.GetVNodalVector(), PointUCurveIndexVArray);
 		//	FPointH OutGradient;
-		//	TBSplineSurface<FPointH> InSurface(Nurbs.GetPoleNumV(), Nurbs.GetVDegre(), Nurbs.GetVNodalVector(), Nurbs.GetHomogeneousPoles());
+		//	TBSplineSurface<FPointH> InSurface(Nurbs.GetPoleNumV(), Nurbs.GetVDegree(), Nurbs.GetVNodalVector(), Nurbs.GetHomogeneousPoles());
 		//	for (int32 IndexU = 0, Index = 0; IndexU <Grid.GetNumU(); ++IndexU)
 		//	{
 		//		InSurface.SetPoles(PointUCurveIndexVArray[IndexU].GetData() + Index);
@@ -717,7 +717,7 @@ void Interpolate2DBSpline(
 
 void EvaluatePoint(const FNURBSCurve& Nurbs, double Coordinate, FCurvePoint& OutPoint, int32 DerivativeOrder)
 {
-	ensure(Nurbs.GetDimension() == 3);
+	ensureCADKernel(Nurbs.GetDimension() == 3);
 
 	OutPoint.DerivativeOrder = DerivativeOrder;
 
diff --git a/Engine/Source/Runtime/Datasmith/CADKernel/Base/Private/Mesh/Criteria/CriteriaGrid.cpp b/Engine/Source/Runtime/Datasmith/CADKernel/Base/Private/Mesh/Criteria/CriteriaGrid.cpp
--- a/Engine/Source/Runtime/Datasmith/CADKernel/Base/Private/Mesh/Criteria/CriteriaGrid.cpp
+++ b/Engine/Source/Runtime/Datasmith/CADKernel/Base/Private/Mesh/Criteria/CriteriaGrid.cpp
@@ -21,7 +21,7 @@ void FCriteriaGrid::Init()
 		TArray<double>& Tab2 = CoordinateGrid[Iso];
 
 		CuttingCount[Iso] = Tab.Num() * 2 - 1;
-		ensure(CuttingCount[Iso]);
+		ensureCADKernel(CuttingCount[Iso]);
 		Tab2.SetNum(CuttingCount[Iso]);
 
 		Tab2[0] = Tab[0];
diff --git a/Engine/Source/Runtime/Datasmith/CADKernel/Base/Private/Mesh/Meshers/MesherTools.cpp b/Engine/Source/Runtime/Datasmith/CADKernel/Base/Private/Mesh/Meshers/MesherTools.cpp
--- a/Engine/Source/Runtime/Datasmith/CADKernel/Base/Private/Mesh/Meshers/MesherTools.cpp
+++ b/Engine/Source/Runtime/Datasmith/CADKernel/Base/Private/Mesh/Meshers/MesherTools.cpp
@@ -44,7 +44,7 @@ void FMesherTools::ComputeFinalCuttingPointsWithPreferredCuttingPoints(const TAr
 
 
 	// Step 2: Define the number of step(each step lasts one) as it must be an integer bigger than 1
-	double CorrectedFinalTravelTime = FMath::Max(((double)((int32)(TravelTime + 0.5))), 1.);
+	double CorrectedFinalTravelTime = FMath::Max((double)(FMath::CeilToInt32(TravelTime)), 1.);
 
 	// Step 3: Adjust the speeds than the travel respect the count of step
 	{
@@ -68,20 +68,20 @@ void FMesherTools::ComputeFinalCuttingPointsWithPreferredCuttingPoints(const TAr
 		{
 			for (; UIndex < DeltaUs.Num(); ++UIndex)
 			{
-				double DeltaU = DeltaUs[UIndex];
-				if (PreferredCuttingPoints[NeighbourIndex].Coordinate - DOUBLE_SMALL_NUMBER > CrossingUs[UIndex + 1])
+				const double DeltaU = DeltaUs[UIndex];
+				const double PreferredCuttingPoint = PreferredCuttingPoints[NeighbourIndex].Coordinate;
+
+				if (PreferredCuttingPoint - DOUBLE_SMALL_NUMBER > CrossingUs[UIndex + 1])
 				{
-					double DeltaTravelTime = (CrossingUs[UIndex + 1] - LastStep) / DeltaU;
+					TravelTime += (CrossingUs[UIndex + 1] - LastStep) / DeltaU;
 					LastStep = CrossingUs[UIndex + 1];
-					TravelTime += DeltaTravelTime;
 					continue;
 				}
 
-				if (CrossingUs[UIndex] - DOUBLE_SMALL_NUMBER < PreferredCuttingPoints[NeighbourIndex].Coordinate && PreferredCuttingPoints[NeighbourIndex].Coordinate < CrossingUs[UIndex + 1] + DOUBLE_SMALL_NUMBER)
+				if (CrossingUs[UIndex] - DOUBLE_SMALL_NUMBER < PreferredCuttingPoint && PreferredCuttingPoint < CrossingUs[UIndex + 1] + DOUBLE_SMALL_NUMBER)
 				{
-					double DeltaTravelTime = (PreferredCuttingPoints[NeighbourIndex].Coordinate - LastStep) / DeltaU;
-					double TravelTimeToNeighbourU = TravelTime + DeltaTravelTime;
-					TravelTimeOfNeighborsU[NeighbourIndex] = TravelTimeToNeighbourU;
+					const double DeltaTravelTime = (PreferredCuttingPoint - LastStep) / DeltaU;
+					TravelTimeOfNeighborsU[NeighbourIndex] = TravelTime + DeltaTravelTime;
 				}
 				break;
 			}
@@ -372,4 +372,4 @@ void FMesherTools::FillImposedIsoCuttingPoints(TArray<double>& UEdgeSetOfInterse
 }
 
 
-}
\ No newline at end of file
+}
diff --git a/Engine/Source/Runtime/Datasmith/CADKernel/Base/Private/Mesh/Meshers/ParametricFaceMesher.cpp b/Engine/Source/Runtime/Datasmith/CADKernel/Base/Private/Mesh/Meshers/ParametricFaceMesher.cpp
--- a/Engine/Source/Runtime/Datasmith/CADKernel/Base/Private/Mesh/Meshers/ParametricFaceMesher.cpp
+++ b/Engine/Source/Runtime/Datasmith/CADKernel/Base/Private/Mesh/Meshers/ParametricFaceMesher.cpp
@@ -15,14 +15,21 @@
 #include "Topo/TopologicalEdge.h"
 #include "Topo/TopologicalFace.h"
 
-
-#ifdef CADKERNEL_DEV
-#include "Mesh/Meshers/MesherReport.h"
-#endif
+#include "Geo/Curves/CurveUtilities.h"
+#include "Geo/Surfaces/SurfaceUtilities.h"
+#include "Topo/TopologicalFaceUtilities.h"
+#include "CompGeom/Delaunay2.h"
+#include "HAL/IConsoleManager.h"
+
+static bool GDetectPlanarFace = true;
+FAutoConsoleVariableRef PlanarFaceCVar(
+	TEXT("CADKernel.FaceMesher.DetectPlanarFace"),
+	GDetectPlanarFace,
+	TEXT(""),
+	ECVF_Default);
 
 namespace UE::CADKernel
 {
-
 FParametricFaceMesher::FParametricFaceMesher(FTopologicalFace& InFace, FModelMesh& InMeshModel, const FMeshingTolerances& InTolerances, bool bActivateThinZoneMeshing)
 	: Face(InFace)
 	, MeshModel(InMeshModel)
@@ -32,17 +39,12 @@ FParametricFaceMesher::FParametricFaceMesher(FTopologicalFace& InFace, FModelMes
 #endif
 	, Grid(Face, MeshModel)
 {
-#ifdef DEBUG_ONLY_SURFACE_TO_DEBUG
-	bDisplay = (Face.GetId() == FaceToDebug);
-#endif
 }
 
 void FParametricFaceMesher::Mesh()
 {
-#ifdef DEBUG_ONLY_SURFACE_TO_DEBUG
-	bDisplay = (Face.GetId() == FaceToDebug);
-#endif
-
+	TRACE_CPUPROFILER_EVENT_SCOPE(UE::CADKernel::FParametricFaceMesher::GenerateCloud)
+	
 	Face.GetOrCreateMesh(MeshModel);
 
 	if (Face.IsNotMeshable())
@@ -56,17 +58,10 @@ void FParametricFaceMesher::Mesh()
 
 	FProgress _(1, TEXT("Meshing Entities : Mesh Surface"));
 
-#ifdef DEBUG_CADKERNEL
-	F3DDebugSession S(bDisplayDebugMeshStep, FString::Printf(TEXT("Mesh of surface %d"), Face.GetId()));
-#endif
-
 	FTimePoint StartTime = FChrono::Now();
 
 	if (!GenerateCloud() || Grid.IsDegenerated())
 	{
-#ifdef CADKERNEL_DEV
-		FMesherReport::GetLogs().AddDegeneratedGrid();
-#endif
 		FMessage::Printf(EVerboseLevel::Log, TEXT("The meshing of the surface %d failed due to a degenerated grid\n"), Face.GetId());
 
 		Face.IsDegenerated();
@@ -76,35 +71,32 @@ void FParametricFaceMesher::Mesh()
 
 	FTimePoint IsoTriangulerStartTime = FChrono::Now();
 
-	FFaceMesh& SurfaceMesh = Face.GetOrCreateMesh(MeshModel);
-	FIsoTriangulator IsoTrianguler(Grid, SurfaceMesh, Tolerances);
-
-	if (IsoTrianguler.Triangulate())
+	if (GDetectPlanarFace && SurfaceUtilities::IsPlanar(*Face.GetCarrierSurface())/*TopologicalFaceUtilities::IsPolygonal(Face)*/)
+	{
+		MeshPlanarFace();
+	}
+	else
 	{
-		if (Face.IsBackOriented())
+		FFaceMesh& SurfaceMesh = Face.GetOrCreateMesh(MeshModel);
+		FIsoTriangulator IsoTrianguler(Grid, SurfaceMesh, Tolerances);
+
+		if (IsoTrianguler.Triangulate())
 		{
-			SurfaceMesh.InverseOrientation();
+			if (Face.IsBackOriented())
+			{
+				SurfaceMesh.InverseOrientation();
+			}
+			MeshModel.AddMesh(SurfaceMesh);
 		}
-		MeshModel.AddMesh(SurfaceMesh);
 	}
-	Face.SetMeshedMarker();
-
-	FDuration TriangulateDuration = FChrono::Elapse(IsoTriangulerStartTime);
-	FDuration Duration = FChrono::Elapse(StartTime);
-
-#ifdef CADKERNEL_DEV
-	FMesherReport::GetChronos().GlobalPointCloudDuration += Grid.Chronos.GeneratePointCloudDuration;
-	FMesherReport::GetChronos().GlobalTriangulateDuration += TriangulateDuration;
-	FMesherReport::GetChronos().GlobalMeshDuration += Duration;
-#endif
-
-	FChrono::PrintClockElapse(EVerboseLevel::Debug, TEXT("   "), TEXT("Meshing"), Duration);
 
+	Face.SetMeshedMarker();
 }
 
 bool FParametricFaceMesher::GenerateCloud()
 {
-	FTimePoint GenerateCloudStartTime = FChrono::Now();
+	TRACE_CPUPROFILER_EVENT_SCOPE(UE::CADKernel::FParametricFaceMesher::GenerateCloud)
+	
 	Grid.DefineCuttingParameters();
 	if (!Grid.GeneratePointCloud())
 	{
@@ -118,28 +110,20 @@ bool FParametricFaceMesher::GenerateCloud()
 		{
 			MeshThinZones();
 		}
-#ifdef CADKERNEL_DEV
-		FMesherReport::GetChronos().GlobalMeshThinZones += FChrono::Elapse(StartTime);
-#endif
 	}
 
 	MeshFaceLoops();
 
 	Grid.ProcessPointCloud();
 
-#ifdef CADKERNEL_DEV
-	FMesherReport::GetChronos().GlobalGeneratePointCloudDuration += FChrono::Elapse(GenerateCloudStartTime);
-#endif
-
 	return true;
 }
 
 void FParametricFaceMesher::MeshThinZones(TArray<FTopologicalEdge*>& EdgesToMesh, const bool bFinalMeshing)
 {
+	TRACE_CPUPROFILER_EVENT_SCOPE(UE::CADKernel::FParametricFaceMesher::MeshThinZones)
+
 	{
-#ifdef DEBUG_MESH_THINZONE_EDGES
-		F3DDebugSession A(bDisplay, TEXT("   - Define Cutting Points"));
-#endif
 		for (FTopologicalEdge* Edge : EdgesToMesh)
 		{
 			if (Edge->IsPreMeshed())
@@ -147,18 +131,8 @@ void FParametricFaceMesher::MeshThinZones(TArray<FTopologicalEdge*>& EdgesToMesh
 				continue;
 			}
 
-#ifdef DEBUG_MESH_THINZONE_EDGES
-			if(bDisplay)
-			{
-				F3DDebugSession B(bDisplay, FString::Printf(TEXT("Edge %d"), Edge->GetId()));
-				Display2DWithScale(*Edge, EVisuProperty::GreenCurve);
-			}
-#endif
 
 			{
-#ifdef DEBUG_MESH_THINZONE_EDGES
-				F3DDebugSession B(bDisplay, FString::Printf(TEXT("Edge %d"), Edge->GetId()));
-#endif
 				for (FThinZoneSide* ZoneSide : Edge->GetThinZoneSides())
 				{
 					if (ZoneSide->HasMarker1())
@@ -185,9 +159,6 @@ void FParametricFaceMesher::MeshThinZones(TArray<FTopologicalEdge*>& EdgesToMesh
 	}
 
 	{
-#ifdef DEBUG_MESH_THINZONE_EDGES
-		F3DDebugSession A(bDisplay, TEXT("   - Mesh"));
-#endif
 		for (FTopologicalEdge* Edge : EdgesToMesh)
 		{
 			if (Edge->IsMeshed())
@@ -198,7 +169,6 @@ void FParametricFaceMesher::MeshThinZones(TArray<FTopologicalEdge*>& EdgesToMesh
 			Mesh(*Edge, bFinalMeshing);
 		}
 	}
-	//Wait();
 }
 
 namespace ThinZoneMesherTools
@@ -221,9 +191,7 @@ void ResetMarkers(TArray<FTopologicalEdge*>& EdgesWithThinZones, TArray<FThinZon
 
 void FParametricFaceMesher::SortThinZoneSides(TArray<FThinZone2D*>& ThinZones)
 {
-#ifdef DEBUG_SORT_THIN_ZONE_SIDES
-	F3DDebugSession A(bDisplay, ("SortThinZoneSides"));
-#endif
+	TRACE_CPUPROFILER_EVENT_SCOPE(UE::CADKernel::FGrid::DefineCuttingParameters)
 
 	Face.ResetMarkersRecursively();
 	int32 EdgeCount = Face.EdgeCount();
@@ -248,10 +216,6 @@ void FParametricFaceMesher::SortThinZoneSides(TArray<FThinZone2D*>& ThinZones)
 	int32 Index = 1;
 	for (FThinZone2D* Zone : ThinZones)
 	{
-#ifdef DEBUG_SORT_THIN_ZONE_SIDES
-		Zone.Display(FString::Printf(TEXT("Zone %d"), Index++), EVisuProperty::YellowCurve);
-#endif
-
 		Zone->CheckEdgesZoneSide();
 
 
@@ -335,10 +299,6 @@ void FParametricFaceMesher::MeshThinZones()
 
 	while (ThinZones.Num())
 	{
-#ifdef DEBUG_THIN_ZONES
-		F3DDebugSession A(bDisplay, ("Mesh thin zone"));
-		DisplayMeshOfFaceLoop();
-#endif
 		int32 WaitingThinZoneCount = ThinZones.Num();
 
 		MeshThinZones(ThinZones);
@@ -354,6 +314,8 @@ void FParametricFaceMesher::MeshThinZones()
 
 void FParametricFaceMesher::MeshThinZones(TArray<FThinZone2D*>& ThinZones)
 {
+	TRACE_CPUPROFILER_EVENT_SCOPE(UE::CADKernel::FParametricFaceMesher::MeshThinZones)
+	
 	SortThinZoneSides(ThinZones);
 
 	TFunction<void(TArray<FTopologicalEdge*>)> TransfereCuttingPointsFromMeshedEdges = [](TArray<FTopologicalEdge*> Edges)
@@ -373,36 +335,12 @@ void FParametricFaceMesher::MeshThinZones(TArray<FThinZone2D*>& ThinZones)
 	TransfereCuttingPointsFromMeshedEdges(ZoneAEdges);
 	TransfereCuttingPointsFromMeshedEdges(ZoneBEdges);
 
-#ifdef DEBUG_THIN_ZONES
-	if(bDisplay)
-	{
-		DisplayThinZoneEdges(TEXT("Side A"), ZoneAEdges, EVisuProperty::BlueCurve, EVisuProperty::PurpleCurve);
-		DisplayThinZoneEdges(TEXT("Side B"), ZoneBEdges, EVisuProperty::RedCurve, EVisuProperty::PinkCurve);
-
-		{
-			F3DDebugSession A(bDisplay, ("Waiting Zone"));
-			for(FThinZone2D* Zone : WaitingThinZones)
-			{
-				Zone->Display(TEXT("Waiting Zone"), EVisuProperty::YellowCurve);
-			}
-		}
-
-		//Wait();
-	}
-#endif
-
 	bool bFinalMeshing = false;
 	{
-#ifdef DEBUG_MESHTHINSURF
-		F3DDebugSession _(bDisplay, TEXT("Step 1"));
-#endif
 		MeshThinZones(ZoneAEdges, bFinalMeshing);
 	}
 
 	{
-#ifdef DEBUG_MESHTHINSURF
-		F3DDebugSession _(bDisplay, TEXT("Step 2"));
-#endif
 		bFinalMeshing = true;
 		MeshThinZones(ZoneBEdges, bFinalMeshing);
 	}
@@ -413,18 +351,8 @@ void FParametricFaceMesher::MeshThinZones(TArray<FThinZone2D*>& ThinZones)
 	}
 
 	{
-#ifdef DEBUG_MESHTHINSURF
-		F3DDebugSession _(bDisplay, TEXT("Step 3"));
-#endif
 		MeshThinZones(ZoneAEdges, bFinalMeshing);
 	}
-
-#ifdef DEBUG_THIN_ZONES
-	{
-		DisplayMeshOfFaceLoop();
-	}
-	//Wait(bDisplay);
-#endif
 }
 
 void FParametricFaceMesher::DefineImposedCuttingPointsBasedOnOtherSideMesh(FThinZoneSide& SideToConstrain)
@@ -466,46 +394,6 @@ void FParametricFaceMesher::DefineImposedCuttingPointsBasedOnOtherSideMesh(FThin
 	SideToConstrain.GetExistingMeshNodes(Face, MeshModel, ReserveExistingMeshNodes, AddToExistingMeshNodes, /*bWithTolerance*/ false);
 	FrontSide.GetExistingMeshNodes(Face, MeshModel, ReserveCrossZoneElements, AddToCrossZoneElements, /*bWithTolerance*/ true);
 
-#ifdef DEBUG_MESHTHINSURF
-	if (bDisplay)
-	{
-		{
-			F3DDebugSession E(bDisplay, FString::Printf(TEXT("ThinZone")));
-			ThinZone::DisplayThinZoneSide(FrontSide, 0, EVisuProperty::BlueCurve);
-			ThinZone::DisplayThinZoneSide(SideToConstrain, 1, EVisuProperty::YellowPoint);
-			Wait(false);
-		}
-
-		F3DDebugSession A(TEXT("Existing Mesh nodes"));
-		{
-			F3DDebugSession A(TEXT("Existing Mesh nodes side to mesh"));
-			for (const TPair<int32, FPoint2D>& MeshNode : ExistingMeshNodes)
-			{
-				DisplayPoint2DWithScale(MeshNode.Value, EVisuProperty::RedPoint, MeshNode.Key);
-			}
-		}
-		{
-			F3DDebugSession A(TEXT("Existing Mesh nodes reference side"));
-			for (const FCrossZoneElement& MeshNode : CrossZoneElements)
-			{
-				DisplayPoint2DWithScale(MeshNode.VertexPoint2D, EVisuProperty::BluePoint, MeshNode.VertexId);
-				if (MeshNode.OppositeVertexIndices[0] >= 0)
-				{
-					FPoint2D* OppositeVertex = ExistingMeshNodes.Find(MeshNode.OppositeVertexIndices[0]);
-					if (OppositeVertex)
-					{
-						DisplayPoint2DWithScale(*OppositeVertex, EVisuProperty::RedPoint, MeshNode.OppositeVertexIndices[0]);
-						DisplaySegmentWithScale(*OppositeVertex, MeshNode.VertexPoint2D, 0, EVisuProperty::BlueCurve);
-					}
-				}
-			}
-		}
-		Wait(false);
-	}
-
-	F3DDebugSession A(bDisplay, TEXT("Find the best projection of existing mesh vertices"));
-#endif
-
 	const double MaxSquareThickness = FrontSide.GetMaxThickness() > SideToConstrain.GetMaxThickness() ? FMath::Square(3. * FrontSide.GetMaxThickness()) : FMath::Square(3. * SideToConstrain.GetMaxThickness());
 
 	// Find the best projection of existing mesh vertices (CrossZone Vertex)
@@ -522,15 +410,6 @@ void FParametricFaceMesher::DefineImposedCuttingPointsBasedOnOtherSideMesh(FThin
 			continue;
 		}
 
-#ifdef DEBUG_MESHTHINSURF
-		if (bDisplay)
-		{
-			F3DDebugSession A(TEXT("Point"));
-			DisplayPoint2DWithScale(CrossZoneElement.VertexPoint2D, EVisuProperty::BluePoint, CrossZoneElement.VertexId);
-			Wait(false);
-		}
-#endif
-
 		double MinSquareThickness = MaxSquareThickness;
 		FPoint2D ClosePoint;
 		FEdgeSegment* CloseSegment = nullptr;
@@ -545,34 +424,12 @@ void FParametricFaceMesher::DefineImposedCuttingPointsBasedOnOtherSideMesh(FThin
 			const double SlopeE = Segment.ComputeOrientedSlopeOf(MeshNodeCoordinate, Segment.GetExtemity(ELimit::End));
 			if (SlopeE < 1. || SlopeS > 3.)
 			{
-#ifdef DEBUG_MESHTHINSURF_
-				if (bDisplay)
-				{
-					F3DDebugSession A(FString::Printf(TEXT("Not candidate slope end %f"), SlopeE));
-					DisplaySegmentWithScale(Segment.GetExtemity(Start), Segment.GetExtemity(End), 0, EVisuProperty::RedCurve);
-					DisplaySegmentWithScale(MeshNodeCoordinate, Segment.GetExtemity(End), 0, EVisuProperty::YellowCurve);
-					DisplaySegmentWithScale(MeshNodeCoordinate, Segment.GetExtemity(Start), 0, EVisuProperty::YellowCurve);
-				}
-#endif
 				continue;
 			}
 
 			double CoordSegmentU;
 			FPoint2D Projection = Segment.ProjectPoint(MeshNodeCoordinate, CoordSegmentU);
 
-
-#ifdef DEBUG_MESHTHINSURF_
-			if (bDisplay)
-			{
-				const double Slope = Segment.ComputeOrientedSlopeOf(MeshNodeCoordinate, Projection);
-				F3DDebugSession A(FString::Printf(TEXT("Close Point slope %f %f %f"), Slope, SlopeS, SlopeE));
-				DisplaySegmentWithScale(Segment.GetExtemity(Start), Segment.GetExtemity(End), 0, EVisuProperty::BlueCurve);
-				DisplaySegmentWithScale(MeshNodeCoordinate, Projection, 0, EVisuProperty::GreenCurve);
-				DisplayPoint2DWithScale(Projection, EVisuProperty::YellowPoint, CrossZoneElement.VertexId);
-				//Wait();
-			}
-#endif
-
 			const double SquareDistance = MeshNodeCoordinate.SquareDistance(Projection);
 			if (MinSquareThickness > SquareDistance)
 			{
@@ -611,25 +468,6 @@ void FParametricFaceMesher::DefineImposedCuttingPointsBasedOnOtherSideMesh(FThin
 
 	Algo::Sort(CrossZoneElements, [](FCrossZoneElement& ElementA, FCrossZoneElement& ElementB) { return ElementA.SquareThickness < ElementB.SquareThickness; });
 
-#ifdef DEBUG_MESHTHINSURF_
-	if (bDisplay)
-	{
-		F3DDebugSession A(TEXT("Result before intersection filtering"));
-		for (FCrossZoneElement& CrossZoneElement : CrossZoneElements)
-		{
-			//F3DDebugSession A(TEXT("Point and Close"));
-			DisplayPoint2DWithScale(CrossZoneElement.VertexPoint2D, EVisuProperty::BluePoint, CrossZoneElement.VertexId);
-			DisplayPoint2DWithScale(CrossZoneElement.OppositePoint2D, EVisuProperty::RedPoint);
-			DisplaySegmentWithScale(CrossZoneElement.VertexPoint2D, CrossZoneElement.OppositePoint2D);
-			if (CrossZoneElement.OppositeSegment)
-			{
-				DisplaySegmentWithScale(CrossZoneElement.OppositeSegment->GetExtemity(ELimit::End), CrossZoneElement.OppositeSegment->GetExtemity(ELimit::Start));
-			}
-		}
-		Wait();
-	}
-#endif
-
 	// Find candidates i.e. CrossZoneElement that are not in intersection with the sides and with selected CrossZoneElement
 	FIntersectionTool IntersectionTool(FrontSide.GetSegments(), SideToConstrain.GetSegments(), CrossZoneElements.Num());
 	for (FCrossZoneElement& CrossZoneElement : CrossZoneElements)
@@ -646,24 +484,6 @@ void FParametricFaceMesher::DefineImposedCuttingPointsBasedOnOtherSideMesh(FThin
 		}
 	}
 
-#ifdef DEBUG_MESHTHINSURF
-	if (bDisplay)
-	{
-		F3DDebugSession M(TEXT("Result final"));
-		for (FCrossZoneElement& CrossZoneElement : CrossZoneElements)
-		{
-			if (CrossZoneElement.bIsSelected)
-			{
-				//F3DDebugSession A(TEXT("Point and Close"));
-				DisplayPoint2DWithScale(CrossZoneElement.VertexPoint2D, EVisuProperty::BluePoint, CrossZoneElement.VertexId);
-				DisplayPoint2DWithScale(CrossZoneElement.OppositePoint2D, EVisuProperty::RedPoint);
-				DisplaySegmentWithScale(CrossZoneElement.VertexPoint2D, CrossZoneElement.OppositePoint2D, 0, EVisuProperty::YellowCurve);
-			}
-		}
-		Wait(false);
-	}
-#endif
-
 	// Add ImposedCuttingPointU
 	for (FCrossZoneElement& CrossZoneElement : CrossZoneElements)
 	{
@@ -685,10 +505,6 @@ void FParametricFaceMesher::DefineImposedCuttingPointsBasedOnOtherSideMesh(FThin
 
 void FParametricFaceMesher::MeshFaceLoops()
 {
-#ifdef DEBUG_MESHTHINSURF
-	F3DDebugSession M(bDisplay, ("Mesh other edges"));
-#endif
-
 	for (const TSharedPtr<FTopologicalLoop>& Loop : Face.GetLoops())
 	{
 		for (const FOrientedEdge& Edge : Loop->GetEdges())
@@ -715,13 +531,11 @@ void FParametricFaceMesher::MeshVerticesOfFace(FTopologicalFace& FaceToProcess)
 	}
 }
 
-#ifdef DEBUG_INTERSECTEDGEISOS
-void DebugIntersectEdgeIsos(const FTopologicalFace& Face, const TArray<double>& IsoCoordinates, EIso TypeIso);
-#endif
-
 void FParametricFaceMesher::Mesh(FTopologicalEdge& InEdge, bool bFinalMeshing)
 {
-  	FTopologicalEdge& ActiveEdge = *InEdge.GetLinkActiveEntity();
+	TRACE_CPUPROFILER_EVENT_SCOPE(UE::CADKernel::FParametricFaceMesher::Mesh)
+
+	FTopologicalEdge& ActiveEdge = *InEdge.GetLinkActiveEntity();
 	if (ActiveEdge.IsMeshed())
 	{
 		if (ActiveEdge.GetMesh()->GetNodeCount() > 0)
@@ -757,60 +571,9 @@ void FParametricFaceMesher::Mesh(FTopologicalEdge& InEdge, bool bFinalMeshing)
 
 	ApplyEdgeCriteria(InEdge);
 
-#ifdef DEBUG_MESH_EDGE
-	if (bDisplay)
-	{
-		F3DDebugSession _(FString::Printf(TEXT("EdgePointsOnDomain %d"), InEdge.GetId()));
-		Display2D(InEdge);
-		Wait();
-	}
-#endif
-
-#ifdef DEBUG_INTERSECTEDGEISOS
-	DebugIntersectEdgeIsos(Face, SurfaceTabU, EIso::IsoU);
-	DebugIntersectEdgeIsos(Face, SurfaceTabV, EIso::IsoV);
-	{
-		F3DDebugSession _(FString::Printf(TEXT("Edge 2D %d"), InEdge.GetId()));
-		UE::CADKernel::Display2D(InEdge);
-	}
-#endif
-
 	InEdge.ComputeIntersectionsWithIsos(SurfaceTabU, EIso::IsoU, ToleranceIso, EdgeIntersectionWithIsoU_Coordinates);
 	InEdge.ComputeIntersectionsWithIsos(SurfaceTabV, EIso::IsoV, ToleranceIso, EdgeIntersectionWithIsoV_Coordinates);
 
-#ifdef DEBUG_INTERSECTEDGEISOS
-	{
-		F3DDebugSession _(FString::Printf(TEXT("Edge %d Intersect with iso"), InEdge.GetId()));
-		TArray<FPoint2D> Intersections;
-		InEdge.Approximate2DPoints(EdgeIntersectionWithIsoU_Coordinates, Intersections);
-		for (const FPoint2D& Point : Intersections)
-		{
-			DisplayPoint(Point);
-		}
-		Intersections.Empty();
-		InEdge.Approximate2DPoints(EdgeIntersectionWithIsoV_Coordinates, Intersections);
-		for (const FPoint2D& Point : Intersections)
-		{
-			DisplayPoint(Point);
-		}
-		Wait();
-	}
-	{
-		F3DDebugSession _(FString::Printf(TEXT("Thin Zone")));
-		for (FLinearBoundary ThinZone : InEdge.GetThinZoneBounds())
-		{
-			TArray<double> Coords;
-			Coords.Add(ThinZone.GetMin());
-			Coords.Add(ThinZone.GetMax());
-			TArray<FPoint2D> ThinZone2D;
-			InEdge.Approximate2DPoints(Coords, ThinZone2D);
-			DisplaySegment(ThinZone2D[0], ThinZone2D[1], EVisuProperty::YellowCurve);
-		}
-		Wait();
-
-	}
-#endif
-
 	FLinearBoundary EdgeBounds = InEdge.GetBoundary();
 
 	TArray<double>& DeltaUs = InEdge.GetDeltaUMaxs();
@@ -1021,20 +784,6 @@ void FParametricFaceMesher::Mesh(FTopologicalEdge& InEdge, bool bFinalMeshing)
 		FinalEdgeCuttingPointCoordinates.Empty(ImposedIsoCuttingPoints.Num() + MaxNumberOfVertex);
 	}
 
-#ifdef DEBUG_GETPREFERREDUVCOORDINATESFROMNEIGHBOURS
-	TArray<double> CuttingPoints2;
-	{
-		double ToleranceGeoEdge = FMath::Min(ExtremityTolerances[0], ExtremityTolerances[1]);
-
-		TArray<FCuttingPoint> Extremities;
-		Extremities.Reserve(2);
-		Extremities.Emplace(EdgeBounds.GetMin(), ECoordinateType::VertexCoordinate, -1, ToleranceGeoEdge);
-		Extremities.Emplace(EdgeBounds.GetMax(), ECoordinateType::VertexCoordinate, -1, ToleranceGeoEdge);
-
-		FMesherTools::ComputeFinalCuttingPointsWithImposedCuttingPoints(InEdge.GetCrossingPointUs(), InEdge.GetDeltaUMaxs(), Extremities, CuttingPoints2);
-	}
-#endif
-
 	if (InEdge.IsDegenerated() || InEdge.IsVirtuallyMeshed())
 	{
 		if (ImposedIsoCuttingPoints.Num() == 2)
@@ -1071,53 +820,6 @@ void FParametricFaceMesher::Mesh(FTopologicalEdge& InEdge, bool bFinalMeshing)
 		}
 	}
 
-#ifdef DEBUG_GETPREFERREDUVCOORDINATESFROMNEIGHBOURS
-	if (InEdge.IsThinZone() && EdgeImposedCuttingPoints.Num())
-	{
-		F3DDebugSession G(TEXT("Mesh(TSharedRef<FEdge> InEdge"));
-		{
-			F3DDebugSession G(TEXT("U From Iso"));
-			for (const FCuttingPoint& CuttingU : ImposedIsoCuttingPoints)
-			{
-				if (CuttingU.OppositNodeIndex >= 0)
-				{
-					UE::CADKernel::DisplayPoint(FPoint2D(CuttingU.Coordinate, 0.0), EVisuProperty::RedPoint);
-				}
-				else
-				{
-					UE::CADKernel::DisplayPoint(FPoint2D(CuttingU.Coordinate, 0.0));
-				}
-			}
-		}
-		{
-			F3DDebugSession _(InEdge.GetThinZoneBounds().Num() != 0, FString::Printf(TEXT("Thin Zone"), InEdge.GetId()));
-			for (FLinearBoundary ThinZone : InEdge.GetThinZoneBounds())
-			{
-				UE::CADKernel::DisplayPoint(FPoint2D(ThinZone.GetMin(), 0.01), EVisuProperty::BluePoint);
-				UE::CADKernel::DisplayPoint(FPoint2D(ThinZone.GetMax(), 0.01), EVisuProperty::BluePoint);
-				DisplaySegment(FPoint2D(ThinZone.GetMin(), 0.01), FPoint2D(ThinZone.GetMax(), 0.01), EVisuProperty::BlueCurve);
-			}
-			Wait();
-		}
-
-		{
-			F3DDebugSession G(TEXT("U From Criteria"));
-			for (double CuttingU : CuttingPoints2)
-			{
-				UE::CADKernel::DisplayPoint(FPoint2D(CuttingU, 0.02), EVisuProperty::RedPoint);
-			}
-		}
-		{
-			F3DDebugSession G(TEXT("U Final (Criteria & Iso)"));
-			for (double CuttingU : CuttingPoints)
-			{
-				UE::CADKernel::DisplayPoint(FPoint2D(CuttingU, 0.04), EVisuProperty::YellowPoint);
-			}
-		}
-		Wait();
-	}
-#endif
-
 	if (bFinalMeshing)
 	{
 		InEdge.GenerateMeshElements(MeshModel);
@@ -1127,47 +829,6 @@ void FParametricFaceMesher::Mesh(FTopologicalEdge& InEdge, bool bFinalMeshing)
 		ActiveEdge.SetPreMeshedMarker();
 		InEdge.SetPreMeshedMarker();
 	}
-
-#ifdef DEBUG_THIN_ZONES
-	if (bFinalMeshing && bDisplay)
-	{
-#ifdef DEBUG_MESHTHINSURF
-		F3DDebugSession A(bDisplay, FString::Printf(TEXT("Edge %d"), InEdge.GetId()));
-#endif
-		Display2DWithScale(InEdge);
-		TArray<FPoint> MeshVertices;
-		if (InEdge.GetMesh())
-		{
-			MeshVertices = InEdge.GetMesh()->GetNodeCoordinates();
-		}
-		if(InEdge.GetStartVertex()->GetMesh())
-		{
-			MeshVertices.Append(InEdge.GetStartVertex()->GetMesh()->GetNodeCoordinates());
-		}
-		if(InEdge.GetEndVertex()->GetMesh())
-		{
-			MeshVertices.Append(InEdge.GetEndVertex()->GetMesh()->GetNodeCoordinates());
-		}
-		if (MeshVertices.Num())
-		{
-			TArray<double> CoordinatesOfMesh;
-			TArray<FPoint> ProjectedPoints;
-			TArray<FPoint2D> Mesh2DPoints;
-			InEdge.ProjectPoints(MeshVertices, CoordinatesOfMesh, ProjectedPoints);
-			Algo::Sort(CoordinatesOfMesh);
-			InEdge.Approximate2DPoints(CoordinatesOfMesh, Mesh2DPoints);
-			int32 Index = 0;
-			for (FPoint2D Point : Mesh2DPoints)
-			{
-				DisplayPoint2DWithScale(Point, EVisuProperty::RedPoint, ++Index);
-			}
-		}
-#ifdef DEBUG_MESHTHINSURF
-		//Wait(bDisplay);
-#endif
-	}
-#endif
-
 }
 
 void FParametricFaceMesher::ApplyEdgeCriteria(FTopologicalEdge& Edge)
@@ -1221,9 +882,6 @@ void FParametricFaceMesher::MeshThinZoneSide(FThinZoneSide& Side, bool bFinalMes
 	}
 	Side.ResetProcessedMarker();
 
-#ifdef DEBUG_MESHTHINSURF
-	F3DDebugSession A(bFinalMeshing && bDisplay, TEXT("SideMesh"));
-#endif
 	for (FTopologicalEdge* Edge : Side.GetEdges())
 	{
 		if (Edge->IsMeshed())
@@ -1231,155 +889,133 @@ void FParametricFaceMesher::MeshThinZoneSide(FThinZoneSide& Side, bool bFinalMes
 			continue;
 		}
 		Mesh(*Edge, bFinalMeshing);
-#ifdef DEBUG_THIN_ZONES
-		if (bFinalMeshing && bDisplay)
-		{
-#ifdef DEBUG_MESHTHINSURF
-			F3DDebugSession A(bDisplay, FString::Printf(TEXT("Edge %d"), Edge->GetId()));
-#endif
-			Display2DWithScale(*Edge);
-			TArray<FPoint> MeshVertices = Edge->GetMesh()->GetNodeCoordinates();
-			MeshVertices.Append(Edge->GetStartVertex()->GetMesh()->GetNodeCoordinates());
-			MeshVertices.Append(Edge->GetEndVertex()->GetMesh()->GetNodeCoordinates());
-			if (MeshVertices.Num())
-			{
-				TArray<double> CoordinatesOfMesh;
-				TArray<FPoint> ProjectedPoints;
-				TArray<FPoint2D> Mesh2DPoints;
-				Edge->ProjectPoints(MeshVertices, CoordinatesOfMesh, ProjectedPoints);
-				Algo::Sort(CoordinatesOfMesh);
-				Edge->Approximate2DPoints(CoordinatesOfMesh, Mesh2DPoints);
-				int32 Index = 0;
-				for (FPoint2D Point : Mesh2DPoints)
-				{
-					DisplayPoint2DWithScale(Point, EVisuProperty::RedPoint, ++Index);
-				}
-			}
-#ifdef DEBUG_MESHTHINSURF
-					//Wait(bDisplay);
-#endif
-		}
-#endif
 	}
-#ifdef DEBUG_MESHTHINSURF
-	//Wait(bDisplay);
-#endif
 }
 
-#ifdef DEBUG_INTERSECTEDGEISOS
-TMap<int32, int32> SurfaceDrawed;
-bool bDisplayIsoCurve = true;
-
-void DebugIntersectEdgeIsos(const FTopologicalFace& Face, const TArray<double>& IsoCoordinates, EIso TypeIso)
+void FParametricFaceMesher::MeshPlanarFace()
 {
-	if (SurfaceDrawed.Find(Face.GetId()) == nullptr)
-	{
-		SurfaceDrawed.Add(Face.GetId(), 0);
-	}
+	TRACE_CPUPROFILER_EVENT_SCOPE(UE::CADKernel::FParametricFaceMesher::MeshPlanarFace)
+		
+	using namespace UE::Geometry;
+
+	TArray<FGeneralPolygon2d> PolygonsOut;
+	PolygonsOut.Reserve(Face.GetLoops().Num());
+
+	const TArray<TArray<FPoint2D>>& Loops2D = Grid.GetLoops2D(EGridSpace::Default2D);
+	const TArray<TSharedPtr<FTopologicalLoop>>& Loops = Face.GetLoops();
+	ensureCADKernel(Loops.Num() == Loops2D.Num());
 
-	if (bDisplayIsoCurve && SurfaceDrawed[Face.GetId()] < 2)
+	FGeneralPolygon2d* GPolygon = nullptr;
+	for (int32 Index = 0; Index < Loops.Num(); ++Index)
 	{
-		SurfaceDrawed[Face.GetId()]++;
+		const TArray<FPoint2D>& PointList = Loops2D[Index];
+		TArray<FVector2d> VertexList;
+		VertexList.Reserve(PointList.Num());
+		for (const FPoint2D& Point : PointList)
+		{
+			VertexList.Emplace(Point.U, Point.V);
+		}
 
-		FSurfacicBoundary Bounds = Face.GetBoundary();
+		FPolygon2d Polygon2d(VertexList);
 
-		F3DDebugSession _(FString::Printf(TEXT("Iso %s 2D %d"), TypeIso == EIso::IsoU ? TEXT("U") : TEXT("V"), Face.GetId()));
-		if (TypeIso == EIso::IsoU)
+		if (Loops[Index]->IsExternal())
 		{
-			for (double U : IsoCoordinates)
-			{
-				FPoint2D Start(U, Bounds[EIso::IsoV].Min);
-				FPoint2D End(U, Bounds[EIso::IsoV].Max);
-				DisplaySegment(Start, End, 0, EVisuProperty::Iso);
-			}
+			GPolygon = &PolygonsOut.Emplace_GetRef(MoveTemp(Polygon2d));
 		}
-		else
+		else if (ensureCADKernel(GPolygon))
 		{
-			for (double V : IsoCoordinates)
+			if (!ensureCADKernel(GPolygon->AddHole(Polygon2d, /*bCheckContainment =*/ true, /*bCheckOrientation =*/ true)))
 			{
-				FPoint2D Start(Bounds[EIso::IsoU].Min, V);
-				FPoint2D End(Bounds[EIso::IsoU].Max, V);
-				DisplaySegment(Start, End, 0, EVisuProperty::Iso);
+				// #cad_kernel: Log verbose a hole failed to be added
 			}
 		}
 	}
-}
 
-#endif
+	// Abort if there is no polygon to triangulate or the outer is only a segment.
+	if (!ensureCADKernel(PolygonsOut.Num() > 0 && PolygonsOut[0].GetOuter().GetVertices().Num() > 2))
+	{
+		// #cadkernel_check: Why this is happening?
+		return;
+	}
 
-#ifdef CADKERNEL_DEBUG
+	FFaceMesh& FaceMesh = Face.GetOrCreateMesh(MeshModel);
 
-void FParametricFaceMesher::DisplayMeshOfFaceLoop()
-{
-	if (bDisplay)
+	TArray<FPoint>& NodeCoordinates = FaceMesh.GetNodeCoordinates();
+	TArray<int32>& TrianglesVerticesIndex = FaceMesh.TrianglesVerticesIndex;
+	TArray<int32>& VerticesGlobalIndex = FaceMesh.VerticesGlobalIndex;
+	TArray<FVector3f>& Normals = FaceMesh.Normals;
+	TArray<FVector2f>& UVMap = FaceMesh.UVMap;
+	
+	// Abort if tessellation is called twice on a Face
+	// #cadkernel_check: Why this is happening?
+	if (TrianglesVerticesIndex.Num() > 0)
 	{
-		F3DDebugSession _(bDisplay, FString::Printf(TEXT("Mesh of Thin Face %d"), Face.GetId()));
-		for (const TSharedPtr<FTopologicalLoop>& Loop : Face.GetLoops())
-		{
-			for (const FOrientedEdge& OrientedEdge : Loop->GetEdges())
-			{
-				FTopologicalEdge& Edge = *OrientedEdge.Entity;
-				if (Edge.GetLinkActiveEdge()->IsMeshed())
-				{
-					const TArray<FPoint>& MeshVetrices = Edge.GetMesh()->GetNodeCoordinates();
-					TArray<FPoint> ProjectedPoints;
-					TArray<double> Coordinates;
-					Edge.ProjectPoints(MeshVetrices, Coordinates, ProjectedPoints);
-					Algo::Sort(Coordinates);
-					Coordinates.Insert(Edge.GetBoundary().GetMin(), 0);
-					Coordinates.Add(Edge.GetBoundary().GetMax());
-					TArray<FPoint2D> Points2D;
-					Edge.Approximate2DPoints(Coordinates, Points2D);
-
-					DisplayPolylineWithScale(Points2D, EVisuProperty::BlueCurve);
-					for (FPoint2D Point : Points2D)
-					{
-						DisplayPoint2DWithScale(Point);
-					}
-				}
-				else
-				{
-					Display2DWithScale(Edge, EVisuProperty::GreenCurve);
-				}
-			}
-		}
-		//Wait();
+		return;
 	}
-}
 
-void FParametricFaceMesher::DisplayThinZoneEdges(const TCHAR* Text, TArray<FTopologicalEdge*>& Edges, EVisuProperty Color, EVisuProperty Color2)
-{
-	if (bDisplay)
+	const int32 GlobalVertexCount = MeshModel.GetVertexCount();
+
+	const FSurface& CarrierSurface = *Face.GetCarrierSurface();
+	FSurfacicPoint SurfacicPoint;
+
+	for (FGeneralPolygon2d& Polygon : PolygonsOut)
 	{
-		F3DDebugSession _(bDisplay, Text);
-		for (const FTopologicalEdge* Edge : Edges)
+		FDelaunay2 Delaunay;
+		Delaunay.bAutomaticallyFixEdgesToDuplicateVertices = false;
+
+		TArray<FIndex3i> Triangles;
+		TArray<FVector2d> Vertices;
+
+		// FDelaunay2::Triangulate can return false although the triangles are valid
+		// In that case, false is returned because some input points were not connected as expected
+		// bFallbackToGeneralizedWinding is set to true to force the generation of triangles
+		// #cadkernel_check: Why those points are not connected?
+		Delaunay.Triangulate(Polygon, &Triangles, &Vertices, true /*bFallbackToGeneralizedWinding*/);
+		if (!ensureCADKernel(Triangles.Num() > 0))
 		{
-			if (Edge->GetLinkActiveEdge()->IsMeshed())
-			{
-				const TArray<FPoint>& MeshVetrices = Edge->GetMesh()->GetNodeCoordinates();
-				TArray<FPoint> ProjectedPoints;
-				TArray<double> Coordinates;
-				Edge->ProjectPoints(MeshVetrices, Coordinates, ProjectedPoints);
-				Algo::Sort(Coordinates);
-				Coordinates.Insert(Edge->GetBoundary().GetMin(), 0);
-				Coordinates.Add(Edge->GetBoundary().GetMax());
-				TArray<FPoint2D> Points2D;
-				Edge->Approximate2DPoints(Coordinates, Points2D);
-
-				DisplayPolylineWithScale(Points2D, Color2);
-				for (FPoint2D Point : Points2D)
-				{
-					DisplayPoint2DWithScale(Point, (EVisuProperty)(Color2 - 1));
-				}
-			}
-			else
-			{
-				Display2DWithScale(*Edge, Color);
-			}
+			continue;
+		}
+
+		const int32 StartIndex = VerticesGlobalIndex.Num();
+		VerticesGlobalIndex.Reserve(StartIndex + Vertices.Num());
+
+		Normals.Reserve(Normals.Num() + Vertices.Num());
+		UVMap.Reserve(UVMap.Num() + Vertices.Num());
+		NodeCoordinates.Reserve(NodeCoordinates.Num() + Vertices.Num());
+
+		int32 Index = StartIndex;
+		for (const FVector2d& Vertex : Vertices)
+		{
+			FPoint Point;
+			FVector3f Normal;
+			CarrierSurface.EvaluatePointAndNormal(FPoint2D(Vertex.X, Vertex.Y), Point, Normal);
+			
+			NodeCoordinates.Emplace(Point);
+			Normals.Emplace(Normal);
+			UVMap.Emplace((FVector2f)Vertex);
+
+			VerticesGlobalIndex.Emplace(Index++ + GlobalVertexCount);
+		}
+
+		TrianglesVerticesIndex.SetNum(TrianglesVerticesIndex.Num() + Triangles.Num() * 3);
+		int32* VertexIndices = TrianglesVerticesIndex.GetData() + StartIndex;
+
+		for (const FIndex3i& TriIndices : Triangles)
+		{
+			VertexIndices[0] = TriIndices.A + StartIndex;
+			VertexIndices[1] = TriIndices.B + StartIndex;
+			VertexIndices[2] = TriIndices.C + StartIndex;
+
+			VertexIndices += 3;
 		}
 	}
-}
 
+	if (Face.IsBackOriented())
+	{
+		FaceMesh.InverseOrientation();
+	}
 
-#endif
-} // namespace
\ No newline at end of file
+	FaceMesh.RegisterCoordinates();
+	MeshModel.AddMesh(FaceMesh);
+}
+} // namespace
diff --git a/Engine/Source/Runtime/Datasmith/CADKernel/Base/Private/Mesh/Meshers/ParametricFaceMesher.h b/Engine/Source/Runtime/Datasmith/CADKernel/Base/Private/Mesh/Meshers/ParametricFaceMesher.h
--- a/Engine/Source/Runtime/Datasmith/CADKernel/Base/Private/Mesh/Meshers/ParametricFaceMesher.h
+++ b/Engine/Source/Runtime/Datasmith/CADKernel/Base/Private/Mesh/Meshers/ParametricFaceMesher.h
@@ -142,10 +142,6 @@ protected:
 
 	FGrid Grid;
 
-#ifdef CADKERNEL_DEV
-	bool bDisplay = false;
-#endif
-
 public:
 
 	FParametricFaceMesher(FTopologicalFace& Face, FModelMesh& InMeshModel, const FMeshingTolerances& InTolerances, bool bActivateThinZoneMeshing);
@@ -157,6 +153,7 @@ private:
 	void Mesh(FTopologicalEdge& InEdge, bool bFinalMeshing = true);
 	void Mesh(FTopologicalVertex& Vertex);
 	void MeshVerticesOfFace(FTopologicalFace& Face);
+	void MeshPlanarFace();
 
 	void MeshFaceLoops();
 
diff --git a/Engine/Source/Runtime/Datasmith/CADKernel/Base/Private/Mesh/Structure/Grid.cpp b/Engine/Source/Runtime/Datasmith/CADKernel/Base/Private/Mesh/Structure/Grid.cpp
--- a/Engine/Source/Runtime/Datasmith/CADKernel/Base/Private/Mesh/Structure/Grid.cpp
+++ b/Engine/Source/Runtime/Datasmith/CADKernel/Base/Private/Mesh/Structure/Grid.cpp
@@ -3,6 +3,7 @@
 #include "Mesh/Structure/Grid.h"
 
 #include "Geo/Sampling/SurfacicSampling.h"
+#include "Geo/Surfaces/SurfaceUtilities.h"
 #include "Mesh/Meshers/MesherTools.h"
 #include "Mesh/Structure/EdgeMesh.h"
 #include "Mesh/Structure/VertexMesh.h"
@@ -21,83 +22,31 @@ FGrid::FGrid(FTopologicalFace& InFace, FModelMesh& InMeshModel)
 	, MinimumElementSize(Tolerance3D * 2.)
 	, MeshModel(InMeshModel)
 {
-#ifdef DEBUG_ONLY_SURFACE_TO_DEBUG
-	bDisplay = (InFace.GetId() == FaceToDebug);
-	Open3DDebugSession(bDisplay, FString::Printf(TEXT("Grid %d"), InFace.GetId()));
-#endif
 }
 
 void FGrid::ProcessPointCloud()
 {
-	FTimePoint StartTime = FChrono::Now();
+	TRACE_CPUPROFILER_EVENT_SCOPE(UE::CADKernel::FGrid::ProcessPointCloud)
 
 	if (!GetMeshOfLoops())
 	{
 		return;
 	}
 
-#ifdef DEBUG_GRID
-	if(bDisplay)
-	{
-		DisplayGridLoops(TEXT("FGrid::Loop 2D"), GetLoops2D(EGridSpace::Default2D), true, true, false);
-		DisplayInnerDomainPoints(TEXT("FGrid::PointCloud 2D"), GetInner2DPoints(EGridSpace::Default2D));
-
-		DisplayGridLoops(TEXT("FGrid::Loop 2D Scaled"), GetLoops2D(EGridSpace::Scaled), true, false, false);
-		DisplayInnerDomainPoints(TEXT("FGrid::PointCloud 2D Scaled"), GetInner2DPoints(EGridSpace::Scaled));
-
-		DisplayGridLoops(TEXT("FGrid::Loop 2D UniformScaled"), GetLoops2D(EGridSpace::UniformScaled), true, true, false);
-		DisplayInnerDomainPoints(TEXT("FGrid::PointCloud 2D UniformScaled"), GetInner2DPoints(EGridSpace::UniformScaled));
-
-		DisplayGridLoops(TEXT("FGrid::Loop 3D"), GetLoops3D(), true, false, false);
-		DisplayInnerDomainPoints(TEXT("FGrid::PointCloud 3D"), GetInner3DPoints());
-
-		DisplayGridLoops(TEXT("FGrid::Loop 3D"), GetLoops3D(), true, false, false);
-		DisplayInnerDomainPoints(TEXT("FGrid::PointCloud 3D"), GetInner3DPoints());
-
-		DisplayGridNormal();
-
-		DisplayInnerDomainPoints(TEXT("FGrid::PointCloud before FindInnerFacePoints"), GetInner2DPoints(EGridSpace::UniformScaled));
-
-		Wait(false);
-	}
-#endif
-
 	FindInnerFacePoints();
 
-#ifdef DEBUG_GRID
-	DisplayInnerDomainPoints(TEXT("FGrid::PointCloud after FindInnerFacePoints"), GetInner2DPoints(EGridSpace::UniformScaled));
-#endif
-
 	FindPointsCloseToLoop();
 
-#ifdef DEBUG_GRID
-	DisplayInnerDomainPoints(TEXT("FGrid::PointCloud after FindPointsCloseToLoop"), GetInner2DPoints(EGridSpace::UniformScaled));
-#endif
-
 	RemovePointsCloseToLoop();
 
-#ifdef DEBUG_GRID
-	DisplayInnerDomainPoints(TEXT("FGrid::PointCloud after RemovePointsCloseToLoop"), GetInner2DPoints(EGridSpace::UniformScaled));
-#endif
-
 	// Removed of Thin zone boundary (the last boundaries). In case of thin zone, the number of 2d boundary will be biggest than 3d boundary one.
 	// Only EGridSpace::UniformScaled is needed.
 	FaceLoops2D[EGridSpace::UniformScaled].SetNum(FaceLoops3D.Num());
-
-#ifdef DEBUG_GRID
-	if (bDisplay)
-	{
-		DisplayGridLoops(TEXT("FGrid::Final Loop 2D"), GetLoops2D(EGridSpace::UniformScaled), true, true, false);
-		Wait(false);
-	}
-#endif
-
-	Chronos.ProcessPointCloudDuration = FChrono::Elapse(StartTime);
 }
 
 void FGrid::DefineCuttingParameters()
 {
-	FTimePoint StartTime = FChrono::Now();
+	TRACE_CPUPROFILER_EVENT_SCOPE(UE::CADKernel::FGrid::DefineCuttingParameters)
 
 	FCuttingGrid PreferredCuttingParametersFromLoops;
 	GetPreferredUVCuttingParametersFromLoops(PreferredCuttingParametersFromLoops);
@@ -106,29 +55,17 @@ void FGrid::DefineCuttingParameters()
 	DefineCuttingParameters(EIso::IsoV, PreferredCuttingParametersFromLoops);
 
 	CuttingSize = CoordinateGrid.Count();
-
-	Chronos.DefineCuttingParametersDuration = FChrono::Elapse(StartTime);
 }
 
 void FGrid::DefineCuttingParameters(EIso Iso, FCuttingGrid& Neighbors)
 {
-	FTimePoint StartTime = FChrono::Now();
+	TRACE_CPUPROFILER_EVENT_SCOPE(UE::CADKernel::FGrid::DefineCuttingParameters)
 
 	const FSurfacicBoundary& Boundary = Face.GetBoundary();
 
-#ifdef DEBUG_GET_PREFERRED_UVCOORDINATES_FROM_NEIGHBOURS
-	TArray<double> CuttingPointTmp;
-#endif
 	if (Neighbors[Iso].Num())
 	{
-#ifdef DEBUG_GET_PREFERRED_UVCOORDINATES_FROM_NEIGHBOURS
-		TArray<FCuttingPoint> Extremities;
-		Extremities.Reserve(2);
-		Extremities.Emplace(Boundary[Iso].Min, ECoordinateType::VertexCoordinate, -1, 0.001);
-		Extremities.Emplace(Boundary[Iso].Max, ECoordinateType::VertexCoordinate, -1, 0.001);
-		FMesherTools::ComputeFinalCuttingPointsWithImposedCuttingPoints(Face.GetCrossingPointCoordinates(Iso), Face.GetCrossingPointDeltaMaxs(Iso), Extremities, CuttingPointTmp);
-#endif
-		FMesherTools::ComputeFinalCuttingPointsWithPreferredCuttingPoints(Face.GetCrossingPointCoordinates(Iso), Face.GetCrossingPointDeltaMaxs(Iso), Neighbors[Iso], Boundary[Iso], Face.GetCuttingCoordinatesAlongIso(Iso));
+		FMesherTools::ComputeFinalCuttingPointsWithPreferredCuttingPoints(Face.GetCrossingPointCoordinates(Iso), Face.GetCrossingPointDeltaMaxs(Iso), Neighbors[Iso], Boundary[Iso], CoordinateGrid[Iso]);
 	}
 	else
 	{
@@ -136,78 +73,25 @@ void FGrid::DefineCuttingParameters(EIso Iso, FCuttingGrid& Neighbors)
 		Extremities.Reserve(2);
 		Extremities.Emplace(Boundary[Iso].Min, ECoordinateType::VertexCoordinate, -1, 0.001);
 		Extremities.Emplace(Boundary[Iso].Max, ECoordinateType::VertexCoordinate, -1, 0.001);
-		FMesherTools::ComputeFinalCuttingPointsWithImposedCuttingPoints(Face.GetCrossingPointCoordinates(Iso), Face.GetCrossingPointDeltaMaxs(Iso), Extremities, Face.GetCuttingCoordinatesAlongIso(Iso));
+		FMesherTools::ComputeFinalCuttingPointsWithImposedCuttingPoints(Face.GetCrossingPointCoordinates(Iso), Face.GetCrossingPointDeltaMaxs(Iso), Extremities, CoordinateGrid[Iso]);
 	}
 
-#ifdef DEBUG_GET_PREFERRED_UVCOORDINATES_FROM_NEIGHBOURS
-	if (bDisplay)
+	// #cadkernel_check: Why does this only apply to planar surfaces?
+	FSurface& CarrierSurface = *Face.GetCarrierSurface();
+	if (SurfaceUtilities::IsPlanar(CarrierSurface))
 	{
-		EIso OtherIso = Other(Iso);
-		F3DDebugSession _(TEXT("GetPreferredUVCoordinatesFromNeighbours"));
-		{
-			F3DDebugSession _(FString::Printf(TEXT("%s From Neighbours"), IsoNames[Iso]));
-			for (FCuttingPoint CuttingU : Neighbors[Iso])
-			{
-				if (Iso == EIso::IsoU)
-				{
-					DisplayPoint(FPoint(CuttingU.Coordinate, Boundary[OtherIso].GetMin() - Boundary[OtherIso].Length() * 1 / 20), EVisuProperty::GreenPoint);
-				}
-				else
-				{
-					DisplayPoint(FPoint(Boundary[OtherIso].GetMin() - Boundary[OtherIso].Length() * 1 / 20, CuttingU.Coordinate), EVisuProperty::GreenPoint);
-				}
-			}
-		}
-		{
-			F3DDebugSession _(FString::Printf(TEXT("%s From Criteria"), IsoNames[Iso]));
-			for (double CuttingU : CuttingPointTmp)
-			{
-				if (Iso == EIso::IsoU)
-				{
-					DisplayPoint(FPoint(CuttingU, Boundary[OtherIso].GetMin() - Boundary[OtherIso].Length() * 1 / 40), EVisuProperty::YellowPoint);
-				}
-				else
-				{
-					DisplayPoint(FPoint(Boundary[OtherIso].GetMin() - Boundary[OtherIso].Length() * 1 / 40, CuttingU), EVisuProperty::YellowPoint);
-				}
-			}
-		}
-		{
-			F3DDebugSession _(FString::Printf(TEXT("%s From Neighbours"), IsoNames[Iso]));
-			for (double CuttingU : Face.GetCuttingCoordinatesAlongIso(Iso))
-			{
-				if (Iso == EIso::IsoU)
-				{
-					DisplayPoint(FPoint(CuttingU, Boundary[OtherIso].GetMin() - Boundary[OtherIso].Length() * 1 / 60), EVisuProperty::BluePoint);
-				}
-				else
-				{
-					DisplayPoint(FPoint(Boundary[OtherIso].GetMin() - Boundary[OtherIso].Length() * 1 / 60, CuttingU), EVisuProperty::BluePoint);
-				}
-			}
-		}
-		//Wait();
+		FCoordinateGrid FaceNotDerivableCoordinates;
+		CarrierSurface.LinesNotDerivables(Face.GetBoundary(), 1, FaceNotDerivableCoordinates);
+
+		ArrayUtils::Complete(CoordinateGrid[Iso], FaceNotDerivableCoordinates[Iso], CarrierSurface.GetIsoTolerance(Iso));
 	}
-#endif
-	CuttingCount[Iso] = CoordinateGrid.IsoCount(Iso);
 
-	Chronos.DefineCuttingParametersDuration = FChrono::Elapse(StartTime);
+	CuttingCount[Iso] = CoordinateGrid.IsoCount(Iso);
 }
 
-//#define DEBUG_GET_PREFERRED_UVCOORDINATES_FROM_NEIGHBOURS
 void FGrid::GetPreferredUVCuttingParametersFromLoops(FCuttingGrid& CuttingParametersFromLoops)
 {
-#ifdef DEBUG_GET_PREFERRED_UVCOORDINATES_FROM_NEIGHBOURS
-	F3DDebugSession _(bDisplay, TEXT("GetPreferredUVCoordinatesFromNeighbours"));
-
-	if (bDisplay)
-	{
-		{
-			F3DDebugSession _(TEXT("Surface 2D"));
-			UE::CADKernel::Display2D(*Face.GetCarrierSurface());
-		}
-	}
-#endif
+	TRACE_CPUPROFILER_EVENT_SCOPE(UE::CADKernel::FGrid::GetPreferredUVCuttingParametersFromLoops)
 
 	int32 nbPoints = 0;
 	for (const TSharedPtr<FTopologicalLoop>& Loop : Face.GetLoops())
@@ -227,14 +111,6 @@ void FGrid::GetPreferredUVCuttingParametersFromLoops(FCuttingGrid& CuttingParame
 		{
 			const TSharedPtr<FTopologicalEdge>& Edge = OrientedEdge.Entity;
 
-#ifdef DEBUG_GET_PREFERRED_UVCOORDINATES_FROM_NEIGHBOURS
-			if (bDisplay)
-			{
-				F3DDebugSession _(FString::Printf(TEXT("Edge %d"), Edge->GetId()));
-				Display2D(*Edge);
-			}
-#endif
-
 			TArray<double> ProjectedPointCoords;
 			TSharedRef<FTopologicalEdge> ActiveEdge = Edge->GetLinkActiveEdge();
 			if (ActiveEdge->IsMeshed())
@@ -267,19 +143,10 @@ void FGrid::GetPreferredUVCuttingParametersFromLoops(FCuttingGrid& CuttingParame
 			TArray<FPoint2D> EdgePoints2D;
 			Edge->Approximate2DPoints(ProjectedPointCoords, EdgePoints2D);
 
-#ifdef DEBUG_GET_PREFERRED_UVCOORDINATES_FROM_NEIGHBOURS
-			F3DDebugSession _(bDisplay, ("Nodes"));
-#endif
 			for (int32 Index = 0; Index < EdgePoints2D.Num(); ++Index)
 			{
 				CuttingParametersFromLoops[EIso::IsoU].Emplace(EdgePoints2D[Index].U, ECoordinateType::OtherCoordinate);
 				CuttingParametersFromLoops[EIso::IsoV].Emplace(EdgePoints2D[Index].V, ECoordinateType::OtherCoordinate);
-#ifdef DEBUG_GET_PREFERRED_UVCOORDINATES_FROM_NEIGHBOURS
-				if (bDisplay)
-				{
-					DisplayPoint(EdgePoints2D[Index]);
-				}
-#endif
 			}
 		}
 	}
@@ -324,7 +191,7 @@ void FGrid::GetPreferredUVCuttingParametersFromLoops(FCuttingGrid& CuttingParame
 
 bool FGrid::GeneratePointCloud()
 {
-	FTimePoint StartTime = FChrono::Now();
+	TRACE_CPUPROFILER_EVENT_SCOPE(UE::CADKernel::FGrid::GeneratePointCloud)
 
 	if (CheckIf2DGridIsDegenerate())
 	{
@@ -342,13 +209,12 @@ bool FGrid::GeneratePointCloud()
 		return false;
 	}
 
-	Chronos.GeneratePointCloudDuration += FChrono::Elapse(StartTime);
 	return true;
 }
 
 void FGrid::FindPointsCloseToLoop()
 {
-	FTimePoint StartTime = FChrono::Now();
+	TRACE_CPUPROFILER_EVENT_SCOPE(UE::CADKernel::FGrid::FindPointsCloseToLoop)
 
 	int32 IndexLoop = 0;
 	int32 Index[2] = { 1, 1 };
@@ -357,28 +223,6 @@ void FGrid::FindPointsCloseToLoop()
 	const FPoint2D* PointA = nullptr;
 	const FPoint2D* PointB = nullptr;
 
-#ifdef DEBUG_FINDPOINTSCLOSETOLOOP
-	if (bDisplay)
-	{
-		Open3DDebugSession(TEXT("FGrid::FindPointsClosedToLoop result"));
-	}
-
-	int32 CellIndex = 0;
-	bool bWaitCell = true;
-	TFunction<void()> DisplayCell = [&]()
-	{
-		if (bDisplay)
-		{
-			F3DDebugSession _(*FString::Printf(TEXT("Cell %d"), CellIndex++));
-			DisplayPoint2DWithScale(Points2D[EGridSpace::UniformScaled][GlobalIndex]);
-			DisplayPoint2DWithScale(Points2D[EGridSpace::UniformScaled][GlobalIndex - 1]);
-			DisplayPoint2DWithScale(Points2D[EGridSpace::UniformScaled][GlobalIndex - 1 - CuttingCount[EIso::IsoU]]);
-			DisplayPoint2DWithScale(Points2D[EGridSpace::UniformScaled][GlobalIndex - CuttingCount[EIso::IsoU]]);
-			Wait(bWaitCell);
-		}
-	};
-#endif
-
 	// Find start index
 	TFunction<void(EIso)> FindPointAIndex = [&](EIso Iso)
 	{
@@ -398,9 +242,6 @@ void FGrid::FindPointsCloseToLoop()
 		SetCloseToLoop(GlobalIndex - 1);
 		SetCloseToLoop(GlobalIndex - 1 - CuttingCount[EIso::IsoU]);
 		SetCloseToLoop(GlobalIndex - CuttingCount[EIso::IsoU]);
-#ifdef DEBUG_FINDPOINTSCLOSETOLOOP
-		DisplayCell();
-#endif
 	};
 
 	TFunction<void(EIso)> Increase = [&](EIso Iso)
@@ -530,14 +371,6 @@ void FGrid::FindPointsCloseToLoop()
 		{
 			PointB = &Loop[BIndex];
 
-#ifdef DEBUG_FINDPOINTSCLOSETOLOOP
-			if (bDisplay)
-			{
-				F3DDebugSession _(TEXT("SEG"));
-				DisplaySegmentWithScale(*PointB, *PointA);
-			}
-#endif
-
 			// Horizontal case
 			if (FindIntersectionIsoStrip(EIso::IsoU))
 			{
@@ -608,21 +441,11 @@ void FGrid::FindPointsCloseToLoop()
 			PointA = PointB;
 		}
 	}
-
-#ifdef DEBUG_FINDPOINTSCLOSETOLOOP
-	if (bDisplay)
-	{
-		Close3DDebugSession();
-		Wait();
-	}
-#endif
-
-	Chronos.FindPointsCloseToLoopDuration += FChrono::Elapse(StartTime);
 }
 
 void FGrid::RemovePointsCloseToLoop()
 {
-	FTimePoint StartTime = FChrono::Now();
+	TRACE_CPUPROFILER_EVENT_SCOPE(UE::CADKernel::FGrid::RemovePointsCloseToLoop)
 
 	struct FGridSegment
 	{
@@ -831,8 +654,6 @@ void FGrid::RemovePointsCloseToLoop()
 			break;
 		}
 	}
-
-	Chronos.RemovePointsClosedToLoopDuration += FChrono::Elapse(StartTime);
 }
 
 /**
@@ -895,6 +716,8 @@ void SlightlyDisplacedPolyline(TArray<FPoint2D>& D2Points, const FSurfacicBounda
 
 void FGrid::GetMeshOfLoop(const FTopologicalLoop& Loop)
 {
+	TRACE_CPUPROFILER_EVENT_SCOPE(UE::CADKernel::FGrid::GetMeshOfLoop)
+
 	int32 LoopNodeCount = 0;
 
 	for (const FOrientedEdge& Edge : Loop.GetEdges())
@@ -919,15 +742,6 @@ void FGrid::GetMeshOfLoop(const FTopologicalLoop& Loop)
 		const TSharedPtr<FTopologicalEdge>& Edge = OrientedEdge.Entity;
 		const TSharedRef<FTopologicalEdge>& ActiveEdge = Edge->GetLinkActiveEdge();
 
-#ifdef DEBUG_GET_MESH_OF_LOOP
-		if (bDisplay)
-		{
-			F3DDebugSession _(bDisplay, FString::Printf(TEXT("Edge %d"), Edge->GetId()));
-			Display2D(*Edge);
-			Wait();
-		}
-#endif
-
 		bool bSameDirection = Edge->IsSameDirection(*ActiveEdge);
 
 		TArray<double> EdgeCuttingPointCoordinates;
@@ -1005,6 +819,22 @@ void FGrid::GetMeshOfLoop(const FTopologicalLoop& Loop)
 				Swap(CuttingPolyline.Coordinates, EdgeCuttingPointCoordinates);
 			}
 
+			// #cadkernel_check: Why does this only applies to planar surfaces
+			const FSurface& CarrierSurface = *OrientedEdge.Entity->GetLoop()->GetFace()->GetCarrierSurface();
+			if (SurfaceUtilities::IsPlanar(CarrierSurface))
+			{
+				// Make sure 'not derivable coordinates' are part of the polyline's coordinates
+				TSharedPtr<FRestrictionCurve> Curve = OrientedEdge.Entity->GetCurve();
+				TArray<double> NotDerivableCoordinates;
+				Curve->FindNotDerivableCoordinates(Curve->GetBoundary(), 1, NotDerivableCoordinates);
+				ArrayUtils::Complete(CuttingPolyline.Coordinates, NotDerivableCoordinates, Curve->GetMinLinearTolerance());
+
+				// Remove duplicates
+				TArray<double> CachedCoordinates(CuttingPolyline.Coordinates);
+				ArrayUtils::RemoveDuplicates(CuttingPolyline.Coordinates, Curve->GetMinLinearTolerance());
+				ensureCADKernel(CuttingPolyline.Coordinates.Num() > 1);
+			}
+			
 			Edge->ApproximatePolyline(CuttingPolyline);
 		}
 
@@ -1046,20 +876,16 @@ void FGrid::GetMeshOfLoop(const FTopologicalLoop& Loop)
 		else
 		{
 			EdgeVerticesIndex = ActiveEdge->GetOrCreateMesh(MeshModel).EdgeVerticesIndex;
-		}
 
-#ifdef DEBUG_GET_MESH_OF_LOOP
-		if (bDisplay)
-		{
-			F3DDebugSession _(bDisplay, FString::Printf(TEXT("Edge %d cutting points on surface"), ActiveEdge->GetId()));
-			for (const FPoint2D& Point2D : CuttingPolyline.Points2D)
+			// #cadkernel_check: Adding the 'not derivable coordinates may introduce new points
+			//		Whys is EdgeVerticesIndex taken from the ActiveEdge?
+			//		Should the EdgeVerticesIndex be regenerated?
+			int32 LastID = EdgeVerticesIndex.Last();
+			while (EdgeVerticesIndex.Num() < CuttingPolyline.Size())
 			{
-				DisplayPoint(Point2D);
+				EdgeVerticesIndex.Add(LastID);
 			}
-			DisplayPolyline(CuttingPolyline.Points2D, EVisuProperty::BlueCurve);
-			Wait();
 		}
-#endif
 
 		if (OrientedEdge.Direction != EOrientation::Front)
 		{
@@ -1074,52 +900,22 @@ void FGrid::GetMeshOfLoop(const FTopologicalLoop& Loop)
 		ensureCADKernel(CuttingPolyline.Size() > 1);
 
 		Loop2D.Append(MoveTemp(CuttingPolyline.Points2D));
-		Loop2D.Pop();
-
-#ifdef DEBUG_GET_MESH_OF_LOOP
-		if (bDisplay)
-		{
-			F3DDebugSession _(bDisplay, FString::Printf(TEXT("Loop with Edge %d"), Edge->GetId()));
-			DisplayPolyline(Loop2D, EVisuProperty::BlueCurve);
-			Wait();
-		}
-#endif
-
-#ifdef CADKERNEL_DEBUG_
-		const FPoint2D* PrevPoint = &Loop2D.Last();
-		if(Loop2D.Num() > 1)
-		{
-			for (const FPoint2D& Point : Loop2D)
-			{
-				double Dist = PrevPoint->Distance(Point);
-				if(Dist > 0.0001)
-				{
-					ensureCADKernel(false);
-				}
-				PrevPoint = &Point;
-			}
-		}
-#endif
+		// Ignore last added vertex as it is equal to first of next edge in loop
+		Loop2D.SetNum(Loop2D.Num() - 1, EAllowShrinking::No); 
 
 		int32 LastIndex = Loop3D.Num();
 		Loop3D.Append(MoveTemp(CuttingPolyline.Points3D));
 		Loop3D[LastIndex] = (ActiveEdge->GetStartVertex((OrientedEdge.Direction == EOrientation::Front) == bSameDirection)->GetLinkActiveEntity()->GetBarycenter());
-		Loop3D.Pop();
+		// Ignore last added vertex as it is equal to first of next edge in loop
+		Loop3D.SetNum(Loop3D.Num() - 1, EAllowShrinking::No);
 
 		LoopNormals.Append(MoveTemp(CuttingPolyline.Normals));
-		LoopNormals.Pop();
+		// Ignore last added normal as it is equal to first of next edge in loop
+		LoopNormals.SetNum(LoopNormals.Num() - 1, EAllowShrinking::No);
 
 		LoopIds.Append(MoveTemp(EdgeVerticesIndex));
-		LoopIds.Pop();
-
-#ifdef DEBUG_GET_MESH_OF_LOOP
-		static int32 GetMeshOfLoopIter = 0;
-		GetMeshOfLoopIter++;
-		if (LoopIds.Num() != Loop2D.Num() || GetMeshOfLoopIter == 2991)
-		{
-			printf("");
-		}
-#endif
+		// Ignore last added index as it is equal to first of next edge in loop
+		LoopIds.SetNum(LoopIds.Num() - 1, EAllowShrinking::No);
 	}
 
 	if (Loop2D.Num() < 3) // degenerated loop
@@ -1133,6 +929,8 @@ void FGrid::GetMeshOfLoop(const FTopologicalLoop& Loop)
 
 void FGrid::GetMeshOfThinZone(const FThinZone2D& ThinZone)
 {
+	TRACE_CPUPROFILER_EVENT_SCOPE(UE::CADKernel::FGrid::GetMeshOfThinZone)
+
 	// ThinZones are identified during "ApplyCriteria". "ApplyCriteria" use the FCriteriaGrid, The EGridSpace::UniformScaled are not the same between FCriteriaGrid and FGrid
 	// This is the reason why, to get the ThinZone in the FGrid UniformScaled space, we need to get the mesh of the thinZone defined by the node ids, 
 	// And then, from this, we retrieve the thinZone points in Grid UniformScaled space.
@@ -1282,6 +1080,8 @@ void FGrid::GetMeshOfThinZone(const FThinZone2D& ThinZone)
 
 bool FGrid::GetMeshOfLoops()
 {
+	TRACE_CPUPROFILER_EVENT_SCOPE(UE::CADKernel::FGrid::GetMeshOfLoops)
+
 	int32 LoopCount = Face.GetLoops().Num();
 	FaceLoops2D[EGridSpace::Default2D].Reserve(LoopCount);
 
@@ -1289,10 +1089,6 @@ bool FGrid::GetMeshOfLoops()
 	NormalsOfFaceLoops.Reserve(LoopCount);
 	NodeIdsOfFaceLoops.Reserve(LoopCount);
 
-#ifdef DEBUG_GET_BOUNDARY_MESH
-	F3DDebugSession _(bDisplay, ("GetLoopMesh"));
-#endif
-
 	// Outer loops is processed first
 	for (const TSharedPtr<FTopologicalLoop>& Loop : Face.GetLoops())
 	{
@@ -1320,23 +1116,10 @@ bool FGrid::GetMeshOfLoops()
  
 	if (Face.HasThinZone())
 	{
-#ifdef DEBUG_GET_BOUNDARY_MESH
-		F3DDebugSession _(bDisplay, ("Thin zone loop"));
-#endif
 		for (const FThinZone2D& ThinZone : Face.GetThinZones())
 		{
 			GetMeshOfThinZone(ThinZone);
-#ifdef DEBUG_GET_BOUNDARY_MESH
-			if(bDisplay)
-			{
-				DisplayGridLoop(TEXT("ThinZone"), FaceLoops2D[EGridSpace::UniformScaled][FaceLoops2D[EGridSpace::UniformScaled].Num() - 1], true, true, false);
-				Wait();
-			}
-#endif
 		}
-#ifdef DEBUG_GET_BOUNDARY_MESH
-		//Wait();
-#endif
 	}
 
 	// Fit loops to Surface bounds.
@@ -1427,14 +1210,16 @@ bool FGrid::CheckIf2DGridIsDegenerate() const
 
 void FGrid::FindInnerFacePoints()
 {
+	TRACE_CPUPROFILER_EVENT_SCOPE(UE::CADKernel::FGrid::FindInnerFacePoints)
+
+	// #cadkernel_check: This could use the TPolygon class???
+
 	// FindInnerDomainPoints: Inner Points <-> bIsOfInnerDomain = true
 	// For each points count the number of intersection with the boundary in the four directions U+ U- V+ V-
 	// It for each the number is pair, the point is outside,
 	// If in 3 directions the point is inner, the point is inner else we have a doubt so it preferable to consider it outside. 
 	// Most of the time, there is a doubt if the point is to close of the boundary. So it will be removed be other criteria
 
-	FTimePoint StartTime = FChrono::Now();
-
 	TFunction<void(TArray<char>&, int32)> AddIntersection = [](TArray<char>& Intersect, int32 Index)
 	{
 		Intersect[Index] = Intersect[Index] > 0 ? 0 : 1;
@@ -1530,10 +1315,6 @@ void FGrid::FindInnerFacePoints()
 		}
 	}
 
-#ifdef DEBUG_FIND_INNER_FACE_POINTS
-	DisplayGridLoops(TEXT("FGrid::Loop 2D After move according tol"), GetLoops2D(EGridSpace::UniformScaled), true, false);
-#endif
-
 	// Intersection along U axis
 	for (const TArray<FPoint2D>& Loop : FaceLoops2D[EGridSpace::UniformScaled])
 	{
@@ -1741,8 +1522,6 @@ void FGrid::FindInnerFacePoints()
 			CountOfInnerNodes--;
 		}
 	}
-
-	Chronos.FindInnerDomainPointsDuration += FChrono::Elapse(StartTime);
 }
 
 bool FGrid::CheckIfExternalLoopIsDegenerate() const
@@ -1766,4 +1545,3 @@ bool FGrid::CheckIfExternalLoopIsDegenerate() const
 
 } // NS CADKernel
 
-
diff --git a/Engine/Source/Runtime/Datasmith/CADKernel/Base/Private/Mesh/Structure/Grid.h b/Engine/Source/Runtime/Datasmith/CADKernel/Base/Private/Mesh/Structure/Grid.h
--- a/Engine/Source/Runtime/Datasmith/CADKernel/Base/Private/Mesh/Structure/Grid.h
+++ b/Engine/Source/Runtime/Datasmith/CADKernel/Base/Private/Mesh/Structure/Grid.h
@@ -51,7 +51,7 @@ protected:
 	/*
 	 * Cutting coordinates of the face respecting the meshing criteria
 	 */
-	const FCoordinateGrid& CoordinateGrid;
+	FCoordinateGrid& CoordinateGrid;
 
 	virtual const FCoordinateGrid& GetCoordinateGrid() const override
 	{
diff --git a/Engine/Source/Runtime/Datasmith/CADKernel/Base/Private/Topo/TopologicalEdge.cpp b/Engine/Source/Runtime/Datasmith/CADKernel/Base/Private/Topo/TopologicalEdge.cpp
--- a/Engine/Source/Runtime/Datasmith/CADKernel/Base/Private/Topo/TopologicalEdge.cpp
+++ b/Engine/Source/Runtime/Datasmith/CADKernel/Base/Private/Topo/TopologicalEdge.cpp
@@ -410,13 +410,18 @@ void FTopologicalEdge::ComputeCrossingPointCoordinates()
 		// #cadkernel_check: Why could this happen? Shouldn't it be detected way earlier?
 		return;
 	}
+
 	double Tolerance = GetTolerance3D();
 
-	FSurfacicPolyline Presampling;
-	FSurfacicCurveSamplerOnParam Sampler(*Curve, Boundary, Tolerance * 10., Tolerance, Presampling);
-	Sampler.Sample();
+	{
+		FSurfacicPolyline Presampling;
+		FSurfacicCurveSamplerOnParam Sampler(*Curve, Boundary, Tolerance * 10., Tolerance, Presampling);
+		Sampler.Sample();
 
-	Presampling.SwapCoordinates(CrossingPointUs);
+		Presampling.SwapCoordinates(CrossingPointUs);
+	}
+
+	ensureCADKernel(FMath::IsNearlyEqual(CrossingPointUs.Last(), Boundary.GetMax(), Curve->GetMinLinearTolerance()));
 	// #cadkernel_check: To investigate - It looks like the sampler does not start and end the sample data with the boundary values???
 	if (!FMath::IsNearlyEqual(CrossingPointUs.Last(), Boundary.GetMax(), UE_DOUBLE_SMALL_NUMBER))
 	{
@@ -769,20 +774,6 @@ void FTopologicalEdge::GenerateMeshElements(FModelMesh& MeshModel)
 		Algo::Reverse(Coordinates);
 	}
 
-#ifdef DEBUG_MESH_EDGE
-	if (bDisplay)
-	{
-		F3DDebugSession _(FString::Printf(TEXT("Edge Mesh %d"), Edge.GetId()));
-		TArray<FPoint2D> Mesh2D;
-		Edge.Approximate2DPoints(CuttingPointCoordinates, Mesh2D);
-		for (const FPoint2D& Vertex : Mesh2D)
-		{
-			DisplayPoint(Vertex, EVisuProperty::RedPoint);
-		}
-		Wait();
-	}
-#endif
-
 	EdgeMesh.RegisterCoordinates();
 	EdgeMesh.Mesh(StartVertexNodeIndex, EndVertexNodeIndex);
 	MeshModel.AddMesh(EdgeMesh);
@@ -855,6 +846,7 @@ FEdgeMesh& FTopologicalEdge::GetOrCreateMesh(FModelMesh& ShellMesh)
 	{
 		Mesh = FEntity::MakeShared<FEdgeMesh>(ShellMesh, *this);
 	}
+
 	return *Mesh;
 }
 
diff --git a/Engine/Source/Runtime/Datasmith/CADKernel/Base/Private/Topo/TopologicalFaceUtilities.cpp b/Engine/Source/Runtime/Datasmith/CADKernel/Base/Private/Topo/TopologicalFaceUtilities.cpp
new file mode 100644
--- /dev/null
+++ b/Engine/Source/Runtime/Datasmith/CADKernel/Base/Private/Topo/TopologicalFaceUtilities.cpp
@@ -0,0 +1,352 @@
+// Copyright Epic Games, Inc. All Rights Reserved.
+#include "Topo/TopologicalFaceUtilities.h"
+
+#include "Geo/Curves/CurveUtilities.h"
+#include "Geo/GeoEnum.h"
+#include "Geo/Surfaces/SurfaceUtilities.h"
+
+#include "Topo/TopologicalFace.h"
+#include "Topo/TopologicalEdge.h"
+#include "Topo/TopologicalLoop.h"
+
+#include "Algo/Reverse.h"
+#include "Algo/BinarySearch.h"
+
+namespace UE::CADKernel::TopologicalFaceUtilities
+{
+	bool IsPlanar(const UE::CADKernel::FTopologicalFace& Face)
+	{
+		using namespace UE::CADKernel;
+
+		bool bIsPlanar = true;
+
+		if (SurfaceUtilities::IsPlanar(*Face.GetCarrierSurface()))
+		{
+			for (const TSharedPtr<FTopologicalLoop>& Loop : Face.GetLoops())
+			{
+				for (const FOrientedEdge& Edge : Loop->GetEdges())
+				{
+					if (CurveUtilities::GetDegree(*Edge.Entity->GetCurve()) == 1)
+					{
+						if (!Edge.Entity->GetTwinEdge())
+						{
+							continue;
+						}
+						else if (CurveUtilities::GetDegree(*Edge.Entity->GetTwinEdge()->GetCurve()) == 1)
+						{
+							continue;
+						}
+					}
+
+					bIsPlanar = false;
+					break;
+				}
+
+				if (!bIsPlanar)
+				{
+					return bIsPlanar;
+				}
+			}
+		}
+
+		return bIsPlanar;
+	}
+	
+	TArray<FVector2d> Get2DPolyline(const UE::CADKernel::FTopologicalEdge& Edge)
+	{
+		using namespace UE::CADKernel;
+
+		const FLinearBoundary& EdgeBounds = Edge.GetBoundary();
+		const TArray<double>& PolyCoords = Edge.GetCurve()->GetPolyline().GetCoordinates();
+		const TArray<FPoint2D>& PolyPoints = Edge.GetCurve()->GetPolyline().Get2DPoints();
+		const FLinearBoundary PolyBounds(PolyCoords[0] - UE_DOUBLE_SMALL_NUMBER, PolyCoords.Last() + UE_DOUBLE_SMALL_NUMBER);
+
+		ensureCADKernel(PolyPoints.Num() > 1);
+
+		TArray<FVector2d> PointsOut;
+		PointsOut.Reserve(PolyPoints.Num());
+
+
+		if (FMath::IsNearlyEqual(EdgeBounds.Min, PolyCoords[0]) && FMath::IsNearlyEqual(EdgeBounds.Max, PolyCoords.Last()))
+		{
+			// The entire polyline is used for the edge, just copy it in
+			for (const FPoint2D& Point : PolyPoints)
+			{
+				PointsOut.Emplace(Point.U, Point.V);
+			}
+		}
+		else
+		{
+			TFunction<void(const double&)> InsertPoint = [&](const double& Value) -> void
+				{
+					if (Value < PolyBounds.Min)
+					{
+						const double Ratio = (Value - PolyCoords[0]) / (PolyCoords[1] - PolyCoords[0]);
+						const FPoint2D Point = PolyPoints[0] + (PolyPoints[1] - PolyPoints[0]) * Ratio;
+
+						PointsOut.Emplace(Point.U, Point.V);
+					}
+					else if (Value > PolyBounds.Max)
+					{
+						const int32 LastIndex = PolyCoords.Num() - 1;
+						const double Ratio = (Value - PolyCoords[LastIndex]) / (PolyCoords[LastIndex] - PolyCoords[LastIndex - 1]);
+						const FPoint2D Point = PolyPoints[LastIndex] + (PolyPoints[LastIndex] - PolyPoints[LastIndex - 1]) * Ratio;
+
+						PointsOut.Emplace(Point.U, Point.V);
+					}
+					else
+					{
+						int32 Index = Algo::LowerBound(PolyCoords, Value);
+
+						if (FMath::IsNearlyEqual(Value, PolyCoords[Index]))
+						{
+							const FPoint2D& Point = PolyPoints[Index];
+							PointsOut.Emplace(Point.U, Point.V);
+						}
+						else
+						{
+							const double Ratio = (Value - PolyCoords[Index - 1]) / (PolyCoords[Index] - PolyCoords[Index - 1]);
+							const FPoint2D Point = PolyPoints[Index - 1] + (PolyPoints[Index] - PolyPoints[Index - 1]) * Ratio;
+
+							PointsOut.Emplace(Point.U, Point.V);
+						}
+					}
+				};
+
+
+			InsertPoint(EdgeBounds.Min);
+			int32 StartIndex = Algo::LowerBound(PolyCoords, EdgeBounds.Min);
+			if (StartIndex == 0)
+			{
+				StartIndex++;
+			}
+			const int32 EndIndex = Algo::LowerBound(PolyCoords, EdgeBounds.Max);
+			ensureCADKernel(EndIndex < PolyCoords.Num());
+			for (int32 Index = StartIndex; Index < EndIndex; ++Index)
+			{
+				PointsOut.Emplace(PolyPoints[Index].U, PolyPoints[Index].V);
+			}
+
+			InsertPoint(EdgeBounds.Max);
+		}
+
+		return PointsOut;
+	}
+
+	TArray<FVector2d> Get2DPolyline(const UE::CADKernel::FOrientedEdge& Edge)
+	{
+		TArray<FVector2d> Points = Get2DPolyline(*Edge.Entity);
+
+		if (Points.Num() < 2)
+		{
+			Points.Reset();
+			return Points;
+		}
+
+		if (Edge.Direction == EOrientation::Back)
+		{
+			if (Points.Num() == 2)
+			{
+				FVector2d Reserve = Points[0];
+				Points[0] = Points[1];
+				Points[1] = Reserve;
+			}
+			else
+			{
+				Algo::Reverse(Points);
+			}
+		}
+
+		return Points;
+	}
+
+	TArray<FVector2d> Get2DPolyline(const UE::CADKernel::FTopologicalLoop& Loop)
+	{
+		using namespace UE::CADKernel;
+
+		const FSurfacicTolerance& IsoTolerances = Loop.GetFace()->GetCarrierSurface()->GetIsoTolerances();
+		const FVector2d Tolerance2D { IsoTolerances.U, IsoTolerances.V };
+
+		TArray<FVector2d> VertexList;
+		// #cad_debug
+		FVector2d Last = FVector2d::ZeroVector;
+
+		int32 Index = 0;
+		for (const FOrientedEdge& Edge : Loop.GetEdges())
+		{
+			// For the time being, assuming all curves a 2D ones
+			TArray<FVector2d> Polyline = Get2DPolyline(Edge);
+			TArray<FVector> Polyline3D = Get3DPolyline(Edge);
+			if (Polyline.Num() > 1)
+			{
+				// #cad_debug
+				{
+					if (VertexList.Num() > 0)
+					{
+						ensureCADKernel(Tolerance2D.ComponentwiseAllGreaterOrEqual(Last - Polyline[0]));
+					}
+					Last = Polyline.Last();
+				}
+				VertexList.Append(MoveTemp(Polyline));
+				VertexList.SetNum(VertexList.Num() - 1, EAllowShrinking::No);
+			}
+			++Index;
+		}
+
+		if (ensureCADKernel(VertexList.Num() > 1))
+		{
+			ensureCADKernel(Tolerance2D.ComponentwiseAllGreaterOrEqual(Last - VertexList[0]));
+		}
+
+		return VertexList;
+	}
+
+	TArray<FVector> Get3DPolyline(const UE::CADKernel::FTopologicalEdge& Edge)
+	{
+		using namespace UE::CADKernel;
+
+		const FLinearBoundary& EdgeBounds = Edge.GetBoundary();
+		const TArray<double>& PolyCoords = Edge.GetCurve()->GetPolyline().GetCoordinates();
+		const TArray<FPoint>& PolyPoints = Edge.GetCurve()->GetPolyline().GetPoints();
+		const FLinearBoundary PolyBounds(PolyCoords[0] - UE_DOUBLE_SMALL_NUMBER, PolyCoords.Last() + UE_DOUBLE_SMALL_NUMBER);
+
+		ensureCADKernel(PolyPoints.Num() > 1);
+
+		TArray<FVector> PointsOut;
+		PointsOut.Reserve(PolyPoints.Num());
+
+		if (FMath::IsNearlyEqual(EdgeBounds.Min, PolyCoords[0]) && FMath::IsNearlyEqual(EdgeBounds.Max, PolyCoords.Last()))
+		{
+			// The entire polyline is used for the edge, just copy it in
+			for (const FPoint& Point : PolyPoints)
+			{
+				PointsOut.Emplace(Point.X, Point.Y, Point.Z);
+			}
+		}
+		else
+		{
+			TFunction<void(const double&)> InsertPoint = [&](const double& Value) -> void
+				{
+					if (Value < PolyBounds.Min)
+					{
+						const double Ratio = (Value - PolyCoords[0]) / (PolyCoords[1] - PolyCoords[0]);
+						const FPoint Point = PolyPoints[0] + (PolyPoints[1] - PolyPoints[0]) * Ratio;
+
+						PointsOut.Emplace(Point.X, Point.Y, Point.Z);
+					}
+					else if (Value > PolyBounds.Max)
+					{
+						const int32 LastIndex = PolyCoords.Num() - 1;
+						const double Ratio = (Value - PolyCoords[LastIndex]) / (PolyCoords[LastIndex] - PolyCoords[LastIndex - 1]);
+						const FPoint Point = PolyPoints[LastIndex] + (PolyPoints[LastIndex] - PolyPoints[LastIndex - 1]) * Ratio;
+
+						PointsOut.Emplace(Point.X, Point.Y, Point.Z);
+					}
+					else
+					{
+						int32 Index = Algo::LowerBound(PolyCoords, Value);
+
+						if (FMath::IsNearlyEqual(Value, PolyCoords[Index]))
+						{
+							const FPoint& Point = PolyPoints[Index];
+							PointsOut.Emplace(Point.X, Point.Y, Point.Z);
+						}
+						else
+						{
+							const double Ratio = (Value - PolyCoords[Index - 1]) / (PolyCoords[Index] - PolyCoords[Index - 1]);
+							const FPoint Point = PolyPoints[Index - 1] + (PolyPoints[Index] - PolyPoints[Index - 1]) * Ratio;
+
+							PointsOut.Emplace(Point.X, Point.Y, Point.Z);
+						}
+					}
+				};
+
+
+			InsertPoint(EdgeBounds.Min);
+			int32 StartIndex = Algo::LowerBound(PolyCoords, EdgeBounds.Min);
+			if (StartIndex == 0)
+			{
+				StartIndex++;
+			}
+			const int32 EndIndex = Algo::LowerBound(PolyCoords, EdgeBounds.Max);
+			ensureCADKernel(EndIndex < PolyCoords.Num());
+			for (int32 Index = StartIndex; Index < EndIndex; ++Index)
+			{
+				PointsOut.Emplace(PolyPoints[Index].X, PolyPoints[Index].Y, PolyPoints[Index].Z);
+			}
+
+			InsertPoint(EdgeBounds.Max);
+		}
+
+		printf(">>> %lf", PolyBounds.Min);
+
+		return PointsOut;
+	}
+
+	TArray<FVector> Get3DPolyline(const UE::CADKernel::FOrientedEdge& Edge)
+	{
+		TArray<FVector> Points = Get3DPolyline(*Edge.Entity);
+
+		if (Points.Num() < 2)
+		{
+			Points.Reset();
+			return Points;
+		}
+
+		if (Edge.Direction == EOrientation::Back)
+		{
+			if (Points.Num() == 2)
+			{
+				FVector Reserve = Points[0];
+				Points[0] = Points[1];
+				Points[1] = Reserve;
+			}
+			else
+			{
+				Algo::Reverse(Points);
+			}
+		}
+
+		return Points;
+	}
+
+	TArray<FVector> Get3DPolyline(const UE::CADKernel::FTopologicalLoop& Loop)
+	{
+		using namespace UE::CADKernel;
+
+		const double Tolerance3D = Loop.GetFace()->GetCarrierSurface()->Get3DTolerance();
+
+		TArray<FVector> VertexList;
+		// #cad_debug
+		FVector Last = FVector::ZeroVector;
+
+		int32 Index = 0;
+		for (const FOrientedEdge& Edge : Loop.GetEdges())
+		{
+			// For the time being, assuming all curves a 2D ones
+			TArray<FVector> Polyline = Get3DPolyline(Edge);
+			if (Polyline.Num() > 1)
+			{
+				// #cad_debug
+				{
+					if (VertexList.Num() > 0)
+					{
+						ensureCADKernel(Last.Equals(Polyline[0], Tolerance3D));
+					}
+					Last = Polyline.Last();
+				}
+
+				VertexList.Append(MoveTemp(Polyline));
+				VertexList.SetNum(VertexList.Num() - 1, EAllowShrinking::No);
+			}
+			++Index;
+		}
+
+		if (ensureCADKernel(VertexList.Num() > 1))
+		{
+			ensureCADKernel(Last.Equals(VertexList[0], Tolerance3D));
+		}
+
+		return VertexList;
+	}
+}
diff --git a/Engine/Source/Runtime/Datasmith/CADKernel/Base/Private/UI/Display.cpp b/Engine/Source/Runtime/Datasmith/CADKernel/Base/Private/UI/Display.cpp
--- a/Engine/Source/Runtime/Datasmith/CADKernel/Base/Private/UI/Display.cpp
+++ b/Engine/Source/Runtime/Datasmith/CADKernel/Base/Private/UI/Display.cpp
@@ -744,7 +744,7 @@ void DisplayControlPolygon(const FSurface& Surface)
 	if (Surface.GetSurfaceType() == ESurface::Bezier)
 	{
 		const FBezierSurface& Bezier = (const FBezierSurface&)Surface;
-		DisplayHull(Bezier.GetPoles(), Bezier.GetUDegre() + 1, Bezier.GetVDegre() + 1);
+		DisplayHull(Bezier.GetPoles(), Bezier.GetUDegree() + 1, Bezier.GetVDegree() + 1);
 		return;
 	}
 	if (Surface.GetSurfaceType() == ESurface::Nurbs)
diff --git a/Engine/Source/Runtime/Datasmith/CADKernel/Base/Public/Core/Database.h b/Engine/Source/Runtime/Datasmith/CADKernel/Base/Public/Core/Database.h
--- a/Engine/Source/Runtime/Datasmith/CADKernel/Base/Public/Core/Database.h
+++ b/Engine/Source/Runtime/Datasmith/CADKernel/Base/Public/Core/Database.h
@@ -18,6 +18,7 @@ class CADKERNEL_API FDatabase
 	friend class FCADKernelArchive;
 	friend class FEntity;
 	friend class FSession;
+	template<typename, ESPMode> friend class SharedPointerInternals::TIntrusiveReferenceController;
 
 protected:
 
@@ -78,7 +79,8 @@ public:
 
 	FDatabase();
 
-	FModel& GetModel();
+	FModel& GetModel() { return *GetModelAsShared(); }
+	TSharedPtr<FModel> GetModelAsShared();
 
 	/**
 	 * Remove from database
diff --git a/Engine/Source/Runtime/Datasmith/CADKernel/Base/Public/Core/Session.h b/Engine/Source/Runtime/Datasmith/CADKernel/Base/Public/Core/Session.h
--- a/Engine/Source/Runtime/Datasmith/CADKernel/Base/Public/Core/Session.h
+++ b/Engine/Source/Runtime/Datasmith/CADKernel/Base/Public/Core/Session.h
@@ -18,6 +18,7 @@ class CADKERNEL_API FSession
 {
 	friend FEntity;
 	friend FCADKernelArchive;
+	template<typename, ESPMode> friend class SharedPointerInternals::TIntrusiveReferenceController;
 
 protected:
 
@@ -36,7 +37,8 @@ public:
 		IntersectionTool::SetTolerance(InGeometricTolerance);
 	}
 
-	FModel& GetModel();
+	FModel& GetModel() { return *GetModelAsShared(); }
+	TSharedPtr<FModel> GetModelAsShared();
 
 	void Serialize(FCADKernelArchive& Ar)
 	{
diff --git a/Engine/Source/Runtime/Datasmith/CADKernel/Base/Public/Core/Types.h b/Engine/Source/Runtime/Datasmith/CADKernel/Base/Public/Core/Types.h
--- a/Engine/Source/Runtime/Datasmith/CADKernel/Base/Public/Core/Types.h
+++ b/Engine/Source/Runtime/Datasmith/CADKernel/Base/Public/Core/Types.h
@@ -109,9 +109,8 @@ enum EVerboseLevel : uint8
 };
 } // namespace UE::CADKernel
 
-#ifdef DO_ENSURE_CADKERNEL
+#if CADKERNEL_DO_ENSURE
 #define ensureCADKernel(InExpression) ensure(InExpression)
 #else
-#define ensureCADKernel(InExpression) ensure(InExpression)
-//#define ensureCADKernel(InExpression) {}
+#define ensureCADKernel(InExpression) (LIKELY(!!(InExpression)))
 #endif
diff --git a/Engine/Source/Runtime/Datasmith/CADKernel/Base/Public/Geo/Curves/Curve.h b/Engine/Source/Runtime/Datasmith/CADKernel/Base/Public/Geo/Curves/Curve.h
--- a/Engine/Source/Runtime/Datasmith/CADKernel/Base/Public/Geo/Curves/Curve.h
+++ b/Engine/Source/Runtime/Datasmith/CADKernel/Base/Public/Geo/Curves/Curve.h
@@ -157,7 +157,7 @@ public:
 	 */
 	virtual void Evaluate2DPoint(double Coordinate, FPoint2D& OutPoint) const
 	{
-		ensure(Dimension == 2);
+		ensureCADKernel(Dimension == 2);
 		FCurvePoint2D CurvePoint;
 		Evaluate2DPoint(Coordinate, CurvePoint);
 		OutPoint = CurvePoint.Point;
@@ -187,7 +187,10 @@ public:
 	 */
 	virtual void Evaluate2DPoints(const TArray<double>& Coordinates, TArray<FCurvePoint2D>& OutPoints, int32 DerivativeOrder = 0) const;
 
-	void FindNotDerivableCoordinates(int32 DerivativeOrder, TArray<double>& OutNotDerivableCoordinates) const;
+	void FindNotDerivableCoordinates(int32 DerivativeOrder, TArray<double>& OutNotDerivableCoordinates) const
+	{
+		return FindNotDerivableCoordinates(Boundary, DerivativeOrder, OutNotDerivableCoordinates);
+	}
 	virtual void FindNotDerivableCoordinates(const FLinearBoundary& InBoundary, int32 DerivativeOrder, TArray<double>& OutNotDerivableCoordinates) const;
 
 
@@ -218,7 +221,7 @@ public:
 	 */
 	virtual void ExtendTo(const FPoint& DesiredPosition)
 	{
-		ensure(Dimension == 3);
+		ensureCADKernel(Dimension == 3);
 
 		NOT_IMPLEMENTED;
 		ensureCADKernel(false);
@@ -229,7 +232,7 @@ public:
 	 */
 	virtual void ExtendTo(const FPoint2D& DesiredPosition)
 	{
-		ensure(Dimension == 2);
+		ensureCADKernel(Dimension == 2);
 		FPoint Point = DesiredPosition;
 		ExtendTo(Point);
 	}
diff --git a/Engine/Source/Runtime/Datasmith/CADKernel/Base/Public/Geo/Curves/CurveUtilities.h b/Engine/Source/Runtime/Datasmith/CADKernel/Base/Public/Geo/Curves/CurveUtilities.h
new file mode 100644
--- /dev/null
+++ b/Engine/Source/Runtime/Datasmith/CADKernel/Base/Public/Geo/Curves/CurveUtilities.h
@@ -0,0 +1,26 @@
+// Copyright Epic Games, Inc. All Rights Reserved.
+#pragma once
+
+#include "CoreMinimal.h"
+
+#include "Geo/Curves/RestrictionCurve.h"
+
+namespace UE::CADKernel
+{
+	class FCurve;
+	class FRestrictionCurve;
+
+	namespace CurveUtilities
+	{
+		int32 CADKERNEL_API GetPoleCount(const UE::CADKernel::FCurve& Curve);
+		int32 CADKERNEL_API GetPoleCount(const UE::CADKernel::FRestrictionCurve& Curve);
+
+		TArray<FVector> CADKERNEL_API GetPoles(const UE::CADKernel::FCurve& Curve);
+		TArray<FVector2d> CADKERNEL_API GetPoles(const UE::CADKernel::FRestrictionCurve& Curve);
+		TArray<FVector2d> CADKERNEL_API Get2DPolyline(const UE::CADKernel::FRestrictionCurve& Curve);
+		TArray<FVector> CADKERNEL_API Get3DPolyline(const UE::CADKernel::FRestrictionCurve& Curve);
+
+		int32 CADKERNEL_API GetDegree(const UE::CADKernel::FCurve& Curve);
+		int32 CADKERNEL_API GetDegree(const UE::CADKernel::FRestrictionCurve& Curve);
+	}
+}
\ No newline at end of file
diff --git a/Engine/Source/Runtime/Datasmith/CADKernel/Base/Public/Geo/Curves/RestrictionCurve.h b/Engine/Source/Runtime/Datasmith/CADKernel/Base/Public/Geo/Curves/RestrictionCurve.h
--- a/Engine/Source/Runtime/Datasmith/CADKernel/Base/Public/Geo/Curves/RestrictionCurve.h
+++ b/Engine/Source/Runtime/Datasmith/CADKernel/Base/Public/Geo/Curves/RestrictionCurve.h
@@ -79,6 +79,11 @@ public:
 		return Curve2D.ToSharedRef();
 	}
 
+	double GetMinLinearTolerance() const
+	{
+		return MinLinearTolerance;
+	}
+
 	TSharedPtr<FEntityGeom> ApplyMatrix(const FMatrixH& InMatrix) const override
 	{
 		ensureCADKernel(false);
@@ -269,7 +274,7 @@ public:
 	 */
 	virtual void Offset(const FPoint& OffsetDirection) override
 	{
-		ensure(false);
+		ensureCADKernel(false);
 	}
 
 };
diff --git a/Engine/Source/Runtime/Datasmith/CADKernel/Base/Public/Geo/Curves/SurfacicCurve.h b/Engine/Source/Runtime/Datasmith/CADKernel/Base/Public/Geo/Curves/SurfacicCurve.h
--- a/Engine/Source/Runtime/Datasmith/CADKernel/Base/Public/Geo/Curves/SurfacicCurve.h
+++ b/Engine/Source/Runtime/Datasmith/CADKernel/Base/Public/Geo/Curves/SurfacicCurve.h
@@ -112,7 +112,7 @@ public:
 	 */
 	virtual void Offset(const FPoint& OffsetDirection) override
 	{
-		ensure(false);
+		ensureCADKernel(false);
 	}
 
 	virtual void EvaluatePoint(double Coordinate, FCurvePoint& OutPoint, int32 DerivativeOrder = 0) const override;
diff --git a/Engine/Source/Runtime/Datasmith/CADKernel/Base/Public/Geo/Sampling/PolylineTools.h b/Engine/Source/Runtime/Datasmith/CADKernel/Base/Public/Geo/Sampling/PolylineTools.h
--- a/Engine/Source/Runtime/Datasmith/CADKernel/Base/Public/Geo/Sampling/PolylineTools.h
+++ b/Engine/Source/Runtime/Datasmith/CADKernel/Base/Public/Geo/Sampling/PolylineTools.h
@@ -374,7 +374,7 @@ public:
 					OutPoints[IPoint].Point = PolylinePoints[Index];
 					OutPoints[IPoint].Gradient = FPoint::ZeroPoint;
 					OutPoints[IPoint].Laplacian = FPoint::ZeroPoint;
-					return;
+					continue;
 				}
 
 				double SectionCoordinate = (InCoordinates[IPoint] - PolylineCoordinates[Index]) / DeltaU;
@@ -396,7 +396,7 @@ public:
 		int32 EndIndex = DichotomyFinder.Find(InCoordinates.Last());
 		bool bUseDichotomy = PolylineTools::IsDichotomyToBePreferred(EndIndex - StartIndex, InCoordinates.Num());
 
-		OutPoints.Empty(InCoordinates.Num());
+		OutPoints.SetNum(InCoordinates.Num(), EAllowShrinking::No);
 		if (bUseDichotomy)
 		{
 			DichotomyFinder.StartLower = StartIndex;
diff --git a/Engine/Source/Runtime/Datasmith/CADKernel/Base/Public/Geo/Surfaces/Surface.h b/Engine/Source/Runtime/Datasmith/CADKernel/Base/Public/Geo/Surfaces/Surface.h
--- a/Engine/Source/Runtime/Datasmith/CADKernel/Base/Public/Geo/Surfaces/Surface.h
+++ b/Engine/Source/Runtime/Datasmith/CADKernel/Base/Public/Geo/Surfaces/Surface.h
@@ -200,6 +200,18 @@ public:
 
 	virtual void EvaluateNormals(const TArray<FPoint2D>& Points2D, TArray<FVector3f>& Normals) const;
 
+	virtual void EvaluatePointAndNormal(const FPoint2D& InPoint2D, FPoint& Point, FVector3f& Normal) const
+	{
+		int32 DerivativeOrder = 1;
+		FSurfacicPoint SurfacicPoint;
+		EvaluatePoint(InPoint2D, SurfacicPoint, DerivativeOrder);
+
+		Point = SurfacicPoint.Point;
+
+		FVector NormalD = SurfacicPoint.GradientU ^ SurfacicPoint.GradientV;
+		Normal = (FVector3f)NormalD.GetSafeNormal();
+	}
+
 	virtual FVector EvaluateNormal(const FPoint2D& InPoint2D) const
 	{
 		int32 DerivativeOrder = 1;
diff --git a/Engine/Source/Runtime/Datasmith/CADKernel/Base/Public/Geo/Surfaces/SurfaceUtilities.h b/Engine/Source/Runtime/Datasmith/CADKernel/Base/Public/Geo/Surfaces/SurfaceUtilities.h
new file mode 100644
--- /dev/null
+++ b/Engine/Source/Runtime/Datasmith/CADKernel/Base/Public/Geo/Surfaces/SurfaceUtilities.h
@@ -0,0 +1,15 @@
+// Copyright Epic Games, Inc. All Rights Reserved.
+#pragma once
+
+#include "CoreMinimal.h"
+
+namespace UE::CADKernel
+{
+	class FSurface;
+
+	namespace SurfaceUtilities
+	{
+		TArray<TArray<FVector3f>> CADKERNEL_API GetPoles(const UE::CADKernel::FSurface& Surface);
+		bool CADKERNEL_API IsPlanar(const UE::CADKernel::FSurface& Surface);
+	}
+}
\ No newline at end of file
diff --git a/Engine/Source/Runtime/Datasmith/CADKernel/Base/Public/Math/MathConst.h b/Engine/Source/Runtime/Datasmith/CADKernel/Base/Public/Math/MathConst.h
--- a/Engine/Source/Runtime/Datasmith/CADKernel/Base/Public/Math/MathConst.h
+++ b/Engine/Source/Runtime/Datasmith/CADKernel/Base/Public/Math/MathConst.h
@@ -33,6 +33,12 @@ template<typename ValueType> const void GetMinMax(ValueType& Min, ValueType& Max
 	}
 }
 
+/** Checks if value is within a range, exclusive on MinValue and MaxValue) */
+template<typename ValueType> bool IsWithinExclusive(const ValueType& TestValue, const ValueType& MinValue, const ValueType& MaxValue)
+{
+	return ((TestValue > MinValue) && (TestValue < MaxValue));
+}
+
 inline int32 RealCompare(const double Value1, const double Value2, const double Tolerance = DOUBLE_SMALL_NUMBER)
 {
 	double Difference = Value1 - Value2;
diff --git a/Engine/Source/Runtime/Datasmith/CADKernel/Base/Public/Math/Point.h b/Engine/Source/Runtime/Datasmith/CADKernel/Base/Public/Math/Point.h
--- a/Engine/Source/Runtime/Datasmith/CADKernel/Base/Public/Math/Point.h
+++ b/Engine/Source/Runtime/Datasmith/CADKernel/Base/Public/Math/Point.h
@@ -374,6 +374,10 @@ PRAGMA_ENABLE_DEPRECATION_WARNINGS
 		return FPoint(PointX, PointY, PointZ);
 	}
 
+	bool Equals(const FPoint& Other, double Tolerance = UE_DOUBLE_SMALL_NUMBER) const
+	{
+		return FMath::IsNearlyEqual(X, Other.X, Tolerance) && FMath::IsNearlyEqual(Y, Other.Y, Tolerance) && FMath::IsNearlyEqual(Z, Other.Z, Tolerance);
+	}
 };
 
 class CADKERNEL_API FPointH
@@ -732,6 +736,11 @@ PRAGMA_ENABLE_DEPRECATION_WARNINGS
 	{
 		return { FMath::Max(PointA.U, PointB.U), FMath::Max(PointA.V, PointB.V) };
 	}
+
+	bool Equals(const FPoint2D& Other, double Tolerance = UE_DOUBLE_SMALL_NUMBER) const
+	{
+		return FMath::IsNearlyEqual(U, Other.U, Tolerance) && FMath::IsNearlyEqual(V, Other.V, Tolerance);
+	}
 };
 
 inline FPoint::FPoint(const FPoint2D& Point)
diff --git a/Engine/Source/Runtime/Datasmith/CADKernel/Base/Public/Mesh/Structure/Mesh.h b/Engine/Source/Runtime/Datasmith/CADKernel/Base/Public/Mesh/Structure/Mesh.h
--- a/Engine/Source/Runtime/Datasmith/CADKernel/Base/Public/Mesh/Structure/Mesh.h
+++ b/Engine/Source/Runtime/Datasmith/CADKernel/Base/Public/Mesh/Structure/Mesh.h
@@ -88,6 +88,11 @@ public:
 	{
 		return TopologicalEntity;
 	}
+
+	FTopologicalEntity& GetGeometricEntity()
+	{
+		return TopologicalEntity;
+	}
 };
 }
 
diff --git a/Engine/Source/Runtime/Datasmith/CADKernel/Base/Public/Topo/TopologicalEdge.h b/Engine/Source/Runtime/Datasmith/CADKernel/Base/Public/Topo/TopologicalEdge.h
--- a/Engine/Source/Runtime/Datasmith/CADKernel/Base/Public/Topo/TopologicalEdge.h
+++ b/Engine/Source/Runtime/Datasmith/CADKernel/Base/Public/Topo/TopologicalEdge.h
@@ -225,7 +225,6 @@ public:
 
 	/**
 	 * It can be linked to the edge if :
-	 *  - they are connected at their extremities,
 	 *  - they have the same length (5% or +/- EdgeLengthTolerance)
 	 *  - they are ~tangent at their extremities i.e @see IsTangentAtExtremitiesWith
 	 */
@@ -234,8 +233,8 @@ public:
 	/**
 	 * Link two edges.
 	 * Two edges can be linked if :
-	 *  - they are connected at their extremities (SquareJoiningTolerance),
-	 *  - they are linkable (@see IsLinkableTo)
+	 *  - Their extremities coincide following SquareJoiningTolerance's criteria,
+	 *  - They are linkable (@see IsLinkableTo)
 	 * 
 	 * This step must be done when the loop is finalize because in some case edges can be delete, split, extend... to avoid problems
 	 */
diff --git a/Engine/Source/Runtime/Datasmith/CADKernel/Base/Public/Topo/TopologicalFaceUtilities.h b/Engine/Source/Runtime/Datasmith/CADKernel/Base/Public/Topo/TopologicalFaceUtilities.h
new file mode 100644
--- /dev/null
+++ b/Engine/Source/Runtime/Datasmith/CADKernel/Base/Public/Topo/TopologicalFaceUtilities.h
@@ -0,0 +1,23 @@
+// Copyright Epic Games, Inc. All Rights Reserved.
+#pragma once
+
+#include "CoreMinimal.h"
+
+namespace UE::CADKernel
+{
+	class FOrientedEdge;
+	class FTopologicalEdge;
+	class FTopologicalFace;
+	class FTopologicalLoop;
+
+	namespace TopologicalFaceUtilities
+	{
+		bool CADKERNEL_API IsPlanar(const UE::CADKernel::FTopologicalFace& Face);
+		TArray<FVector2d> CADKERNEL_API Get2DPolyline(const UE::CADKernel::FTopologicalEdge& Edge);
+		TArray<FVector2d> CADKERNEL_API Get2DPolyline(const UE::CADKernel::FOrientedEdge& Edge);
+		TArray<FVector2d> CADKERNEL_API Get2DPolyline(const UE::CADKernel::FTopologicalLoop& Loop);
+		TArray<FVector> CADKERNEL_API Get3DPolyline(const UE::CADKernel::FTopologicalEdge& Edge);
+		TArray<FVector> CADKERNEL_API Get3DPolyline(const UE::CADKernel::FOrientedEdge& Edge);
+		TArray<FVector> CADKERNEL_API Get3DPolyline(const UE::CADKernel::FTopologicalLoop& Loop);
+	}
+}
\ No newline at end of file
