Fixed issue with misoriented holes.

#rnx

diff --git a/Engine/Source/Runtime/Datasmith/CADKernel/Base/Private/Mesh/Meshers/ParametricFaceMesher.cpp b/Engine/Source/Runtime/Datasmith/CADKernel/Base/Private/Mesh/Meshers/ParametricFaceMesher.cpp
--- a/Engine/Source/Runtime/Datasmith/CADKernel/Base/Private/Mesh/Meshers/ParametricFaceMesher.cpp
+++ b/Engine/Source/Runtime/Datasmith/CADKernel/Base/Private/Mesh/Meshers/ParametricFaceMesher.cpp
@@ -69,13 +69,7 @@ void FParametricFaceMesher::Mesh()
 		return;
 	}
 
-	FTimePoint IsoTriangulerStartTime = FChrono::Now();
-
-	if (GDetectPlanarFace && SurfaceUtilities::IsPlanar(*Face.GetCarrierSurface())/*TopologicalFaceUtilities::IsPolygonal(Face)*/)
-	{
-		MeshPlanarFace();
-	}
-	else
+	if (!MeshPlanarFace())
 	{
 		FFaceMesh& SurfaceMesh = Face.GetOrCreateMesh(MeshModel);
 		FIsoTriangulator IsoTrianguler(Grid, SurfaceMesh, Tolerances);
@@ -892,12 +886,17 @@ void FParametricFaceMesher::MeshThinZoneSide(FThinZoneSide& Side, bool bFinalMes
 	}
 }
 
-void FParametricFaceMesher::MeshPlanarFace()
+bool FParametricFaceMesher::MeshPlanarFace()
 {
 	TRACE_CPUPROFILER_EVENT_SCOPE(UE::CADKernel::FParametricFaceMesher::MeshPlanarFace)
 		
 	using namespace UE::Geometry;
 
+	if (!GDetectPlanarFace || !SurfaceUtilities::IsPlanar(*Face.GetCarrierSurface()))
+	{
+		return false;
+	}
+
 	TArray<FGeneralPolygon2d> PolygonsOut;
 	PolygonsOut.Reserve(Face.GetLoops().Num());
 
@@ -906,6 +905,7 @@ void FParametricFaceMesher::MeshPlanarFace()
 	ensureCADKernel(Loops.Num() == Loops2D.Num());
 
 	FGeneralPolygon2d* GPolygon = nullptr;
+	bool bOuterIsCW = false;
 	for (int32 Index = 0; Index < Loops.Num(); ++Index)
 	{
 		const TArray<FPoint2D>& PointList = Loops2D[Index];
@@ -920,13 +920,20 @@ void FParametricFaceMesher::MeshPlanarFace()
 
 		if (Loops[Index]->IsExternal())
 		{
+			bOuterIsCW = Polygon2d.IsClockwise();
 			GPolygon = &PolygonsOut.Emplace_GetRef(MoveTemp(Polygon2d));
 		}
 		else if (ensureCADKernel(GPolygon))
 		{
+			if (bOuterIsCW == Polygon2d.IsClockwise())
+			{
+				// The hole must be in the reverse direction
+				Polygon2d.Reverse();
+			}
 			if (!ensureCADKernel(GPolygon->AddHole(Polygon2d, /*bCheckContainment =*/ true, /*bCheckOrientation =*/ true)))
 			{
-				// #cad_kernel: Log verbose a hole failed to be added
+				// #cadkernel_check: Why this is happening?
+				return false;
 			}
 		}
 	}
@@ -935,7 +942,7 @@ void FParametricFaceMesher::MeshPlanarFace()
 	if (!ensureCADKernel(PolygonsOut.Num() > 0 && PolygonsOut[0].GetOuter().GetVertices().Num() > 2))
 	{
 		// #cadkernel_check: Why this is happening?
-		return;
+		return false;
 	}
 
 	FFaceMesh& FaceMesh = Face.GetOrCreateMesh(MeshModel);
@@ -950,7 +957,8 @@ void FParametricFaceMesher::MeshPlanarFace()
 	// #cadkernel_check: Why this is happening?
 	if (TrianglesVerticesIndex.Num() > 0)
 	{
-		return;
+		ensureCADKernel(false);
+		return false;
 	}
 
 	const int32 GlobalVertexCount = MeshModel.GetVertexCount();
@@ -973,7 +981,8 @@ void FParametricFaceMesher::MeshPlanarFace()
 		Delaunay.Triangulate(Polygon, &Triangles, &Vertices, true /*bFallbackToGeneralizedWinding*/);
 		if (!ensureCADKernel(Triangles.Num() > 0))
 		{
-			continue;
+			// #cadkernel_check: Why this is happening?
+			return false;
 		}
 
 		const int32 StartIndex = VerticesGlobalIndex.Num();
@@ -1017,5 +1026,7 @@ void FParametricFaceMesher::MeshPlanarFace()
 
 	FaceMesh.RegisterCoordinates();
 	MeshModel.AddMesh(FaceMesh);
+
+	return true;
 }
 } // namespace
diff --git a/Engine/Source/Runtime/Datasmith/CADKernel/Base/Private/Mesh/Meshers/ParametricFaceMesher.h b/Engine/Source/Runtime/Datasmith/CADKernel/Base/Private/Mesh/Meshers/ParametricFaceMesher.h
--- a/Engine/Source/Runtime/Datasmith/CADKernel/Base/Private/Mesh/Meshers/ParametricFaceMesher.h
+++ b/Engine/Source/Runtime/Datasmith/CADKernel/Base/Private/Mesh/Meshers/ParametricFaceMesher.h
@@ -153,7 +153,7 @@ private:
 	void Mesh(FTopologicalEdge& InEdge, bool bFinalMeshing = true);
 	void Mesh(FTopologicalVertex& Vertex);
 	void MeshVerticesOfFace(FTopologicalFace& Face);
-	void MeshPlanarFace();
+	bool MeshPlanarFace();
 
 	void MeshFaceLoops();
 
