Disabled by default solution for bad tessellation of planar surfaces

#rnx

diff --git a/Engine/Source/Runtime/Datasmith/CADKernel/Base/Private/Mesh/Meshers/ParametricFaceMesher.cpp b/Engine/Source/Runtime/Datasmith/CADKernel/Base/Private/Mesh/Meshers/ParametricFaceMesher.cpp
--- a/Engine/Source/Runtime/Datasmith/CADKernel/Base/Private/Mesh/Meshers/ParametricFaceMesher.cpp
+++ b/Engine/Source/Runtime/Datasmith/CADKernel/Base/Private/Mesh/Meshers/ParametricFaceMesher.cpp
@@ -21,7 +21,7 @@
 #include "CompGeom/Delaunay2.h"
 #include "HAL/IConsoleManager.h"
 
-static bool GDetectPlanarFace = true;
+static bool GDetectPlanarFace = false;
 FAutoConsoleVariableRef PlanarFaceCVar(
 	TEXT("CADKernel.FaceMesher.DetectPlanarFace"),
 	GDetectPlanarFace,
