Replaced FPoint, FFPoint, FPoint2D and FPointH by FVector, FVector3f, FVector2d and FVector4d

#rnx

diff --git a/Engine/Plugins/Enterprise/DatasmithCADImporter/Source/CADInterfaces/Private/TechSoftBridge.cpp b/Engine/Plugins/Enterprise/DatasmithCADImporter/Source/CADInterfaces/Private/TechSoftBridge.cpp
--- a/Engine/Plugins/Enterprise/DatasmithCADImporter/Source/CADInterfaces/Private/TechSoftBridge.cpp
+++ b/Engine/Plugins/Enterprise/DatasmithCADImporter/Source/CADInterfaces/Private/TechSoftBridge.cpp
@@ -161,10 +161,10 @@ UE::CADKernel::FSurfacicBoundary GetSurfacicBoundary(A3DDomainData& Domain, cons
 	UE::CADKernel::EIso VIndex = UVReparameterization.GetSwapUV() ? UE::CADKernel::EIso::IsoU : UE::CADKernel::EIso::IsoV;
 
 	UE::CADKernel::FSurfacicBoundary Boundary;
-	Boundary[UIndex].Min = FMath::Min(Min.U, Max.U);
-	Boundary[UIndex].Max = FMath::Max(Min.U, Max.U);
-	Boundary[VIndex].Min = FMath::Min(Min.V, Max.V);
-	Boundary[VIndex].Max = FMath::Max(Min.V, Max.V);
+	Boundary[UIndex].Min = FMath::Min(Min.X, Max.X);
+	Boundary[UIndex].Max = FMath::Max(Min.X, Max.X);
+	Boundary[VIndex].Min = FMath::Min(Min.Y, Max.Y);
+	Boundary[VIndex].Max = FMath::Max(Min.Y, Max.Y);
 
 	return Boundary;
 }
diff --git a/Engine/Plugins/Enterprise/DatasmithCADImporter/Source/CADInterfaces/Private/TechSoftBridge.h b/Engine/Plugins/Enterprise/DatasmithCADImporter/Source/CADInterfaces/Private/TechSoftBridge.h
--- a/Engine/Plugins/Enterprise/DatasmithCADImporter/Source/CADInterfaces/Private/TechSoftBridge.h
+++ b/Engine/Plugins/Enterprise/DatasmithCADImporter/Source/CADInterfaces/Private/TechSoftBridge.h
@@ -140,8 +140,8 @@ public:
 
 	void Apply(UE::CADKernel::FPoint2D& Point) const
 	{
-		Point.U = Scale[UE::CADKernel::EIso::IsoU] * Point.U + Offset[UE::CADKernel::EIso::IsoU];
-		Point.V = Scale[UE::CADKernel::EIso::IsoV] * Point.V + Offset[UE::CADKernel::EIso::IsoV];
+		Point.X = Scale[UE::CADKernel::EIso::IsoU] * Point.X + Offset[UE::CADKernel::EIso::IsoU];
+		Point.Y = Scale[UE::CADKernel::EIso::IsoV] * Point.Y + Offset[UE::CADKernel::EIso::IsoV];
 	}
 
 private:
diff --git a/Engine/Source/Runtime/Datasmith/CADKernel/Base/Private/Geo/Curves/CurveUtilities.cpp b/Engine/Source/Runtime/Datasmith/CADKernel/Base/Private/Geo/Curves/CurveUtilities.cpp
--- a/Engine/Source/Runtime/Datasmith/CADKernel/Base/Private/Geo/Curves/CurveUtilities.cpp
+++ b/Engine/Source/Runtime/Datasmith/CADKernel/Base/Private/Geo/Curves/CurveUtilities.cpp
@@ -155,7 +155,7 @@ namespace UE::CADKernel::CurveUtilities
 
 		for (const FPoint2D& Point2D : Point2Ds)
 		{
-			SurfacicPoints.Emplace(Point2D.U, Point2D.V);
+			SurfacicPoints.Emplace(Point2D.X, Point2D.Y);
 		}
 
 		return SurfacicPoints;
diff --git a/Engine/Source/Runtime/Datasmith/CADKernel/Base/Private/Geo/Surfaces/BezierSurface.cpp b/Engine/Source/Runtime/Datasmith/CADKernel/Base/Private/Geo/Surfaces/BezierSurface.cpp
--- a/Engine/Source/Runtime/Datasmith/CADKernel/Base/Private/Geo/Surfaces/BezierSurface.cpp
+++ b/Engine/Source/Runtime/Datasmith/CADKernel/Base/Private/Geo/Surfaces/BezierSurface.cpp
@@ -30,7 +30,7 @@ void FBezierSurface::EvaluatePoint(const FPoint2D& InPoint2D, FSurfacicPoint& Ou
 	}
 
 	{
-		double Coordinate = InPoint2D.U;
+		double Coordinate = InPoint2D.X;
 
 		TArray<FPoint> UAuxiliaryPoles;
 		UAuxiliaryPoles.SetNum(UPoleNum);
@@ -106,7 +106,7 @@ void FBezierSurface::EvaluatePoint(const FPoint2D& InPoint2D, FSurfacicPoint& Ou
 		}
 	}
 
-	double Coordinate = InPoint2D.V;
+	double Coordinate = InPoint2D.Y;
 
 	// Compute Point, Gradient and Laplacian at V coordinate with De Casteljau's algorithm,
 	for (int32 Vndex = VPoleNum - 2; Vndex >= 0; Vndex--)
diff --git a/Engine/Source/Runtime/Datasmith/CADKernel/Base/Private/Geo/Surfaces/CompositeSurface.cpp b/Engine/Source/Runtime/Datasmith/CADKernel/Base/Private/Geo/Surfaces/CompositeSurface.cpp
--- a/Engine/Source/Runtime/Datasmith/CADKernel/Base/Private/Geo/Surfaces/CompositeSurface.cpp
+++ b/Engine/Source/Runtime/Datasmith/CADKernel/Base/Private/Geo/Surfaces/CompositeSurface.cpp
@@ -78,19 +78,19 @@ void FCompositeSurface::EvaluatePoint(const FPoint2D& InSurfacicCoordinate, FSur
 
 	FPoint2D Point2DPatch;
 	FDichotomyFinder FinderU(GlobalCoordinates[EIso::IsoU]);
-	int32 PatchUIndex = FinderU.Find(Coordinate.U);
-	Point2DPatch.U = (Coordinate.U - GlobalCoordinates[EIso::IsoU][PatchUIndex]) / (GlobalCoordinates[EIso::IsoU][PatchUIndex + 1] - GlobalCoordinates[EIso::IsoU][PatchUIndex]);
+	int32 PatchUIndex = FinderU.Find(Coordinate.X);
+	Point2DPatch.X = (Coordinate.X - GlobalCoordinates[EIso::IsoU][PatchUIndex]) / (GlobalCoordinates[EIso::IsoU][PatchUIndex + 1] - GlobalCoordinates[EIso::IsoU][PatchUIndex]);
 
 	FDichotomyFinder FinderV(GlobalCoordinates[EIso::IsoV]);
-	int32 PatchVIndex = FinderV.Find(Coordinate.V);
-	Point2DPatch.V = (Coordinate.V - GlobalCoordinates[EIso::IsoV][PatchVIndex]) / (GlobalCoordinates[EIso::IsoV][PatchVIndex + 1] - GlobalCoordinates[EIso::IsoV][PatchVIndex]);
+	int32 PatchVIndex = FinderV.Find(Coordinate.Y);
+	Point2DPatch.Y = (Coordinate.Y - GlobalCoordinates[EIso::IsoV][PatchVIndex]) / (GlobalCoordinates[EIso::IsoV][PatchVIndex + 1] - GlobalCoordinates[EIso::IsoV][PatchVIndex]);
 
 	int32 SurfaceIndex = PatchVIndex * GetSurfNum(EIso::IsoU) + PatchUIndex;
 
 	ensureCADKernel(SurfaceIndex == Surfaces.Num());
 
-	Point2DPatch.U = NativeUVBoundaries[SurfaceIndex][EIso::IsoU].GetAt(Point2DPatch.U);
-	Point2DPatch.V = NativeUVBoundaries[SurfaceIndex][EIso::IsoV].GetAt(Point2DPatch.V);
+	Point2DPatch.X = NativeUVBoundaries[SurfaceIndex][EIso::IsoU].GetAt(Point2DPatch.X);
+	Point2DPatch.Y = NativeUVBoundaries[SurfaceIndex][EIso::IsoV].GetAt(Point2DPatch.Y);
 
 	Surfaces[SurfaceIndex]->EvaluatePoint(Point2DPatch, OutPoint3D, InDerivativeOrder);
 }
diff --git a/Engine/Source/Runtime/Datasmith/CADKernel/Base/Private/Geo/Surfaces/ConeSurface.cpp b/Engine/Source/Runtime/Datasmith/CADKernel/Base/Private/Geo/Surfaces/ConeSurface.cpp
--- a/Engine/Source/Runtime/Datasmith/CADKernel/Base/Private/Geo/Surfaces/ConeSurface.cpp
+++ b/Engine/Source/Runtime/Datasmith/CADKernel/Base/Private/Geo/Surfaces/ConeSurface.cpp
@@ -14,26 +14,26 @@ void FConeSurface::EvaluatePoint(const FPoint2D& InPoint2D, FSurfacicPoint& OutP
 	OutPoint3D.DerivativeOrder = InDerivativeOrder;
 
 	double DeltaVR = tan(ConeAngle);
-	double Radius = StartRadius + InPoint2D.V * DeltaVR;
+	double Radius = StartRadius + InPoint2D.Y * DeltaVR;
 
-	OutPoint3D.Point.Set(Radius * cos(InPoint2D.U), Radius * sin(InPoint2D.U), InPoint2D.V);
+	OutPoint3D.Point.Set(Radius * cos(InPoint2D.X), Radius * sin(InPoint2D.X), InPoint2D.Y);
 
 	OutPoint3D.Point = Matrix.Multiply(OutPoint3D.Point);
 
 	if (InDerivativeOrder > 0)
 	{
-		OutPoint3D.GradientU = FPoint(-Radius * sin(InPoint2D.U), Radius * cos(InPoint2D.U), 0.);
-		OutPoint3D.GradientV = FPoint(DeltaVR * cos(InPoint2D.U), DeltaVR * sin(InPoint2D.U), 1.0);
+		OutPoint3D.GradientU = FPoint(-Radius * sin(InPoint2D.X), Radius * cos(InPoint2D.X), 0.);
+		OutPoint3D.GradientV = FPoint(DeltaVR * cos(InPoint2D.X), DeltaVR * sin(InPoint2D.X), 1.0);
 
 		OutPoint3D.GradientU = Matrix.MultiplyVector(OutPoint3D.GradientU);
 		OutPoint3D.GradientV = Matrix.MultiplyVector(OutPoint3D.GradientV);
 
 		if (InDerivativeOrder > 1)
 		{
-			OutPoint3D.LaplacianU = FPoint(-Radius * cos(InPoint2D.U), -Radius * sin(InPoint2D.U), 0.);
+			OutPoint3D.LaplacianU = FPoint(-Radius * cos(InPoint2D.X), -Radius * sin(InPoint2D.X), 0.);
 			OutPoint3D.LaplacianV = FPoint(0., 0., 0.);
 
-			OutPoint3D.LaplacianUV = FPoint(-DeltaVR * sin(InPoint2D.U), DeltaVR * cos(InPoint2D.U), 0.);
+			OutPoint3D.LaplacianUV = FPoint(-DeltaVR * sin(InPoint2D.X), DeltaVR * cos(InPoint2D.X), 0.);
 
 			OutPoint3D.LaplacianU = Matrix.MultiplyVector(OutPoint3D.LaplacianU);
 			OutPoint3D.LaplacianV = Matrix.MultiplyVector(OutPoint3D.LaplacianV);
diff --git a/Engine/Source/Runtime/Datasmith/CADKernel/Base/Private/Geo/Surfaces/CoonsSurface.cpp b/Engine/Source/Runtime/Datasmith/CADKernel/Base/Private/Geo/Surfaces/CoonsSurface.cpp
--- a/Engine/Source/Runtime/Datasmith/CADKernel/Base/Private/Geo/Surfaces/CoonsSurface.cpp
+++ b/Engine/Source/Runtime/Datasmith/CADKernel/Base/Private/Geo/Surfaces/CoonsSurface.cpp
@@ -97,7 +97,7 @@ void FCoonsSurface::EvaluatePoint(const FPoint2D& InPoint2D, FSurfacicPoint& Out
 	ComputePointOnRuledSurface(Curves+2, InPoint2D, EIso::IsoV, RuledPoint2, Curve34Points);
 
 	// 3. bilinear surface
-	FPoint BilinearPoint = (Corners[0]*(1.0-InPoint2D.U) + Corners[1]*InPoint2D.U)*(1.0-InPoint2D.V) + (Corners[2]*(1.0-InPoint2D.U) + Corners[3]*InPoint2D.U)*InPoint2D.V;
+	FPoint BilinearPoint = (Corners[0]*(1.0-InPoint2D.X) + Corners[1]*InPoint2D.X)*(1.0-InPoint2D.Y) + (Corners[2]*(1.0-InPoint2D.X) + Corners[3]*InPoint2D.X)*InPoint2D.Y;
 
 	OutPoint3D.DerivativeOrder = InDerivativeOrder;
 
@@ -106,14 +106,14 @@ void FCoonsSurface::EvaluatePoint(const FPoint2D& InPoint2D, FSurfacicPoint& Out
 
 	if (InDerivativeOrder > 0)
 	{
-		FPoint DeltaU1 = Curve12Points[0].Gradient + InPoint2D.V * (Curve12Points[1].Gradient - Curve12Points[0].Gradient);
+		FPoint DeltaU1 = Curve12Points[0].Gradient + InPoint2D.Y * (Curve12Points[1].Gradient - Curve12Points[0].Gradient);
 		FPoint DeltaU2 = Curve34Points[1].Point - Curve34Points[0].Point;
-		FPoint DeltaU3 = (Corners[1] * (1.0 - InPoint2D.V) + Corners[3] * InPoint2D.V) - (Corners[0] * (1.0 - InPoint2D.V) + Corners[2] * InPoint2D.V);
+		FPoint DeltaU3 = (Corners[1] * (1.0 - InPoint2D.Y) + Corners[3] * InPoint2D.Y) - (Corners[0] * (1.0 - InPoint2D.Y) + Corners[2] * InPoint2D.Y);
 		OutPoint3D.GradientU = DeltaU1 + DeltaU2 - DeltaU3;
 
 		FPoint DeltaV1 = Curve12Points[1].Point - Curve12Points[0].Point;;
-		FPoint DeltaV2 = Curve34Points[0].Gradient + InPoint2D.U * (Curve34Points[1].Gradient - Curve34Points[0].Gradient);
-		FPoint DeltaV3 = (Corners[2] * (1.0 - InPoint2D.U) + Corners[3] * InPoint2D.U) - (Corners[0] * (1.0 - InPoint2D.U) + Corners[1] * InPoint2D.U);
+		FPoint DeltaV2 = Curve34Points[0].Gradient + InPoint2D.X * (Curve34Points[1].Gradient - Curve34Points[0].Gradient);
+		FPoint DeltaV3 = (Corners[2] * (1.0 - InPoint2D.X) + Corners[3] * InPoint2D.X) - (Corners[0] * (1.0 - InPoint2D.X) + Corners[1] * InPoint2D.X);
 		OutPoint3D.GradientV = DeltaV1 + DeltaV2 - DeltaV3;
 
 		if (InDerivativeOrder > 1)
@@ -121,8 +121,8 @@ void FCoonsSurface::EvaluatePoint(const FPoint2D& InPoint2D, FSurfacicPoint& Out
 			FPoint LaplacianUV1 = (Curve12Points[1].Laplacian - Curve12Points[0].Laplacian);
 			FPoint LaplacianUV2 = (Curve34Points[1].Laplacian - Curve34Points[0].Laplacian);
 
-			OutPoint3D.LaplacianU = Curve12Points[0].Laplacian + InPoint2D.V * LaplacianUV1;
-			OutPoint3D.LaplacianV = Curve34Points[0].Laplacian + InPoint2D.U * LaplacianUV2;
+			OutPoint3D.LaplacianU = Curve12Points[0].Laplacian + InPoint2D.Y * LaplacianUV1;
+			OutPoint3D.LaplacianV = Curve34Points[0].Laplacian + InPoint2D.X * LaplacianUV2;
 
 			OutPoint3D.LaplacianUV = LaplacianUV1 + LaplacianUV2;
 		}
diff --git a/Engine/Source/Runtime/Datasmith/CADKernel/Base/Private/Geo/Surfaces/CylinderSurface.cpp b/Engine/Source/Runtime/Datasmith/CADKernel/Base/Private/Geo/Surfaces/CylinderSurface.cpp
--- a/Engine/Source/Runtime/Datasmith/CADKernel/Base/Private/Geo/Surfaces/CylinderSurface.cpp
+++ b/Engine/Source/Runtime/Datasmith/CADKernel/Base/Private/Geo/Surfaces/CylinderSurface.cpp
@@ -29,12 +29,12 @@ void FCylinderSurface::EvaluatePoint(const FPoint2D& InPoint2D, FSurfacicPoint&
 {
 	OutPoint3D.DerivativeOrder = InDerivativeOrder;
 
-	OutPoint3D.Point.Set(Radius*(cos(InPoint2D.U)), Radius*(sin(InPoint2D.U)), InPoint2D.V);
+	OutPoint3D.Point.Set(Radius*(cos(InPoint2D.X)), Radius*(sin(InPoint2D.X)), InPoint2D.Y);
 	OutPoint3D.Point = Matrix.Multiply(OutPoint3D.Point);
 
 	if(InDerivativeOrder>0) 
 	{
-		OutPoint3D.GradientU = FPoint(Radius*(-sin(InPoint2D.U)), Radius*(cos(InPoint2D.U)), 0.0);
+		OutPoint3D.GradientU = FPoint(Radius*(-sin(InPoint2D.X)), Radius*(cos(InPoint2D.X)), 0.0);
 		OutPoint3D.GradientV = FPoint(0.0, 0.0, 1.0);
 
 		OutPoint3D.GradientU = Matrix.MultiplyVector(OutPoint3D.GradientU);
@@ -42,7 +42,7 @@ void FCylinderSurface::EvaluatePoint(const FPoint2D& InPoint2D, FSurfacicPoint&
 
 		if (InDerivativeOrder > 1) 
 		{
-			OutPoint3D.LaplacianU = FPoint(Radius * (-cos(InPoint2D.U)), Radius * (-sin(InPoint2D.U)), 0.0);
+			OutPoint3D.LaplacianU = FPoint(Radius * (-cos(InPoint2D.X)), Radius * (-sin(InPoint2D.X)), 0.0);
 			OutPoint3D.LaplacianU = Matrix.MultiplyVector(OutPoint3D.LaplacianU);
 
 			OutPoint3D.LaplacianV = FPoint(0.0, 0.0, 0.0);
diff --git a/Engine/Source/Runtime/Datasmith/CADKernel/Base/Private/Geo/Surfaces/NURBSSurface.h b/Engine/Source/Runtime/Datasmith/CADKernel/Base/Private/Geo/Surfaces/NURBSSurface.h
--- a/Engine/Source/Runtime/Datasmith/CADKernel/Base/Private/Geo/Surfaces/NURBSSurface.h
+++ b/Engine/Source/Runtime/Datasmith/CADKernel/Base/Private/Geo/Surfaces/NURBSSurface.h
@@ -129,7 +129,7 @@ public:
 		{
 			for (const double& NodalValue : UNodalVector)
 			{
-				if (ValidateValue(UVPoint.U, NodalValue, ToleranceU))
+				if (ValidateValue(UVPoint.X, NodalValue, ToleranceU))
 				{
 					break;
 				}
@@ -137,7 +137,7 @@ public:
 
 			for (const double& NodalValue : VNodalVector)
 			{
-				if (ValidateValue(UVPoint.V, NodalValue, ToleranceV))
+				if (ValidateValue(UVPoint.Y, NodalValue, ToleranceV))
 				{
 					break;
 				}
diff --git a/Engine/Source/Runtime/Datasmith/CADKernel/Base/Private/Geo/Surfaces/RevolutionSurface.cpp b/Engine/Source/Runtime/Datasmith/CADKernel/Base/Private/Geo/Surfaces/RevolutionSurface.cpp
--- a/Engine/Source/Runtime/Datasmith/CADKernel/Base/Private/Geo/Surfaces/RevolutionSurface.cpp
+++ b/Engine/Source/Runtime/Datasmith/CADKernel/Base/Private/Geo/Surfaces/RevolutionSurface.cpp
@@ -27,9 +27,9 @@ FRevolutionSurface::FRevolutionSurface(const double InToleranceGeometric, TShare
 void FRevolutionSurface::EvaluatePoint(const FPoint2D& InSurfacicCoordinate, FSurfacicPoint& OutPoint3D, int32 InDerivativeOrder) const
 {
 	FCurvePoint GeneratrixPoint;
-	Generatrix->EvaluatePoint(InSurfacicCoordinate.U, GeneratrixPoint, InDerivativeOrder);
+	Generatrix->EvaluatePoint(InSurfacicCoordinate.X, GeneratrixPoint, InDerivativeOrder);
 
-	FMatrixH Matrix = FMatrixH::MakeRotationMatrix(InSurfacicCoordinate.V, RotationAxis);
+	FMatrixH Matrix = FMatrixH::MakeRotationMatrix(InSurfacicCoordinate.Y, RotationAxis);
 
 	OutPoint3D.DerivativeOrder = InDerivativeOrder;
 	OutPoint3D.Point = Matrix.PointRotation(GeneratrixPoint.Point, Axis->GetStartPoint());
diff --git a/Engine/Source/Runtime/Datasmith/CADKernel/Base/Private/Geo/Surfaces/RuledSurface.cpp b/Engine/Source/Runtime/Datasmith/CADKernel/Base/Private/Geo/Surfaces/RuledSurface.cpp
--- a/Engine/Source/Runtime/Datasmith/CADKernel/Base/Private/Geo/Surfaces/RuledSurface.cpp
+++ b/Engine/Source/Runtime/Datasmith/CADKernel/Base/Private/Geo/Surfaces/RuledSurface.cpp
@@ -34,8 +34,8 @@ void FRuledSurface::LinesNotDerivables(const FSurfacicBoundary& InBoundary, int3
 
 void FRuledSurface::EvaluatePoint(const FPoint2D& InPoint2D, FSurfacicPoint& OutPoint3D, int32 InDerivativeOrder) const
 {
-	double CoordinateCurve0 = Curves[0]->GetBoundary().Min + InPoint2D.U * (Curves[0]->GetBoundary().Max - Curves[0]->GetBoundary().Min);
-	double CoordinateCurve1 = Curves[1]->GetBoundary().Min + InPoint2D.U * (Curves[1]->GetBoundary().Max - Curves[1]->GetBoundary().Min);
+	double CoordinateCurve0 = Curves[0]->GetBoundary().Min + InPoint2D.X * (Curves[0]->GetBoundary().Max - Curves[0]->GetBoundary().Min);
+	double CoordinateCurve1 = Curves[1]->GetBoundary().Min + InPoint2D.X * (Curves[1]->GetBoundary().Max - Curves[1]->GetBoundary().Min);
 
 	FCurvePoint CurvePoint[2];
 	Curves[0]->EvaluatePoint(CoordinateCurve0, CurvePoint[0], InDerivativeOrder);
@@ -44,16 +44,16 @@ void FRuledSurface::EvaluatePoint(const FPoint2D& InPoint2D, FSurfacicPoint& Out
 	OutPoint3D.GradientV = CurvePoint[1].Point - CurvePoint[0].Point;
 
 	OutPoint3D.DerivativeOrder = InDerivativeOrder;
-	OutPoint3D.Point = CurvePoint[0].Point + OutPoint3D.GradientV * InPoint2D.V;
+	OutPoint3D.Point = CurvePoint[0].Point + OutPoint3D.GradientV * InPoint2D.Y;
 
 	if (InDerivativeOrder > 0)
 	{
 		OutPoint3D.LaplacianUV = CurvePoint[1].Gradient - CurvePoint[0].Gradient;
-		OutPoint3D.GradientU = CurvePoint[0].Gradient + InPoint2D.V * OutPoint3D.LaplacianUV;
+		OutPoint3D.GradientU = CurvePoint[0].Gradient + InPoint2D.Y * OutPoint3D.LaplacianUV;
 
 		if (InDerivativeOrder > 1)
 		{
-			OutPoint3D.LaplacianU = CurvePoint[0].Laplacian + InPoint2D.V * (CurvePoint[1].Laplacian - CurvePoint[0].Laplacian);
+			OutPoint3D.LaplacianU = CurvePoint[0].Laplacian + InPoint2D.Y * (CurvePoint[1].Laplacian - CurvePoint[0].Laplacian);
 			OutPoint3D.LaplacianV = FPoint(0.0, 0.0, 0.0);
 		}
 	}
diff --git a/Engine/Source/Runtime/Datasmith/CADKernel/Base/Private/Geo/Surfaces/SphericalSurface.h b/Engine/Source/Runtime/Datasmith/CADKernel/Base/Private/Geo/Surfaces/SphericalSurface.h
--- a/Engine/Source/Runtime/Datasmith/CADKernel/Base/Private/Geo/Surfaces/SphericalSurface.h
+++ b/Engine/Source/Runtime/Datasmith/CADKernel/Base/Private/Geo/Surfaces/SphericalSurface.h
@@ -85,11 +85,11 @@ public:
 
 	virtual void EvaluatePoint(const FPoint2D& InSurfacicCoordinate, FSurfacicPoint& OutPoint3D, int32 InDerivativeOrder = 0) const override
 	{
-		double CosU = cos(InSurfacicCoordinate.U);
-		double CosV = cos(InSurfacicCoordinate.V);
+		double CosU = cos(InSurfacicCoordinate.X);
+		double CosV = cos(InSurfacicCoordinate.Y);
 
-		double SinU = sin(InSurfacicCoordinate.U);
-		double SinV = sin(InSurfacicCoordinate.V);
+		double SinU = sin(InSurfacicCoordinate.X);
+		double SinV = sin(InSurfacicCoordinate.Y);
 
 		OutPoint3D.DerivativeOrder = InDerivativeOrder;
 		OutPoint3D.Point.Set(Radius * CosV * CosU, Radius * CosV * SinU, Radius * SinV);
@@ -120,11 +120,11 @@ public:
 
 	virtual FPoint2D EvaluatePointInCylindricalSpace(const FPoint2D& InSurfacicCoordinate) const override
 	{
-		double CosU = cos(InSurfacicCoordinate.U);
-		double CosV = cos(InSurfacicCoordinate.V);
+		double CosU = cos(InSurfacicCoordinate.X);
+		double CosV = cos(InSurfacicCoordinate.Y);
 
-		double SinU = sin(InSurfacicCoordinate.U);
-		double SwapOrientation = (InSurfacicCoordinate.V < DOUBLE_PI&& InSurfacicCoordinate.V >= 0) ? 1.0 : -1.0;
+		double SinU = sin(InSurfacicCoordinate.X);
+		double SwapOrientation = (InSurfacicCoordinate.Y < DOUBLE_PI&& InSurfacicCoordinate.Y >= 0) ? 1.0 : -1.0;
 
 		return FPoint(Radius * CosV * CosU * SwapOrientation, Radius * CosV * SinU);
 	}
diff --git a/Engine/Source/Runtime/Datasmith/CADKernel/Base/Private/Geo/Surfaces/Surface.cpp b/Engine/Source/Runtime/Datasmith/CADKernel/Base/Private/Geo/Surfaces/Surface.cpp
--- a/Engine/Source/Runtime/Datasmith/CADKernel/Base/Private/Geo/Surfaces/Surface.cpp
+++ b/Engine/Source/Runtime/Datasmith/CADKernel/Base/Private/Geo/Surfaces/Surface.cpp
@@ -134,7 +134,7 @@ void FSurface::EvaluatePoints(const TArray<FCurvePoint2D>& InPoints2D, FSurfacic
 
 	for (int32 Index = 0; Index < InPoints2D.Num(); ++Index)
 	{
-		Polyline.Tangents.Emplace(Points3D[Index].GradientU * InPoints2D[Index].Gradient.U + Points3D[Index].GradientV * InPoints2D[Index].Gradient.V);
+		Polyline.Tangents.Emplace(Points3D[Index].GradientU * InPoints2D[Index].Gradient.X + Points3D[Index].GradientV * InPoints2D[Index].Gradient.Y);
 	}
 
 	for (const FSurfacicPoint& Point : Points3D)
diff --git a/Engine/Source/Runtime/Datasmith/CADKernel/Base/Private/Geo/Surfaces/TabulatedCylinderSurface.cpp b/Engine/Source/Runtime/Datasmith/CADKernel/Base/Private/Geo/Surfaces/TabulatedCylinderSurface.cpp
--- a/Engine/Source/Runtime/Datasmith/CADKernel/Base/Private/Geo/Surfaces/TabulatedCylinderSurface.cpp
+++ b/Engine/Source/Runtime/Datasmith/CADKernel/Base/Private/Geo/Surfaces/TabulatedCylinderSurface.cpp
@@ -31,10 +31,10 @@ void FTabulatedCylinderSurface::EvaluatePoint(const FPoint2D& InSurfacicCoordina
 
 	FCurvePoint PointU;
 	FPoint LaplacianU;
-	GuideCurve->EvaluatePoint(CurveBounds.GetAt(InSurfacicCoordinate.U), PointU, InDerivativeOrder);
+	GuideCurve->EvaluatePoint(CurveBounds.GetAt(InSurfacicCoordinate.X), PointU, InDerivativeOrder);
 
 	OutPoint3D.DerivativeOrder = InDerivativeOrder;
-	OutPoint3D.Point = PointU.Point + DirectorVector * InSurfacicCoordinate.V;
+	OutPoint3D.Point = PointU.Point + DirectorVector * InSurfacicCoordinate.Y;
 
 	if (InDerivativeOrder > 0) 
 	{
diff --git a/Engine/Source/Runtime/Datasmith/CADKernel/Base/Private/Geo/Surfaces/TorusSurface.h b/Engine/Source/Runtime/Datasmith/CADKernel/Base/Private/Geo/Surfaces/TorusSurface.h
--- a/Engine/Source/Runtime/Datasmith/CADKernel/Base/Private/Geo/Surfaces/TorusSurface.h
+++ b/Engine/Source/Runtime/Datasmith/CADKernel/Base/Private/Geo/Surfaces/TorusSurface.h
@@ -91,11 +91,11 @@ public:
 
 	virtual void EvaluatePoint(const FPoint2D& InSurfacicCoordinate, FSurfacicPoint& OutPoint3D, int32 InDerivativeOrder = 0) const override
 	{
-		const double CosU = cos(InSurfacicCoordinate.U);
-		const double CosV = cos(InSurfacicCoordinate.V);
+		const double CosU = cos(InSurfacicCoordinate.X);
+		const double CosV = cos(InSurfacicCoordinate.Y);
 
-		const double SinU = sin(InSurfacicCoordinate.U);
-		const double SinV = sin(InSurfacicCoordinate.V);
+		const double SinU = sin(InSurfacicCoordinate.X);
+		const double SinV = sin(InSurfacicCoordinate.Y);
 		double Rho = (MajorRadius + MinorRadius * CosV);
 
 		OutPoint3D.DerivativeOrder = InDerivativeOrder;
@@ -130,14 +130,14 @@ public:
 
 	virtual FPoint2D EvaluatePointInCylindricalSpace(const FPoint2D& InSurfacicCoordinate) const override
 	{
-		const double CosU = cos(InSurfacicCoordinate.U);
-		const double CosV = cos(InSurfacicCoordinate.V);
+		const double CosU = cos(InSurfacicCoordinate.X);
+		const double CosV = cos(InSurfacicCoordinate.Y);
 
-		const double SinU = sin(InSurfacicCoordinate.U);
+		const double SinU = sin(InSurfacicCoordinate.X);
 
 		double Rho = (MajorRadius + MinorRadius * CosV);
 
-		double SwapOrientation = (InSurfacicCoordinate.V < DOUBLE_PI&& InSurfacicCoordinate.V >= 0) ? 1.0 : -1.0;
+		double SwapOrientation = (InSurfacicCoordinate.Y < DOUBLE_PI&& InSurfacicCoordinate.Y >= 0) ? 1.0 : -1.0;
 
 		return FPoint2D(Rho * CosU * SwapOrientation, Rho * SinU);
 	}
diff --git a/Engine/Source/Runtime/Datasmith/CADKernel/Base/Private/Math/BSpline.cpp b/Engine/Source/Runtime/Datasmith/CADKernel/Base/Private/Math/BSpline.cpp
--- a/Engine/Source/Runtime/Datasmith/CADKernel/Base/Private/Math/BSpline.cpp
+++ b/Engine/Source/Runtime/Datasmith/CADKernel/Base/Private/Math/BSpline.cpp
@@ -272,148 +272,7 @@ void Interpolate1DBSpline(
 /*******************************fin fonction************************************/
 void InterpolerBSpline2D(const FNURBSSurface& Nurbs, FGrid& Grid)
 {
-	/*
-		//Nurbs.Display(TEXT("Nurbs");
-
-		// pour chaque Reseau U, on calcule les poles Pj(U),
-		// PointUCurveIndexVArray[CoordinateU][IndexV] point(CoordinateU) calculé sur la courbe Pole[IndexV]
-		// Les courbes Nurbs defini par les reseaux PointUCurveIndexVArray[CoordinateU] vont permettre de calculer les points finaux
-		TArray<TArray<FPointH>> PointUCurveIndexVArray;
-		PointUCurveIndexVArray.SetNum(Grid.GetNumU());
-		for (TArray<FPointH>& PointUCurveIndexV : PointUCurveIndexVArray)
-		{
-			PointUCurveIndexV.SetNum(Nurbs.GetPoleNumV());
-		}
-
-		// pour chaque Reseau Pj(U), on calcule les points P(UV)
-		//TBSplineData<FPointH> Data(FMath::Max(Nurbs.GetUDegree(), Nurbs.GetVDegree()) + 1);
-
-		//FTimePoint StartTime2 = FChrono::Now();
-		//TBSplineData<FPointH> Data2(FMath::Max(Nurbs.GetUDegree() , Grid.GetNumV()) + 1);
-		//FChrono::PrintClockElapse(LOG, "    ", "Alloc", FChrono::Elapse(StartTime2), ETimeUnit::NanoSeconds);
-
-		//FTimePoint StartTime = FChrono::Now();
-		TBSplineData<FPointH> Data(FMath::Max(Grid.GetNumU(), Grid.GetNumV()) + 1, FMath::Max(Nurbs.GetPoleNumU(), Nurbs.GetPoleNumV()) + 1);
-		//auto Stop = FChrono::Elapse(StartTime);
-		//FChrono::PrintClockElapse(LOG, "    ", "Alloc", Stop, ETimeUnit::NanoSeconds);
-
-		{
-			FPointH OutGradient;
-			TBSplineSurface<FPointH> Surface(Nurbs.GetPoleNumU(), Nurbs.GetUDegree(), Nurbs.GetNodalVectorU(), Nurbs.GetHomogeneousPoles());
-			//string Message = "Iteration IndexV " + Utils::ToString(IndexV);
-			Surface.Display(TEXT("in"));
-			Wait();
-			Data.SegmentIndex = 0;
-			for (int32 CoordinateU = 0; CoordinateU <Grid.GetNumU(); ++CoordinateU)
-			{
-				InterpolerPointBSpline1D(Surface, Data, Grid.GetUCoordinates()[CoordinateU], PointUCurveIndexVArray[CoordinateU]);
-			}
-		}
-
-		TBSplineSurface<FPointH> OutSurface(Nurbs.GetPoleNumV(), Nurbs.GetVDegree(), Nurbs.GetNodalVectorV(), PointUCurveIndexVArray);
-		OutSurface.Display(TEXT("Hello"));
-		Wait();
-
-		return;
-		//{
-		//TBSplineSurface<FPointH> OutSurface(Nurbs.GetPoleNumV(), Nurbs.GetVDegree(), Nurbs.GetVNodalVector(), PointUCurveIndexVArray);
-		//	FPointH OutGradient;
-		//	TBSplineSurface<FPointH> InSurface(Nurbs.GetPoleNumV(), Nurbs.GetVDegree(), Nurbs.GetVNodalVector(), Nurbs.GetHomogeneousPoles());
-		//	for (int32 IndexU = 0, Index = 0; IndexU <Grid.GetNumU(); ++IndexU)
-		//	{
-		//		InSurface.SetPoles(PointUCurveIndexVArray[IndexU].GetData() + Index);
-		//		//string Message = "Iteration IndexU " + Utils::ToString(IndexU);
-		//		//Curve.Display(Message);
-		//		Data.SegmentIndex = 0;
-		//		for (int32 CoordinateV = 0; CoordinateV <Grid.GetNumV(); ++CoordinateV)
-		//		{
-		//			FPointH Point;
-		//			InterpolerPointBSpline1D(InSurface, Data, Grid.GetVCoordinates()[CoordinateV], Point);
-		//			//Grid.Get3DPoints()[CoordinateV*Grid.GetNumU() + IndexU] = Point;
-		//		}
-		//	}
-		//}
-		*/
-}
-
-#ifdef UNDEF
-void InterpolerPointBSpline1D(const TBSplineSurface<FPointH>& Surface, TBSplineData<FPointH>& Data, double Coordinate, TArray<FPointH>& OutPoints)
-{
-	int32 StartIndex = Surface.Degre;
-	int32 EndIndex = Surface.PoleNum;
-
-	// Check the coordinate
-	if (Coordinate < Surface.NodalVector[StartIndex])
-	{
-		Coordinate = Surface.NodalVector[StartIndex];
-	}
-
-	if (Coordinate > Surface.NodalVector[EndIndex])
-	{
-		Coordinate = Surface.NodalVector[EndIndex];
-	}
-
-	// find the segment
-	double CoordinateTmp = Coordinate + SMALL_NUMBER_SQUARE;
-
-	int32 SegmentIndex = Data.SegmentIndex;
-	while (CoordinateTmp < Surface.NodalVector[SegmentIndex] && CoordinateTmp < Surface.NodalVector[SegmentIndex + 1])
-	{
-		SegmentIndex--;
-	}
-	for (; SegmentIndex <EndIndex - 1 && CoordinateTmp > Surface.NodalVector[SegmentIndex + 1]; ++SegmentIndex);
-	Data.SegmentIndex = SegmentIndex;
-
-	// compute the point
-	DeBoorValeursBSpline(Surface, Data, Coordinate, OutPoints);
-}
-
-void DeBoorValeursBSpline(const TBSplineSurface<FPointH>& Surface, TBSplineData<FPointH>& Data, const double Coordinate, TArray<FPointH>& OutPoints)
-{
-	const TArray<TArray<FPointH>>* CurrentSurfacePoles = &Surface.Poles;
-	TArray<TArray<FPointH>>* IterationSurfacePoles = &Data.IterativePoles1;
-
-
-	for (int32 IndexK = 1; IndexK <= Surface.Degre; IndexK++)
-	{
-		for (int32 Index = Data.SegmentIndex - Surface.Degre + IndexK; Index <= Data.SegmentIndex; Index++)
-		{
-			double Alpha = 0;
-			double Beta = 1;
-			double DeltaU = Surface.NodalVector[Index + Surface.Degre + 1 - IndexK] - Surface.NodalVector[Index];
-			if (DeltaU > 1.e-13)
-			{
-				DeltaU = 1 / DeltaU;
-				Alpha = (Coordinate - Surface.NodalVector[Index]) * DeltaU;
-				Beta = (Surface.NodalVector[Index + Surface.Degre + 1 - IndexK] - Coordinate) * DeltaU;
-				// Beta = 1 - Alpha ???
-			}
-
-			// Calcul du point courant
-			auto IterationPoles = IterationSurfacePoles->begin();
-			for (const auto& CurrentPoles : *CurrentSurfacePoles)
-			{
-				(*IterationPoles)[Index] = CurrentPoles[Index] * Alpha + CurrentPoles[Index - 1] * Beta;
-				++IterationPoles;
-			}
-		}
-
-		CurrentSurfacePoles = IterationSurfacePoles;
-		IterationSurfacePoles = (&Data.IterativePoles1 == IterationSurfacePoles) ? &Data.IterativePoles2 : &Data.IterativePoles1;
-	}
-
-	for (int32 Index = 0; Index < Surface.Poles.Num(); ++Index)
-	{
-		OutPoints[Index] = (*CurrentSurfacePoles)[Index][Data.SegmentIndex];
-	}
-
-	//Open3DDebugSession(TEXT("Result " + Utils::ToString(Coordinate));
-	//DisplayPoint(OutPoint);
-	//Close3DDebugSession();
-
-	//OutGradient =* IterationNewGradient;
 }
-#endif
 
 /******************************description fonction*****************************/
 /* Auteur:						 - Fichier: bspline.c										*/
@@ -793,8 +652,8 @@ void Evaluate2DPoint(const FNURBSCurve& Nurbs, double Coordinate, FCurvePoint2D&
 	{
 		if (Nurbs.IsRational())
 		{
-			OutPoint.Gradient.U = (GradientH[0] - GradientH[2] * ValueH[0] / ValueH[2]) / ValueH[2];
-			OutPoint.Gradient.V = (GradientH[1] - GradientH[2] * ValueH[1] / ValueH[2]) / ValueH[2];
+			OutPoint.Gradient.X = (GradientH[0] - GradientH[2] * ValueH[0] / ValueH[2]) / ValueH[2];
+			OutPoint.Gradient.Y = (GradientH[1] - GradientH[2] * ValueH[1] / ValueH[2]) / ValueH[2];
 		}
 		else
 		{
@@ -805,8 +664,8 @@ void Evaluate2DPoint(const FNURBSCurve& Nurbs, double Coordinate, FCurvePoint2D&
 		{
 			if (Nurbs.IsRational())
 			{
-				OutPoint.Laplacian.U = (LaplacianH[0] * ValueH[2] - LaplacianH[2] * ValueH[0] - 2.0 * GradientH[2] * (GradientH[0] - GradientH[2] * ValueH[0] / ValueH[2])) / (ValueH[2] * ValueH[2]);
-				OutPoint.Laplacian.V = (LaplacianH[1] * ValueH[2] - LaplacianH[2] * ValueH[1] - 2.0 * GradientH[2] * (GradientH[1] - GradientH[2] * ValueH[1] / ValueH[2])) / (ValueH[2] * ValueH[2]);
+				OutPoint.Laplacian.X = (LaplacianH[0] * ValueH[2] - LaplacianH[2] * ValueH[0] - 2.0 * GradientH[2] * (GradientH[0] - GradientH[2] * ValueH[0] / ValueH[2])) / (ValueH[2] * ValueH[2]);
+				OutPoint.Laplacian.Y = (LaplacianH[1] * ValueH[2] - LaplacianH[2] * ValueH[1] - 2.0 * GradientH[2] * (GradientH[1] - GradientH[2] * ValueH[1] / ValueH[2])) / (ValueH[2] * ValueH[2]);
 			}
 			else
 			{
@@ -847,7 +706,7 @@ void EvaluatePoint(const FNURBSSurface& Nurbs, const FPoint2D& InPoint2D, FSurfa
 
 	BSpline::Interpolate2DBSpline(Nurbs.GetDegree(EIso::IsoU), Nurbs.GetDegree(EIso::IsoV), Nurbs.GetPoleCount(EIso::IsoU), Nurbs.GetPoleCount(EIso::IsoV),
 		Nurbs.GetNodalVector(EIso::IsoU).GetData(), Nurbs.GetNodalVector(EIso::IsoV).GetData(), Nurbs.IsRational() ? 4 : 3, Nurbs.GetHPoles().GetData(),
-		1, &InPoint2D.U, 1, &InPoint2D.V, InDerivativeOrder, Point4, GradientU, GradientV, LaplacianU, LaplacianV, LaplacianUV);
+		1, &InPoint2D.X, 1, &InPoint2D.Y, InDerivativeOrder, Point4, GradientU, GradientV, LaplacianU, LaplacianV, LaplacianUV);
 
 	OutPoint3D.DerivativeOrder = InDerivativeOrder;
 
diff --git a/Engine/Source/Runtime/Datasmith/CADKernel/Base/Private/Math/Geometry.cpp b/Engine/Source/Runtime/Datasmith/CADKernel/Base/Private/Math/Geometry.cpp
--- a/Engine/Source/Runtime/Datasmith/CADKernel/Base/Private/Math/Geometry.cpp
+++ b/Engine/Source/Runtime/Datasmith/CADKernel/Base/Private/Math/Geometry.cpp
@@ -218,9 +218,9 @@ bool DoIntersect(const FSegment2D& SegmentAB, const FSegment2D& SegmentCD, TFunc
 					return false;
 				}
 
-				if (fabs(Context.AB.U) > fabs(Context.AB.V))
+				if (fabs(Context.AB.X) > fabs(Context.AB.Y))
 				{
-					if (FastIntersectionTestWhenCoincident(SegmentAB[0].U, SegmentAB[1].U, SegmentCD[0].U, SegmentCD[1].U))
+					if (FastIntersectionTestWhenCoincident(SegmentAB[0].X, SegmentAB[1].X, SegmentCD[0].X, SegmentCD[1].X))
 					{
 						return ConfirmIntersectionWhenNearlyCoincident(Context);
 					}
@@ -231,7 +231,7 @@ bool DoIntersect(const FSegment2D& SegmentAB, const FSegment2D& SegmentCD, TFunc
 				}
 				else
 				{
-					if (FastIntersectionTestWhenCoincident(SegmentAB[0].V, SegmentAB[1].V, SegmentCD[0].V, SegmentCD[1].V))
+					if (FastIntersectionTestWhenCoincident(SegmentAB[0].Y, SegmentAB[1].Y, SegmentCD[0].Y, SegmentCD[1].Y))
 					{
 						return ConfirmIntersectionWhenNearlyCoincident(Context);
 					}
diff --git a/Engine/Source/Runtime/Datasmith/CADKernel/Base/Private/Math/Point.cpp b/Engine/Source/Runtime/Datasmith/CADKernel/Base/Private/Math/Point.cpp
--- a/Engine/Source/Runtime/Datasmith/CADKernel/Base/Private/Math/Point.cpp
+++ b/Engine/Source/Runtime/Datasmith/CADKernel/Base/Private/Math/Point.cpp
@@ -13,18 +13,10 @@ const FPoint FPoint::UnitPoint(1., 1., 1.);
 const FPoint FPoint::FarawayPoint(HUGE_VALUE, HUGE_VALUE, HUGE_VALUE);
 const int32 FPoint::Dimension = 3;
 
-const FFPoint FFPoint::ZeroPoint(0.f, 0.f, 0.f);
-const FFPoint FFPoint::FarawayPoint(HUGE_VALUE, HUGE_VALUE, HUGE_VALUE);
-const int32 FFPoint::Dimension = 3;
-
 const FPoint2D FPoint2D::ZeroPoint(0., 0.);
 const FPoint2D FPoint2D::FarawayPoint(HUGE_VALUE, HUGE_VALUE);
 const int32 FPoint2D::Dimension = 2;
 
-const FPointH FPointH::ZeroPoint(0., 0., 0., 1.);
-const FPointH FPointH::FarawayPoint(HUGE_VALUE, HUGE_VALUE, HUGE_VALUE, 1.);
-const int32 FPointH::Dimension = 4;
-
 double FPoint::SignedAngle(const FPoint & Other, const FPoint & Normal) const
 {
 	FPoint Vector1 = *this; 
diff --git a/Engine/Source/Runtime/Datasmith/CADKernel/Base/Private/Mesh/Meshers/BowyerWatsonTriangulator.h b/Engine/Source/Runtime/Datasmith/CADKernel/Base/Private/Mesh/Meshers/BowyerWatsonTriangulator.h
--- a/Engine/Source/Runtime/Datasmith/CADKernel/Base/Private/Mesh/Meshers/BowyerWatsonTriangulator.h
+++ b/Engine/Source/Runtime/Datasmith/CADKernel/Base/Private/Mesh/Meshers/BowyerWatsonTriangulator.h
@@ -59,7 +59,7 @@ public:
 		F3DDebugSession DelaunayDebugSession(bDisplay, TEXT("Delaunay Algo"));
 #endif // DEBUG_DELAUNAY
 
-		Vertices.Sort([](const TPair<int32, FPoint2D>& Vertex1, const TPair<int32, FPoint2D>& Vertex2) {return (Vertex1.Value.U + Vertex1.Value.V) > (Vertex2.Value.U + Vertex2.Value.V); });
+		Vertices.Sort([](const TPair<int32, FPoint2D>& Vertex1, const TPair<int32, FPoint2D>& Vertex2) {return (Vertex1.Value.X + Vertex1.Value.Y) > (Vertex2.Value.X + Vertex2.Value.Y); });
 
 #ifdef DEBUG_BOWYERWATSON
 		if (bDisplay)
diff --git a/Engine/Source/Runtime/Datasmith/CADKernel/Base/Private/Mesh/Meshers/IsoTriangulator.cpp b/Engine/Source/Runtime/Datasmith/CADKernel/Base/Private/Mesh/Meshers/IsoTriangulator.cpp
--- a/Engine/Source/Runtime/Datasmith/CADKernel/Base/Private/Mesh/Meshers/IsoTriangulator.cpp
+++ b/Engine/Source/Runtime/Datasmith/CADKernel/Base/Private/Mesh/Meshers/IsoTriangulator.cpp
@@ -310,13 +310,13 @@ void FIsoTriangulator::FillMeshNodes()
 	for (FLoopNode& Node : LoopNodes)
 	{
 		const FPoint2D& UVCoordinate = Node.Get2DPoint(EGridSpace::Scaled, Grid);
-		Mesh.UVMap.Emplace(UVCoordinate.U, UVCoordinate.V);
+		Mesh.UVMap.Emplace(UVCoordinate.X, UVCoordinate.Y);
 	}
 
 	for (FIsoInnerNode& Node : InnerNodes)
 	{
 		const FPoint2D& UVCoordinate = Node.Get2DPoint(EGridSpace::Scaled, Grid);
-		Mesh.UVMap.Emplace(UVCoordinate.U, UVCoordinate.V);
+		Mesh.UVMap.Emplace(UVCoordinate.X, UVCoordinate.Y);
 	}
 }
 
@@ -554,9 +554,9 @@ void FIsoTriangulator::BuildInnerSegments()
 			{
 				for (const FPoint2D& LoopPoint : Loop)
 				{
-					if (FMath::IsNearlyEqual(LoopPoint.V, Node1.V, Tolerance))
+					if (FMath::IsNearlyEqual(LoopPoint.Y, Node1.Y, Tolerance))
 					{
-						if (Node1.U - DOUBLE_SMALL_NUMBER < LoopPoint.U && LoopPoint.U < Node2.U + DOUBLE_SMALL_NUMBER)
+						if (Node1.X - DOUBLE_SMALL_NUMBER < LoopPoint.X && LoopPoint.X < Node2.X + DOUBLE_SMALL_NUMBER)
 						{
 #ifdef DEBUG_BUILDINNERSEGMENTS
 							if (bDisplay)
@@ -578,9 +578,9 @@ void FIsoTriangulator::BuildInnerSegments()
 			{
 				for (const FPoint2D& LoopPoint : Loop)
 				{
-					if (FMath::IsNearlyEqual(LoopPoint.U, Node1.U, Tolerance))
+					if (FMath::IsNearlyEqual(LoopPoint.X, Node1.X, Tolerance))
 					{
-						if (Node1.V - DOUBLE_SMALL_NUMBER < LoopPoint.V && LoopPoint.V < Node2.V + DOUBLE_SMALL_NUMBER)
+						if (Node1.Y - DOUBLE_SMALL_NUMBER < LoopPoint.Y && LoopPoint.Y < Node2.Y + DOUBLE_SMALL_NUMBER)
 						{
 #ifdef DEBUG_BUILDINNERSEGMENTS
 							if (bDisplay)
@@ -1029,8 +1029,8 @@ void FIsoTriangulator::FindCellContainingBoundaryNodes(TArray<FCell>& Cells)
 			if (!LoopPoint.IsDelete())
 			{
 				const FPoint2D& Coordinate = LoopPoint.Get2DPoint(EGridSpace::UniformScaled, Grid);
-				ArrayUtils::FindCoordinateIndex(IsoUCoordinates, Coordinate.U, IndexU);
-				ArrayUtils::FindCoordinateIndex(IsoVCoordinates, Coordinate.V, IndexV);
+				ArrayUtils::FindCoordinateIndex(IsoUCoordinates, Coordinate.X, IndexU);
+				ArrayUtils::FindCoordinateIndex(IsoVCoordinates, Coordinate.Y, IndexV);
 
 #ifdef DEBUG_BUILD_CELLS
 				if (bDisplay)
@@ -1215,7 +1215,7 @@ void FIsoTriangulator::FindInnerGridCellSurroundingSmallLoop()
 		int32 IndexU = 0;
 		for (; IndexU < NumU - 1; ++IndexU)
 		{
-			if ((FirstPoint.U > UCoordinates[IndexU]) && (FirstPoint.U < UCoordinates[IndexU + 1] + DOUBLE_SMALL_NUMBER))
+			if ((FirstPoint.X > UCoordinates[IndexU]) && (FirstPoint.X < UCoordinates[IndexU + 1] + DOUBLE_SMALL_NUMBER))
 			{
 				break;
 			}
@@ -1224,7 +1224,7 @@ void FIsoTriangulator::FindInnerGridCellSurroundingSmallLoop()
 		int32 IndexV = 0;
 		for (; IndexV < NumV - 1; ++IndexV)
 		{
-			if ((FirstPoint.V > VCoordinates[IndexV]) && (FirstPoint.V < VCoordinates[IndexV + 1] + DOUBLE_SMALL_NUMBER))
+			if ((FirstPoint.Y > VCoordinates[IndexV]) && (FirstPoint.Y < VCoordinates[IndexV + 1] + DOUBLE_SMALL_NUMBER))
 			{
 				break;
 			}
@@ -1238,7 +1238,7 @@ void FIsoTriangulator::FindInnerGridCellSurroundingSmallLoop()
 		bool bBoudardyIsSurrounded = true;
 		for (const FPoint2D& LoopPoint : Loops[LoopIndex])
 		{
-			if (LoopPoint.U < UMin || LoopPoint.U > UMax || LoopPoint.V < VMin || LoopPoint.V > VMax)
+			if (LoopPoint.X < UMin || LoopPoint.X > UMax || LoopPoint.Y < VMin || LoopPoint.Y > VMax)
 			{
 				bBoudardyIsSurrounded = false;
 				break;
diff --git a/Engine/Source/Runtime/Datasmith/CADKernel/Base/Private/Mesh/Meshers/IsoTriangulator.h b/Engine/Source/Runtime/Datasmith/CADKernel/Base/Private/Mesh/Meshers/IsoTriangulator.h
--- a/Engine/Source/Runtime/Datasmith/CADKernel/Base/Private/Mesh/Meshers/IsoTriangulator.h
+++ b/Engine/Source/Runtime/Datasmith/CADKernel/Base/Private/Mesh/Meshers/IsoTriangulator.h
@@ -280,7 +280,7 @@ public:
 			{
 				const FPoint2D& Node1Coordinates = Node1->Get2DPoint(EGridSpace::Default2D, Grid);
 				const FPoint2D& Node2Coordinates = Node2->Get2DPoint(EGridSpace::Default2D, Grid);
-				return (Node1Coordinates.U + Node1Coordinates.V) < (Node2Coordinates.U + Node2Coordinates.V);
+				return (Node1Coordinates.X + Node1Coordinates.Y) < (Node2Coordinates.X + Node2Coordinates.Y);
 			});
 	}
 
diff --git a/Engine/Source/Runtime/Datasmith/CADKernel/Base/Private/Mesh/Meshers/IsoTriangulator/IntersectionIsoSegmentTool.h b/Engine/Source/Runtime/Datasmith/CADKernel/Base/Private/Mesh/Meshers/IsoTriangulator/IntersectionIsoSegmentTool.h
--- a/Engine/Source/Runtime/Datasmith/CADKernel/Base/Private/Mesh/Meshers/IsoTriangulator/IntersectionIsoSegmentTool.h
+++ b/Engine/Source/Runtime/Datasmith/CADKernel/Base/Private/Mesh/Meshers/IsoTriangulator/IntersectionIsoSegmentTool.h
@@ -77,7 +77,7 @@ class FIsoUSegment : public FIntersectionIsoSegment
 {
 public:
 	FIsoUSegment(const FPoint2D& StartPoint, const FPoint2D& EndPoint)
-		: FIntersectionIsoSegment(StartPoint, EndPoint, StartPoint.U, StartPoint.V, EndPoint.V)
+		: FIntersectionIsoSegment(StartPoint, EndPoint, StartPoint.X, StartPoint.Y, EndPoint.Y)
 	{}
 
 	virtual FPoint2D GetMinPoint() const override
@@ -94,7 +94,7 @@ class FIsoVSegment : public FIntersectionIsoSegment
 {
 public:
 	FIsoVSegment(const FPoint2D& StartPoint, const FPoint2D& EndPoint)
-		: FIntersectionIsoSegment(StartPoint, EndPoint, StartPoint.V, StartPoint.U, EndPoint.U)
+		: FIntersectionIsoSegment(StartPoint, EndPoint, StartPoint.Y, StartPoint.X, EndPoint.X)
 	{}
 
 	virtual FPoint2D GetMinPoint() const override
diff --git a/Engine/Source/Runtime/Datasmith/CADKernel/Base/Private/Mesh/Meshers/IsoTriangulator/IsoNode.h b/Engine/Source/Runtime/Datasmith/CADKernel/Base/Private/Mesh/Meshers/IsoTriangulator/IsoNode.h
--- a/Engine/Source/Runtime/Datasmith/CADKernel/Base/Private/Mesh/Meshers/IsoTriangulator/IsoNode.h
+++ b/Engine/Source/Runtime/Datasmith/CADKernel/Base/Private/Mesh/Meshers/IsoTriangulator/IsoNode.h
@@ -10,8 +10,6 @@ namespace UE::CADKernel
 {
 class FGrid;
 class FIsoSegment;
-class FPoint;
-class FPoint2D;
 
 enum class EIsoLink : uint8
 {
diff --git a/Engine/Source/Runtime/Datasmith/CADKernel/Base/Private/Mesh/Meshers/ParametricFaceMesher.cpp b/Engine/Source/Runtime/Datasmith/CADKernel/Base/Private/Mesh/Meshers/ParametricFaceMesher.cpp
--- a/Engine/Source/Runtime/Datasmith/CADKernel/Base/Private/Mesh/Meshers/ParametricFaceMesher.cpp
+++ b/Engine/Source/Runtime/Datasmith/CADKernel/Base/Private/Mesh/Meshers/ParametricFaceMesher.cpp
@@ -913,7 +913,7 @@ bool FParametricFaceMesher::MeshPlanarFace()
 		VertexList.Reserve(PointList.Num());
 		for (const FPoint2D& Point : PointList)
 		{
-			VertexList.Emplace(Point.U, Point.V);
+			VertexList.Emplace(Point.X, Point.Y);
 		}
 
 		FPolygon2d Polygon2d(VertexList);
diff --git a/Engine/Source/Runtime/Datasmith/CADKernel/Base/Private/Mesh/Structure/Grid.cpp b/Engine/Source/Runtime/Datasmith/CADKernel/Base/Private/Mesh/Structure/Grid.cpp
--- a/Engine/Source/Runtime/Datasmith/CADKernel/Base/Private/Mesh/Structure/Grid.cpp
+++ b/Engine/Source/Runtime/Datasmith/CADKernel/Base/Private/Mesh/Structure/Grid.cpp
@@ -145,8 +145,8 @@ void FGrid::GetPreferredUVCuttingParametersFromLoops(FCuttingGrid& CuttingParame
 
 			for (int32 Index = 0; Index < EdgePoints2D.Num(); ++Index)
 			{
-				CuttingParametersFromLoops[EIso::IsoU].Emplace(EdgePoints2D[Index].U, ECoordinateType::OtherCoordinate);
-				CuttingParametersFromLoops[EIso::IsoV].Emplace(EdgePoints2D[Index].V, ECoordinateType::OtherCoordinate);
+				CuttingParametersFromLoops[EIso::IsoU].Emplace(EdgePoints2D[Index].X, ECoordinateType::OtherCoordinate);
+				CuttingParametersFromLoops[EIso::IsoV].Emplace(EdgePoints2D[Index].Y, ECoordinateType::OtherCoordinate);
 			}
 		}
 	}
@@ -382,12 +382,12 @@ void FGrid::FindPointsCloseToLoop()
 				continue;
 			}
 
-			ABv = PointB->V - PointA->V;
-			ABu = PointB->U - PointA->U;
+			ABv = PointB->Y - PointA->Y;
+			ABu = PointB->X - PointA->X;
 			if (FMath::Abs(ABu) > FMath::Abs(ABv))
 			{
 				Slope = ABv / ABu;
-				Origin = PointA->V - Slope * PointA->U;
+				Origin = PointA->Y - Slope * PointA->X;
 				if (ABu > 0)
 				{
 					if (ABv > 0)
@@ -414,7 +414,7 @@ void FGrid::FindPointsCloseToLoop()
 			else
 			{
 				Slope = ABu / ABv;
-				Origin = PointA->U - Slope * PointA->V;
+				Origin = PointA->X - Slope * PointA->Y;
 				if (ABu > 0)
 				{
 					if (ABv > 0)
@@ -462,32 +462,32 @@ void FGrid::RemovePointsCloseToLoop()
 			: StartPoint(&SPoint)
 			, EndPoint(&EPoint)
 		{
-			StartPointWeight = StartPoint->U + StartPoint->V;
-			EndPointWeight = EndPoint->U + EndPoint->V;
+			StartPointWeight = StartPoint->X + StartPoint->Y;
+			EndPointWeight = EndPoint->X + EndPoint->Y;
 			if (StartPointWeight > EndPointWeight)
 			{
 				Swap(StartPointWeight, EndPointWeight);
 				Swap(StartPoint, EndPoint);
 			}
-			if (StartPoint->U < EndPoint->U)
+			if (StartPoint->X < EndPoint->X)
 			{
-				UMin = StartPoint->U;
-				UMax = EndPoint->U;
+				UMin = StartPoint->X;
+				UMax = EndPoint->X;
 			}
 			else
 			{
-				UMin = EndPoint->U;
-				UMax = StartPoint->U;
+				UMin = EndPoint->X;
+				UMax = StartPoint->X;
 			}
-			if (StartPoint->V < EndPoint->V)
+			if (StartPoint->Y < EndPoint->Y)
 			{
-				VMin = StartPoint->V;
-				VMax = EndPoint->V;
+				VMin = StartPoint->Y;
+				VMax = EndPoint->Y;
 			}
 			else
 			{
-				VMin = EndPoint->V;
-				VMax = StartPoint->V;
+				VMin = EndPoint->Y;
+				VMax = StartPoint->Y;
 			}
 		}
 	};
@@ -535,7 +535,7 @@ void FGrid::RemovePointsCloseToLoop()
 		SortedPointIndexes.Reserve(IndexOfPointsNearAndInsideLoop.Num());
 		for (const int32& Index : IndexOfPointsNearAndInsideLoop)
 		{
-			GridPointWeight.Add(Points2D[EGridSpace::UniformScaled][Index].U + Points2D[EGridSpace::UniformScaled][Index].V);
+			GridPointWeight.Add(Points2D[EGridSpace::UniformScaled][Index].X + Points2D[EGridSpace::UniformScaled][Index].Y);
 		}
 		for (int32 Index = 0; Index < IndexOfPointsNearAndInsideLoop.Num(); ++Index)
 		{
@@ -622,19 +622,19 @@ void FGrid::RemovePointsCloseToLoop()
 			{
 				continue;
 			}
-			if (Point2D.U + DeltaU < Segment.UMin)
+			if (Point2D.X + DeltaU < Segment.UMin)
 			{
 				continue;
 			}
-			if (Point2D.U - DeltaU > Segment.UMax)
+			if (Point2D.X - DeltaU > Segment.UMax)
 			{
 				continue;
 			}
-			if (Point2D.V + DeltaV < Segment.VMin)
+			if (Point2D.Y + DeltaV < Segment.VMin)
 			{
 				continue;
 			}
-			if (Point2D.V - DeltaV > Segment.VMax)
+			if (Point2D.Y - DeltaV > Segment.VMax)
 			{
 				continue;
 			}
@@ -645,7 +645,7 @@ void FGrid::RemovePointsCloseToLoop()
 			// If Projected point is in the oval center on Point2D => the node is too close
 			FPoint2D ProjectionToPoint = Abs(Point2D - Projection);
 
-			if (ProjectionToPoint.U > DeltaU * 0.05 || ProjectionToPoint.V > DeltaV * 0.05)
+			if (ProjectionToPoint.X > DeltaU * 0.05 || ProjectionToPoint.Y > DeltaV * 0.05)
 			{
 				continue;
 			}
@@ -665,8 +665,8 @@ void ScaleAndSwap(FPoint2D& Normal, const FPoint2D& StartPoint, const FSurfacicB
 {
 	Normal.Normalize();
 	FPoint2D MainDirection = Normal;
-	MainDirection.U /= Boundary[EIso::IsoU].Length();
-	MainDirection.V /= Boundary[EIso::IsoV].Length();
+	MainDirection.X /= Boundary[EIso::IsoU].Length();
+	MainDirection.Y /= Boundary[EIso::IsoV].Length();
 
 	TFunction<void(const EIso)> SwapAndScale = [&](const EIso Iso)
 	{
@@ -687,7 +687,7 @@ void ScaleAndSwap(FPoint2D& Normal, const FPoint2D& StartPoint, const FSurfacicB
 		Normal *= Boundary[Iso].Length() / 20;
 	};
 
-	if (MainDirection.U > MainDirection.V)
+	if (MainDirection.X > MainDirection.Y)
 	{
 		SwapAndScale(EIso::IsoU);
 	}
@@ -1175,8 +1175,8 @@ void FGrid::ScaleLoops()
 		{
 			const FPoint2D& Point = Loop[Index];
 
-			ArrayUtils::FindCoordinateIndex(CoordinateGrid[EIso::IsoU], Point.U, IndexU);
-			ArrayUtils::FindCoordinateIndex(CoordinateGrid[EIso::IsoV], Point.V, IndexV);
+			ArrayUtils::FindCoordinateIndex(CoordinateGrid[EIso::IsoU], Point.X, IndexU);
+			ArrayUtils::FindCoordinateIndex(CoordinateGrid[EIso::IsoV], Point.Y, IndexV);
 
 			ComputeNewCoordinate(Points2D[EGridSpace::Scaled], IndexU, IndexV, Point, ScaledLoop[Index]);
 			ComputeNewCoordinate(Points2D[EGridSpace::UniformScaled], IndexU, IndexV, Point, UniformScaledLoop[Index]);
@@ -1251,28 +1251,28 @@ void FGrid::FindInnerFacePoints()
 		{
 			for (FPoint2D& Point : Loop)
 			{
-				while (IndexV != 0 && (Point.V < UniformCuttingCoordinates[EIso::IsoV][IndexV]))
+				while (IndexV != 0 && (Point.Y < UniformCuttingCoordinates[EIso::IsoV][IndexV]))
 				{
 					IndexV--;
 				}
 				for (; IndexV < CuttingCount[EIso::IsoV]; ++IndexV)
 				{
-					if (Point.V + SmallToleranceV < UniformCuttingCoordinates[EIso::IsoV][IndexV])
+					if (Point.Y + SmallToleranceV < UniformCuttingCoordinates[EIso::IsoV][IndexV])
 					{
 						break;
 					}
-					if (Point.V - SmallToleranceV > UniformCuttingCoordinates[EIso::IsoV][IndexV])
+					if (Point.Y - SmallToleranceV > UniformCuttingCoordinates[EIso::IsoV][IndexV])
 					{
 						continue;
 					}
 
 					if (IndexV == 0)
 					{
-						Point.V += SmallToleranceV;
+						Point.Y += SmallToleranceV;
 					}
 					else
 					{
-						Point.V -= SmallToleranceV;
+						Point.Y -= SmallToleranceV;
 					}
 
 					break;
@@ -1282,28 +1282,28 @@ void FGrid::FindInnerFacePoints()
 					IndexV--;
 				}
 
-				while (IndexU != 0 && (Point.U < UniformCuttingCoordinates[EIso::IsoU][IndexU]))
+				while (IndexU != 0 && (Point.X < UniformCuttingCoordinates[EIso::IsoU][IndexU]))
 				{
 					IndexU--;
 				}
 				for (; IndexU < CuttingCount[EIso::IsoU]; ++IndexU)
 				{
-					if (Point.U + SmallToleranceU < UniformCuttingCoordinates[EIso::IsoU][IndexU])
+					if (Point.X + SmallToleranceU < UniformCuttingCoordinates[EIso::IsoU][IndexU])
 					{
 						break;
 					}
-					if (Point.U - SmallToleranceU > UniformCuttingCoordinates[EIso::IsoU][IndexU])
+					if (Point.X - SmallToleranceU > UniformCuttingCoordinates[EIso::IsoU][IndexU])
 					{
 						continue;
 					}
 
 					if (IndexU == 0)
 					{
-						Point.U += SmallToleranceU;
+						Point.X += SmallToleranceU;
 					}
 					else
 					{
-						Point.U -= SmallToleranceU;
+						Point.X -= SmallToleranceU;
 					}
 					break;
 				}
@@ -1322,10 +1322,10 @@ void FGrid::FindInnerFacePoints()
 		for (const FPoint2D& LoopPoint : Loop)
 		{
 			const FPoint2D* SecondSegmentPoint = &LoopPoint;
-			double UMin = FirstSegmentPoint->U;
-			double VMin = FirstSegmentPoint->V;
-			double Umax = SecondSegmentPoint->U;
-			double Vmax = SecondSegmentPoint->V;
+			double UMin = FirstSegmentPoint->X;
+			double VMin = FirstSegmentPoint->Y;
+			double Umax = SecondSegmentPoint->X;
+			double Vmax = SecondSegmentPoint->Y;
 			GetMinMax(UMin, Umax);
 			GetMinMax(VMin, Vmax);
 
@@ -1335,7 +1335,7 @@ void FGrid::FindInnerFacePoints()
 			// AB^AP = Pv*ABu - Pu*ABv + AuABVMinusAvABu
 			FPoint2D PointA;
 			FPoint2D PointB;
-			if (FirstSegmentPoint->V < SecondSegmentPoint->V)
+			if (FirstSegmentPoint->Y < SecondSegmentPoint->Y)
 			{
 				PointA = *FirstSegmentPoint;
 				PointB = *SecondSegmentPoint;
@@ -1345,9 +1345,9 @@ void FGrid::FindInnerFacePoints()
 				PointA = *SecondSegmentPoint;
 				PointB = *FirstSegmentPoint;
 			}
-			double ABv = PointB.V - PointA.V;
-			double ABu = PointB.U - PointA.U;
-			double AuABVMinusAvABu = PointA.U * ABv - PointA.V * ABu;
+			double ABv = PointB.Y - PointA.Y;
+			double ABu = PointB.X - PointA.X;
+			double AuABVMinusAvABu = PointA.X * ABv - PointA.Y * ABu;
 
 			int32 IndexV = 0;
 			int32 Index = 0;
@@ -1412,10 +1412,10 @@ void FGrid::FindInnerFacePoints()
 		for (const FPoint2D& LoopPoint : Loop)
 		{
 			const FPoint2D* SecondSegmentPoint = &LoopPoint;
-			double UMin = FirstSegmentPoint->U;
-			double VMin = FirstSegmentPoint->V;
-			double Umax = SecondSegmentPoint->U;
-			double Vmax = SecondSegmentPoint->V;
+			double UMin = FirstSegmentPoint->X;
+			double VMin = FirstSegmentPoint->Y;
+			double Umax = SecondSegmentPoint->X;
+			double Vmax = SecondSegmentPoint->Y;
 			GetMinMax(UMin, Umax);
 			GetMinMax(VMin, Vmax);
 
@@ -1425,7 +1425,7 @@ void FGrid::FindInnerFacePoints()
 			// AB^AP = Pv*ABu - Pu*ABv + AuABVMinusAvABu
 			FPoint2D PointA;
 			FPoint2D PointB;
-			if (FirstSegmentPoint->U < SecondSegmentPoint->U)
+			if (FirstSegmentPoint->X < SecondSegmentPoint->X)
 			{
 				PointA = *FirstSegmentPoint;
 				PointB = *SecondSegmentPoint;
@@ -1436,9 +1436,9 @@ void FGrid::FindInnerFacePoints()
 				PointB = *FirstSegmentPoint;
 			}
 
-			double ABu = PointB.U - PointA.U;
-			double ABv = PointB.V - PointA.V;
-			double AuABVMinusAvABu = PointA.U * ABv - PointA.V * ABu;
+			double ABu = PointB.X - PointA.X;
+			double ABv = PointB.Y - PointA.Y;
+			double AuABVMinusAvABu = PointA.X * ABv - PointA.Y * ABu;
 			int32 Index = 0;
 			for (int32 IndexU = 0; IndexU < CuttingCount[EIso::IsoU]; ++IndexU)
 			{
diff --git a/Engine/Source/Runtime/Datasmith/CADKernel/Base/Private/Mesh/Structure/Grid.h b/Engine/Source/Runtime/Datasmith/CADKernel/Base/Private/Mesh/Structure/Grid.h
--- a/Engine/Source/Runtime/Datasmith/CADKernel/Base/Private/Mesh/Structure/Grid.h
+++ b/Engine/Source/Runtime/Datasmith/CADKernel/Base/Private/Mesh/Structure/Grid.h
@@ -25,8 +25,6 @@ class FIsoNode;
 class FIsoSegment;
 class FLoopNode;
 class FModelMesh;
-class FPoint;
-class FPoint2D;
 class FThinZone2DFinder;
 class FTopologicalFace;
 
diff --git a/Engine/Source/Runtime/Datasmith/CADKernel/Base/Private/Mesh/Structure/GridBase.cpp b/Engine/Source/Runtime/Datasmith/CADKernel/Base/Private/Mesh/Structure/GridBase.cpp
--- a/Engine/Source/Runtime/Datasmith/CADKernel/Base/Private/Mesh/Structure/GridBase.cpp
+++ b/Engine/Source/Runtime/Datasmith/CADKernel/Base/Private/Mesh/Structure/GridBase.cpp
@@ -158,7 +158,7 @@ bool FGridBase::ScaleGrid()
 			{
 				NumUV = IPointV * CuttingCount[EIso::IsoU] + IPointU;
 				Length += LastPoint.Distance(Points3D[NumUV]);
-				Points2D[EGridSpace::Scaled][NumUV].Set(Points2D[EGridSpace::UniformScaled][NumUV].U, Length);
+				Points2D[EGridSpace::Scaled][NumUV].Set(Points2D[EGridSpace::UniformScaled][NumUV].X, Length);
 				LastPoint = Points3D[NumUV];
 			}
 
@@ -168,7 +168,7 @@ bool FGridBase::ScaleGrid()
 			{
 				NumUV = IPointV * CuttingCount[EIso::IsoU] + IPointU;
 				Length -= LastPoint.Distance(Points3D[NumUV]);
-				Points2D[EGridSpace::Scaled][NumUV].Set(Points2D[EGridSpace::UniformScaled][NumUV].U, Length);
+				Points2D[EGridSpace::Scaled][NumUV].Set(Points2D[EGridSpace::UniformScaled][NumUV].X, Length);
 				LastPoint = Points3D[NumUV];
 			}
 		}
@@ -196,7 +196,7 @@ bool FGridBase::ScaleGrid()
 			{
 				NumUV = IPointV * CuttingCount[EIso::IsoU] + IPointU;
 				Length += LastPoint.Distance(Points3D[NumUV]);
-				Points2D[EGridSpace::Scaled][NumUV].Set(Length, Points2D[EGridSpace::UniformScaled][NumUV].V);
+				Points2D[EGridSpace::Scaled][NumUV].Set(Length, Points2D[EGridSpace::UniformScaled][NumUV].Y);
 				LastPoint = Points3D[NumUV];
 			}
 
@@ -206,7 +206,7 @@ bool FGridBase::ScaleGrid()
 			{
 				NumUV = IPointV * CuttingCount[EIso::IsoU] + IPointU;
 				Length -= LastPoint.Distance(Points3D[NumUV]);
-				Points2D[EGridSpace::Scaled][NumUV].Set(Length, Points2D[EGridSpace::UniformScaled][NumUV].V);
+				Points2D[EGridSpace::Scaled][NumUV].Set(Length, Points2D[EGridSpace::UniformScaled][NumUV].Y);
 				LastPoint = Points3D[NumUV];
 			}
 		}
@@ -229,8 +229,8 @@ void FGridBase::TransformPoints(EGridSpace DestinationSpace, const TArray<FPoint
 	{
 		const FPoint2D& Point = InPointsToScale[Index];
 
-		ArrayUtils::FindCoordinateIndex(CoordinateGrid[EIso::IsoU], Point.U, IndexU);
-		ArrayUtils::FindCoordinateIndex(CoordinateGrid[EIso::IsoV], Point.V, IndexV);
+		ArrayUtils::FindCoordinateIndex(CoordinateGrid[EIso::IsoU], Point.X, IndexU);
+		ArrayUtils::FindCoordinateIndex(CoordinateGrid[EIso::IsoV], Point.Y, IndexV);
 
 		ComputeNewCoordinate(Points2D[DestinationSpace], IndexU, IndexV, Point, OutTransformedPoints[Index]);
 	}
diff --git a/Engine/Source/Runtime/Datasmith/CADKernel/Base/Private/Mesh/Structure/GridBase.h b/Engine/Source/Runtime/Datasmith/CADKernel/Base/Private/Mesh/Structure/GridBase.h
--- a/Engine/Source/Runtime/Datasmith/CADKernel/Base/Private/Mesh/Structure/GridBase.h
+++ b/Engine/Source/Runtime/Datasmith/CADKernel/Base/Private/Mesh/Structure/GridBase.h
@@ -154,9 +154,9 @@ protected:
 		const FPoint2D& PointU1V1 = NewGrid[(IndexV + 1) * CuttingCount[EIso::IsoU] + (IndexU + 1)];
 
 		const double U1MinusU0 = CoordinateGrid[EIso::IsoU][IndexU + 1] - CoordinateGrid[EIso::IsoU][IndexU];
-		const double U0MinusU  = CoordinateGrid[EIso::IsoU][IndexU] - InPoint.U;
+		const double U0MinusU  = CoordinateGrid[EIso::IsoU][IndexU] - InPoint.X;
 		const double V1MinusV0 = CoordinateGrid[EIso::IsoV][IndexV + 1] - CoordinateGrid[EIso::IsoV][IndexV];
-		const double V0MinusV  = CoordinateGrid[EIso::IsoV][IndexV] - InPoint.V;
+		const double V0MinusV  = CoordinateGrid[EIso::IsoV][IndexV] - InPoint.Y;
 
 		OutNewScaledPoint =
 			PointU0V0 +
diff --git a/Engine/Source/Runtime/Datasmith/CADKernel/Base/Private/Mesh/Structure/LoopCleaner.cpp b/Engine/Source/Runtime/Datasmith/CADKernel/Base/Private/Mesh/Structure/LoopCleaner.cpp
--- a/Engine/Source/Runtime/Datasmith/CADKernel/Base/Private/Mesh/Structure/LoopCleaner.cpp
+++ b/Engine/Source/Runtime/Datasmith/CADKernel/Base/Private/Mesh/Structure/LoopCleaner.cpp
@@ -2245,25 +2245,25 @@ void FLoopCleaner::FindBestLoopExtremity()
 		Index++;
 		const FPoint2D& Point = Node.Get2DPoint(EGridSpace::UniformScaled, Grid);
 
-		if (Point.U > UMax)
+		if (Point.X > UMax)
 		{
-			UMax = Point.U;
+			UMax = Point.X;
 			ExtremityNodes[0] = &Node;
 		}
-		if (Point.U < UMin)
+		if (Point.X < UMin)
 		{
-			UMin = Point.U;
+			UMin = Point.X;
 			ExtremityNodes[1] = &Node;
 		}
 
-		if (Point.V > VMax)
+		if (Point.Y > VMax)
 		{
-			VMax = Point.V;
+			VMax = Point.Y;
 			ExtremityNodes[2] = &Node;
 		}
-		if (Point.V < VMin)
+		if (Point.Y < VMin)
 		{
-			VMin = Point.V;
+			VMin = Point.Y;
 			ExtremityNodes[3] = &Node;
 		}
 	}
@@ -2338,25 +2338,25 @@ EOrientation FLoopCleaner::GetLoopOrientation(const FLoopNode* StartNode)
 			case 0:
 			{
 				// UMax
-				Slope += PreviousPoint.V < NextPoint.V ? 0.1 : -0.1;
+				Slope += PreviousPoint.Y < NextPoint.Y ? 0.1 : -0.1;
 				break;
 			}
 			case 1:
 			{
 				// UMin
-				Slope += PreviousPoint.V > NextPoint.V ? 0.1 : -0.1;
+				Slope += PreviousPoint.Y > NextPoint.Y ? 0.1 : -0.1;
 				break;
 			}
 			case 2:
 			{
 				// VMax
-				Slope += PreviousPoint.U > NextPoint.U ? 0.1 : -0.1;
+				Slope += PreviousPoint.X > NextPoint.X ? 0.1 : -0.1;
 				break;
 			}
 			case 3:
 			{
 				// VMin
-				Slope += PreviousPoint.U < NextPoint.U ? 0.1 : -0.1;
+				Slope += PreviousPoint.X < NextPoint.X ? 0.1 : -0.1;
 				break;
 			}
 			default:
@@ -2380,25 +2380,25 @@ EOrientation FLoopCleaner::GetLoopOrientation(const FLoopNode* StartNode)
 	{
 		const FPoint2D& Point = Node->Get2DPoint(EGridSpace::UniformScaled, Grid);
 
-		if (Point.U > UMax)
+		if (Point.X > UMax)
 		{
-			UMax = Point.U;
+			UMax = Point.X;
 			ExtremityNodes[0] = Node;
 		}
-		if (Point.U < UMin)
+		if (Point.X < UMin)
 		{
-			UMin = Point.U;
+			UMin = Point.X;
 			ExtremityNodes[1] = Node;
 		}
 
-		if (Point.V > VMax)
+		if (Point.Y > VMax)
 		{
-			VMax = Point.V;
+			VMax = Point.Y;
 			ExtremityNodes[2] = Node;
 		}
-		if (Point.V < VMin)
+		if (Point.Y < VMin)
 		{
-			VMin = Point.V;
+			VMin = Point.Y;
 			ExtremityNodes[3] = Node;
 		}
 		Node = LoopCleanerImpl::GetNextConstNodeImpl(Node);
diff --git a/Engine/Source/Runtime/Datasmith/CADKernel/Base/Private/Topo/TopologicalFaceUtilities.cpp b/Engine/Source/Runtime/Datasmith/CADKernel/Base/Private/Topo/TopologicalFaceUtilities.cpp
--- a/Engine/Source/Runtime/Datasmith/CADKernel/Base/Private/Topo/TopologicalFaceUtilities.cpp
+++ b/Engine/Source/Runtime/Datasmith/CADKernel/Base/Private/Topo/TopologicalFaceUtilities.cpp
@@ -72,7 +72,7 @@ namespace UE::CADKernel::TopologicalFaceUtilities
 			// The entire polyline is used for the edge, just copy it in
 			for (const FPoint2D& Point : PolyPoints)
 			{
-				PointsOut.Emplace(Point.U, Point.V);
+				PointsOut.Emplace(Point.X, Point.Y);
 			}
 		}
 		else
@@ -84,7 +84,7 @@ namespace UE::CADKernel::TopologicalFaceUtilities
 						const double Ratio = (Value - PolyCoords[0]) / (PolyCoords[1] - PolyCoords[0]);
 						const FPoint2D Point = PolyPoints[0] + (PolyPoints[1] - PolyPoints[0]) * Ratio;
 
-						PointsOut.Emplace(Point.U, Point.V);
+						PointsOut.Emplace(Point.X, Point.Y);
 					}
 					else if (Value > PolyBounds.Max)
 					{
@@ -92,7 +92,7 @@ namespace UE::CADKernel::TopologicalFaceUtilities
 						const double Ratio = (Value - PolyCoords[LastIndex]) / (PolyCoords[LastIndex] - PolyCoords[LastIndex - 1]);
 						const FPoint2D Point = PolyPoints[LastIndex] + (PolyPoints[LastIndex] - PolyPoints[LastIndex - 1]) * Ratio;
 
-						PointsOut.Emplace(Point.U, Point.V);
+						PointsOut.Emplace(Point.X, Point.Y);
 					}
 					else
 					{
@@ -101,14 +101,14 @@ namespace UE::CADKernel::TopologicalFaceUtilities
 						if (FMath::IsNearlyEqual(Value, PolyCoords[Index]))
 						{
 							const FPoint2D& Point = PolyPoints[Index];
-							PointsOut.Emplace(Point.U, Point.V);
+							PointsOut.Emplace(Point.X, Point.Y);
 						}
 						else
 						{
 							const double Ratio = (Value - PolyCoords[Index - 1]) / (PolyCoords[Index] - PolyCoords[Index - 1]);
 							const FPoint2D Point = PolyPoints[Index - 1] + (PolyPoints[Index] - PolyPoints[Index - 1]) * Ratio;
 
-							PointsOut.Emplace(Point.U, Point.V);
+							PointsOut.Emplace(Point.X, Point.Y);
 						}
 					}
 				};
@@ -124,7 +124,7 @@ namespace UE::CADKernel::TopologicalFaceUtilities
 			ensureCADKernel(EndIndex < PolyCoords.Num());
 			for (int32 Index = StartIndex; Index < EndIndex; ++Index)
 			{
-				PointsOut.Emplace(PolyPoints[Index].U, PolyPoints[Index].V);
+				PointsOut.Emplace(PolyPoints[Index].X, PolyPoints[Index].Y);
 			}
 
 			InsertPoint(EdgeBounds.Max);
@@ -165,7 +165,7 @@ namespace UE::CADKernel::TopologicalFaceUtilities
 		using namespace UE::CADKernel;
 
 		const FSurfacicTolerance& IsoTolerances = Loop.GetFace()->GetCarrierSurface()->GetIsoTolerances();
-		const FVector2d Tolerance2D { IsoTolerances.U, IsoTolerances.V };
+		const FVector2d Tolerance2D { IsoTolerances.X, IsoTolerances.Y };
 
 		TArray<FVector2d> VertexList;
 		// #cad_debug
diff --git a/Engine/Source/Runtime/Datasmith/CADKernel/Base/Private/Topo/TopologicalLoop.cpp b/Engine/Source/Runtime/Datasmith/CADKernel/Base/Private/Topo/TopologicalLoop.cpp
--- a/Engine/Source/Runtime/Datasmith/CADKernel/Base/Private/Topo/TopologicalLoop.cpp
+++ b/Engine/Source/Runtime/Datasmith/CADKernel/Base/Private/Topo/TopologicalLoop.cpp
@@ -32,7 +32,7 @@ TSharedPtr<FTopologicalLoop> FTopologicalLoop::Make(const TArray<TSharedPtr<FTop
 
 		const FPoint2D StartPoint = Edge0.Approximate2DPoint(Edge0.GetBoundary().GetMin());
 		const FPoint2D EndPoint = Edge0.Approximate2DPoint(Edge0.GetBoundary().GetMax());
-		if (FMath::IsNearlyEqual(FMath::Abs(StartPoint.V), DOUBLE_HALF_PI) && FMath::IsNearlyEqual(FMath::Abs(EndPoint.V), DOUBLE_HALF_PI))
+		if (FMath::IsNearlyEqual(FMath::Abs(StartPoint.Y), DOUBLE_HALF_PI) && FMath::IsNearlyEqual(FMath::Abs(EndPoint.Y), DOUBLE_HALF_PI))
 		{
 			const double TolU = Surface.GetIsoTolerance(EIso::IsoU);
 			const bool Edge0IsIso = Edge0.GetCurve()->IsIso(EIso::IsoU, TolU);
@@ -301,25 +301,25 @@ bool FTopologicalLoop::Orient()
 
 		for (int32 Index = StartIndex; Index != EndIndex; Index += Increment)
 		{
-			if (LoopSampling[Index].U > UMax)
+			if (LoopSampling[Index].X > UMax)
 			{
-				UMax = LoopSampling[Index].U;
+				UMax = LoopSampling[Index].X;
 				IndexUMax = Index;
 			}
-			if (LoopSampling[Index].U < UMin)
+			if (LoopSampling[Index].X < UMin)
 			{
-				UMin = LoopSampling[Index].U;
+				UMin = LoopSampling[Index].X;
 				IndexUMin = Index;
 			}
 
-			if (LoopSampling[Index].V > VMax)
+			if (LoopSampling[Index].Y > VMax)
 			{
-				VMax = LoopSampling[Index].V;
+				VMax = LoopSampling[Index].Y;
 				IndexVMax = Index;
 			}
-			if (LoopSampling[Index].V < VMin)
+			if (LoopSampling[Index].Y < VMin)
 			{
-				VMin = LoopSampling[Index].V;
+				VMin = LoopSampling[Index].Y;
 				IndexVMin = Index;
 			}
 		}
@@ -399,19 +399,19 @@ bool FTopologicalLoop::Orient()
 		{
 			// Extrema case: the slop is compute between the next segment and the nearest BBox side 
 			double ReferenceSlope = 0;
-			if (FMath::IsNearlyEqual(LoopSampling[Index].U, UMin))
+			if (FMath::IsNearlyEqual(LoopSampling[Index].X, UMin))
 			{
 				ReferenceSlope = 6;
 			}
-			else if (FMath::IsNearlyEqual(LoopSampling[Index].U, UMax))
+			else if (FMath::IsNearlyEqual(LoopSampling[Index].X, UMax))
 			{
 				ReferenceSlope = 2;
 			}
-			else if (FMath::IsNearlyEqual(LoopSampling[Index].V, VMin))
+			else if (FMath::IsNearlyEqual(LoopSampling[Index].Y, VMin))
 			{
 				ReferenceSlope = 0;
 			}
-			else if (FMath::IsNearlyEqual(LoopSampling[Index].V, VMax))
+			else if (FMath::IsNearlyEqual(LoopSampling[Index].Y, VMax))
 			{
 				ReferenceSlope = 4;
 			}
@@ -949,10 +949,10 @@ void FTopologicalLoop::EnsureLogicalClosing(const double Tolerance3D)
 				}
 			}
 		}
-		else if (FMath::Abs(Gap2D.U) > LargeGapTolerance2D.U || FMath::Abs(Gap2D.V) > LargeGapTolerance2D.V)
+		else if (FMath::Abs(Gap2D.X) > LargeGapTolerance2D.X || FMath::Abs(Gap2D.Y) > LargeGapTolerance2D.Y)
 		{
 			// if the gap is not so large and the extremity of on edge is tangent to the gap segment then we extend the edge to close the gap
-			FMessage::Printf(Debug, TEXT("Loop %d Gap 2D : [%f, %f] vs Tol2D [%f, %f]\n"), Id, FMath::Abs(Gap2D.U), FMath::Abs(Gap2D.V), Tolerance2D.U, Tolerance2D.V);
+			FMessage::Printf(Debug, TEXT("Loop %d Gap 2D : [%f, %f] vs Tol2D [%f, %f]\n"), Id, FMath::Abs(Gap2D.X), FMath::Abs(Gap2D.Y), Tolerance2D.X, Tolerance2D.Y);
 
 			// If the gap in 2d is big e.g. side of a degenerated patch => build an edge
 			TSharedPtr<FTopologicalEdge> Edge = FTopologicalEdge::Make(Surface, PreviousExtremity.Point2D, PreviousEdgeEndVertex, EdgeExtremities[ExtremityIndex].Point2D, EdgeStartVertex);
@@ -965,10 +965,10 @@ void FTopologicalLoop::EnsureLogicalClosing(const double Tolerance3D)
 				++Index;
 			}
 		}
-		else if (FMath::Abs(Gap2D.U) > SmallGapTolerance2D.U || FMath::Abs(Gap2D.V) > SmallGapTolerance2D.V)
+		else if (FMath::Abs(Gap2D.X) > SmallGapTolerance2D.X || FMath::Abs(Gap2D.Y) > SmallGapTolerance2D.Y)
 		{
 			// if the gap is not so large and the extremity of on edge is tangent to the gap segment then we extend the edge to close the gap
-			FMessage::Printf(Log, TEXT("Loop %d Gap 2D : [%f, %f] vs Tol2D [%f, %f]\n"), Id, FMath::Abs(Gap2D.U), FMath::Abs(Gap2D.V), Tolerance2D.U, Tolerance2D.V);
+			FMessage::Printf(Log, TEXT("Loop %d Gap 2D : [%f, %f] vs Tol2D [%f, %f]\n"), Id, FMath::Abs(Gap2D.X), FMath::Abs(Gap2D.Y), Tolerance2D.X, Tolerance2D.Y);
 
 			const FPoint2D PreviousTangent = PreviousEdge.Entity->GetTangent2DAt(*PreviousEdgeEndVertex);
 			const FPoint2D EdgeTangent = OrientedEdge.Entity->GetTangent2DAt(*EdgeStartVertex);
diff --git a/Engine/Source/Runtime/Datasmith/CADKernel/Base/Private/UI/Display.cpp b/Engine/Source/Runtime/Datasmith/CADKernel/Base/Private/UI/Display.cpp
--- a/Engine/Source/Runtime/Datasmith/CADKernel/Base/Private/UI/Display.cpp
+++ b/Engine/Source/Runtime/Datasmith/CADKernel/Base/Private/UI/Display.cpp
@@ -477,10 +477,10 @@ void DisplayAABB(const FAABB& Aabb, FIdent Ident)
 void DisplayAABB2D(const FAABB2D& aabb, FIdent Ident)
 {
 #ifdef CADKERNEL_DEBUG
-	FPoint A(aabb.GetMin().U, aabb.GetMin().V, 0.);
-	FPoint B(aabb.GetMin().U, aabb.GetMax().V, 0.);
-	FPoint C(aabb.GetMax().U, aabb.GetMax().V, 0.);
-	FPoint D(aabb.GetMax().U, aabb.GetMin().V, 0.);
+	FPoint A(aabb.GetMin().X, aabb.GetMin().Y, 0.);
+	FPoint B(aabb.GetMin().X, aabb.GetMax().Y, 0.);
+	FPoint C(aabb.GetMax().X, aabb.GetMax().Y, 0.);
+	FPoint D(aabb.GetMax().X, aabb.GetMin().Y, 0.);
 
 	TArray<FPoint> Points;
 	Points.Add(A);
@@ -1285,14 +1285,6 @@ void DisplaySegment(const FPoint2D& Point1, const FPoint2D& Point2, FIdent Ident
 #endif
 }
 
-void DisplaySegment(const FPointH& Point1, const FPointH& Point2, FIdent Ident, EVisuProperty Property)
-{
-#ifdef CADKERNEL_DEBUG
-	F3DDebugSegment G(Ident);
-	DrawSegment(Point1, Point2, Property);
-#endif
-}
-
 void DisplayLoop(const FTopologicalFace& Surface)
 {
 #ifdef CADKERNEL_DEBUG
diff --git a/Engine/Source/Runtime/Datasmith/CADKernel/Base/Public/Geo/GeoPoint.h b/Engine/Source/Runtime/Datasmith/CADKernel/Base/Public/Geo/GeoPoint.h
--- a/Engine/Source/Runtime/Datasmith/CADKernel/Base/Public/Geo/GeoPoint.h
+++ b/Engine/Source/Runtime/Datasmith/CADKernel/Base/Public/Geo/GeoPoint.h
@@ -85,16 +85,16 @@ struct FCurvePoint
 
 		if (DerivativeOrder > 0)
 		{
-			Gradient = SurfacicPoint.GradientU * Point2D.Gradient.U + SurfacicPoint.GradientV * Point2D.Gradient.V;
+			Gradient = SurfacicPoint.GradientU * Point2D.Gradient.X + SurfacicPoint.GradientV * Point2D.Gradient.Y;
 		}
 
 		if (DerivativeOrder > 1)
 		{
-			Laplacian = SurfacicPoint.LaplacianU * FMath::Square(Point2D.Gradient.U)
-				+ 2.0 * SurfacicPoint.LaplacianUV * Point2D.Gradient.U * Point2D.Gradient.V
-				+ SurfacicPoint.LaplacianV * FMath::Square(Point2D.Gradient.V)
-				+ SurfacicPoint.GradientU * Point2D.Laplacian.U
-				+ SurfacicPoint.GradientV * Point2D.Laplacian.V;
+			Laplacian = SurfacicPoint.LaplacianU * FMath::Square(Point2D.Gradient.X)
+				+ 2.0 * SurfacicPoint.LaplacianUV * Point2D.Gradient.X * Point2D.Gradient.Y
+				+ SurfacicPoint.LaplacianV * FMath::Square(Point2D.Gradient.Y)
+				+ SurfacicPoint.GradientU * Point2D.Laplacian.X
+				+ SurfacicPoint.GradientV * Point2D.Laplacian.Y;
 		}
 	}
 };
diff --git a/Engine/Source/Runtime/Datasmith/CADKernel/Base/Public/Math/Boundary.h b/Engine/Source/Runtime/Datasmith/CADKernel/Base/Public/Math/Boundary.h
--- a/Engine/Source/Runtime/Datasmith/CADKernel/Base/Public/Math/Boundary.h
+++ b/Engine/Source/Runtime/Datasmith/CADKernel/Base/Public/Math/Boundary.h
@@ -281,8 +281,8 @@ public:
 
 	void Set(const FPoint2D& Point1, const FPoint2D& Point2)
 	{
-		UVBoundaries[EIso::IsoU].Set(Point1.U, Point2.U);
-		UVBoundaries[EIso::IsoV].Set(Point1.V, Point2.V);
+		UVBoundaries[EIso::IsoU].Set(Point1.X, Point2.X);
+		UVBoundaries[EIso::IsoV].Set(Point1.Y, Point2.Y);
 	}
 
 	friend FArchive& operator<<(FArchive& Ar, FSurfacicBoundary& Boundary)
@@ -395,8 +395,8 @@ public:
 
 	void ExtendTo(const FPoint2D& Point)
 	{
-		UVBoundaries[EIso::IsoU].ExtendTo(Point.U);
-		UVBoundaries[EIso::IsoV].ExtendTo(Point.V);
+		UVBoundaries[EIso::IsoU].ExtendTo(Point.X);
+		UVBoundaries[EIso::IsoV].ExtendTo(Point.Y);
 	}
 
 	void ExtendTo(const FPoint& Point)
@@ -434,8 +434,8 @@ public:
 	 */
 	void MoveInsideIfNot(FPoint2D& Point, const double Tolerance = DOUBLE_SMALL_NUMBER) const
 	{
-		UVBoundaries[EIso::IsoU].MoveInsideIfNot(Point.U, Tolerance);
-		UVBoundaries[EIso::IsoV].MoveInsideIfNot(Point.V, Tolerance);
+		UVBoundaries[EIso::IsoU].MoveInsideIfNot(Point.X, Tolerance);
+		UVBoundaries[EIso::IsoV].MoveInsideIfNot(Point.Y, Tolerance);
 	}
 
 	double Length(const EIso& Iso) const
diff --git a/Engine/Source/Runtime/Datasmith/CADKernel/Base/Public/Math/Geometry.h b/Engine/Source/Runtime/Datasmith/CADKernel/Base/Public/Math/Geometry.h
--- a/Engine/Source/Runtime/Datasmith/CADKernel/Base/Public/Math/Geometry.h
+++ b/Engine/Source/Runtime/Datasmith/CADKernel/Base/Public/Math/Geometry.h
@@ -45,8 +45,8 @@ CADKERNEL_API inline FPoint2D ComputeCircumCircleCenter(const FPoint2D& InPoint0
 	// CenterV  = 1/D * (Bu.|B|.|B| - Cu.|C|.|C|) = -1/D * SquareNorms ^ CBu
 	double SquareNormB = Segment_P0_P1.SquareLength();
 	double SquareNormC = Segment_P0_P2.SquareLength();
-	double CenterU = (SquareNormB * Segment_P0_P2.V - SquareNormC * Segment_P0_P1.V) / D;
-	double CenterV = (SquareNormC * Segment_P0_P1.U - SquareNormB * Segment_P0_P2.U) / D;
+	double CenterU = (SquareNormB * Segment_P0_P2.Y - SquareNormC * Segment_P0_P1.Y) / D;
+	double CenterV = (SquareNormC * Segment_P0_P1.X - SquareNormB * Segment_P0_P2.X) / D;
 
 	return FPoint2D(CenterU, CenterV) + InPoint0;
 }
@@ -76,8 +76,8 @@ CADKERNEL_API inline FPoint ComputeCircumCircleCenter(const FPoint& Point0, cons
 	double SquareNormB = D2Point1.SquareLength();
 	double SquareNormC = D2Point2.SquareLength();
 
-	double CenterU = (SquareNormB * D2Point2.V - SquareNormC * D2Point1.V) / D;
-	double CenterV = (SquareNormC * D2Point1.U - SquareNormB * D2Point2.U) / D;
+	double CenterU = (SquareNormB * D2Point2.Y - SquareNormC * D2Point1.Y) / D;
+	double CenterV = (SquareNormC * D2Point1.X - SquareNormB * D2Point2.X) / D;
 
 	FPoint Center2D(CenterU, CenterV, 0);
 	return Matrix * Center2D;
@@ -98,8 +98,8 @@ CADKERNEL_API inline FPoint2D ComputeCircumCircleCenterAndSquareRadius(const FPo
 
 	double SquareNormB = Segment_P0_P1.SquareLength();
 	double SquareNormC = Segment_P0_P2.SquareLength();
-	double CenterU = (SquareNormB * Segment_P0_P2.V - SquareNormC * Segment_P0_P1.V) / D;
-	double CenterV = (SquareNormC * Segment_P0_P1.U - SquareNormB * Segment_P0_P2.U) / D;
+	double CenterU = (SquareNormB * Segment_P0_P2.Y - SquareNormC * Segment_P0_P1.Y) / D;
+	double CenterV = (SquareNormC * Segment_P0_P1.X - SquareNormB * Segment_P0_P2.X) / D;
 
 	FPoint2D Center(CenterU, CenterV);
 	OutSquareRadius = Center.SquareLength();
@@ -211,29 +211,29 @@ struct CADKERNEL_API TTriangle
 			double Segment1PointSegment01 = Segment1Point * Segment01;
 			double Segment1PointSegment02 = Segment1Point * Segment02;
 
-			OutCoordinate.U = ((Segment1PointSegment01 * SquareLength02) - (Segment1PointSegment02 * Seg01Seg02)) / Det;
-			OutCoordinate.V = ((Segment1PointSegment02 * SquareLength01) - (Segment1PointSegment02 * Seg01Seg02)) / Det;
+			OutCoordinate.X = ((Segment1PointSegment01 * SquareLength02) - (Segment1PointSegment02 * Seg01Seg02)) / Det;
+			OutCoordinate.Y = ((Segment1PointSegment02 * SquareLength01) - (Segment1PointSegment02 * Seg01Seg02)) / Det;
 
 			// tester la solution pour choisir parmi 4 possibilites
-			if (OutCoordinate.U < 0.0)
+			if (OutCoordinate.X < 0.0)
 			{
 				// the project point is on the segment 02
 				SideIndex = Side20;
 			}
-			else if (OutCoordinate.V < 0.0)
+			else if (OutCoordinate.Y < 0.0)
 			{
 				// the project point is on the segment 01
 				SideIndex = Side01;
 			}
-			else if ((OutCoordinate.U + OutCoordinate.V) > 1.0)
+			else if ((OutCoordinate.X + OutCoordinate.Y) > 1.0)
 			{
 				// the project point is on the segment 12
 				SideIndex = Side12;
 			}
 			else {
 				// the project point is inside the Segment
-				Segment01 = Segment01 * OutCoordinate.U;
-				Segment02 = Segment02 * OutCoordinate.V;
+				Segment01 = Segment01 * OutCoordinate.X;
+				Segment02 = Segment02 * OutCoordinate.Y;
 				PointType ProjectedPoint = Segment01 + Segment02;
 				ProjectedPoint = ProjectedPoint + Point0;
 				return ProjectedPoint;
@@ -247,18 +247,18 @@ struct CADKERNEL_API TTriangle
 		{
 		case Side01:
 			ProjectedPoint = ProjectPointOnSegment(InPoint, Point0, Point1, SegmentCoordinate);
-			OutCoordinate.U = SegmentCoordinate;
-			OutCoordinate.V = 0.0;
+			OutCoordinate.X = SegmentCoordinate;
+			OutCoordinate.Y = 0.0;
 			break;
 		case Side20:
 			ProjectedPoint = ProjectPointOnSegment(InPoint, Point0, Point2, SegmentCoordinate);
-			OutCoordinate.U = 0.0;
-			OutCoordinate.V = SegmentCoordinate;
+			OutCoordinate.X = 0.0;
+			OutCoordinate.Y = SegmentCoordinate;
 			break;
 		case Side12:
 			ProjectedPoint = ProjectPointOnSegment(InPoint, Point1, Point2, SegmentCoordinate);
-			OutCoordinate.U = 1.0 - SegmentCoordinate;
-			OutCoordinate.V = SegmentCoordinate;
+			OutCoordinate.X = 1.0 - SegmentCoordinate;
+			OutCoordinate.Y = SegmentCoordinate;
 			break;
 		}
 		return ProjectedPoint;
diff --git a/Engine/Source/Runtime/Datasmith/CADKernel/Base/Public/Math/Point.h b/Engine/Source/Runtime/Datasmith/CADKernel/Base/Public/Math/Point.h
--- a/Engine/Source/Runtime/Datasmith/CADKernel/Base/Public/Math/Point.h
+++ b/Engine/Source/Runtime/Datasmith/CADKernel/Base/Public/Math/Point.h
@@ -109,6 +109,11 @@ PRAGMA_ENABLE_DEPRECATION_WARNINGS
 		return X * Point.X + Y * Point.Y + Z * Point.Z;
 	}
 
+	double operator|(const FPoint& Point) const
+	{
+		return X * Point.X + Y * Point.Y + Z * Point.Z;
+	}
+
 	FPoint operator*(double Scale) const
 	{
 		return FPoint(X * Scale, Y * Scale, Z * Scale);
@@ -380,7 +385,7 @@ PRAGMA_ENABLE_DEPRECATION_WARNINGS
 	}
 };
 
-class CADKERNEL_API FPointH
+class CADKERNEL_API FPoint2D
 {
 public:
 	union
@@ -389,98 +394,10 @@ public:
 		{
 			double X;
 			double Y;
-			double Z;
-			double W;
 		};
 
 		UE_DEPRECATED(all, "For internal use only")
-		double XYZW[4];
-	};
-
-public:
-	/** A zero point (0,0,0)*/
-	static const FPointH ZeroPoint;
-	static const FPointH FarawayPoint;
-	static const int32 Dimension;
-
-	FPointH()
-		: X(0.)
-		, Y(0.)
-		, Z(0.)
-		, W(1.)
-	{
-	}
-
-	FPointH(const FPoint& Point, double Weight)
-		: X(Point.X* Weight)
-		, Y(Point.Y* Weight)
-		, Z(Point.Z* Weight)
-		, W(Weight)
-	{
-	}
-
-	FPointH(const double InX, const double InY, const double InZ, const double InW)
-		: X(InX)
-		, Y(InY)
-		, Z(InZ)
-		, W(InW)
-	{
-	}
-
-	constexpr double& operator[](int32 Index)
-	{
-PRAGMA_DISABLE_DEPRECATION_WARNINGS
-		return XYZW[Index];
-PRAGMA_ENABLE_DEPRECATION_WARNINGS
-	}
-
-	constexpr double operator[](int32 Index) const
-	{
-PRAGMA_DISABLE_DEPRECATION_WARNINGS
-		return XYZW[Index];
-PRAGMA_ENABLE_DEPRECATION_WARNINGS
-	}
-
-	FPointH operator+(const FPointH& Point) const
-	{
-		return FPointH(X + Point.X, Y + Point.Y, Z + Point.Z, W + Point.W);
-	}
-
-	FPointH operator-(const FPointH& Point) const
-	{
-		return FPointH(X - Point.X, Y - Point.Y, Z - Point.Z, W - Point.W);
-	}
-
-	FPointH operator*(double Factor) const
-	{
-		return FPointH(X * Factor, Y * Factor, Z * Factor, W * Factor);
-	}
-
-	operator FPoint() const
-	{
-		return FPoint(X / W, Y / W, Z / W);
-	}
-
-	friend FArchive& operator<<(FArchive& Ar, FPointH& Point)
-	{
-		Ar.Serialize(&Point, sizeof(FPointH));
-		return Ar;
-	}
-};
-
-class CADKERNEL_API FPoint2D
-{
-public:
-	union
-	{
-		struct
-		{
-			double U;
-			double V;
-		};
-
-		UE_DEPRECATED(all, "For internal use only")
-		double UV[2];
+		double XY[2];
 	};
 
 
@@ -493,42 +410,42 @@ public:
 
 	FPoint2D()
 	{
-		U = V = 0.0;
+		X = Y = 0.0;
 	}
 
 	FPoint2D(double NewU, double NewV)
-		: U(NewU)
-		, V(NewV)
+		: X(NewU)
+		, Y(NewV)
 	{
 	}
 
 	FPoint2D(const FPoint& Point)
-		: U(Point.X)
-		, V(Point.Y)
+		: X(Point.X)
+		, Y(Point.Y)
 	{
 	}
 
 	FPoint2D(const FPoint2D& Point)
-		: U(Point.U)
-		, V(Point.V)
+		: X(Point.X)
+		, Y(Point.Y)
 	{
 	}
 
 	FPoint2D& operator=(const FPoint2D& Point)
 	{
-		U = Point.U;
-		V = Point.V;
+		X = Point.X;
+		Y = Point.Y;
 		return *this;
 	}
 
 	FPoint2D Rotate(double Theta)
 	{
-		return FPoint2D(U * cos(Theta) - V * sin(Theta), U * sin(Theta) + V * cos(Theta));
+		return FPoint2D(X * cos(Theta) - Y * sin(Theta), X * sin(Theta) + Y * cos(Theta));
 	}
 
 	FPoint2D Middle(const FPoint2D& Point) const
 	{
-		return FPoint(0.5 * (U + Point.U), 0.5 * (V + Point.V));
+		return FPoint(0.5 * (X + Point.X), 0.5 * (Y + Point.Y));
 	}
 
 	double Distance(const FPoint2D& Point) const
@@ -538,7 +455,7 @@ public:
 
 	double SquareDistance(const FPoint2D& Point) const
 	{
-		return FMath::Square(Point.U - U) + FMath::Square(Point.V - V);
+		return FMath::Square(Point.X - X) + FMath::Square(Point.Y - Y);
 	}
 
 	double Length() const
@@ -548,7 +465,7 @@ public:
 
 	double SquareLength() const
 	{
-		return FMath::Square(U) + FMath::Square(V);
+		return FMath::Square(X) + FMath::Square(Y);
 	}
 
 	FPoint2D& Normalize(double& Norm)
@@ -556,12 +473,12 @@ public:
 		Norm = Length();
 		if (FMath::IsNearlyZero(Norm))
 		{
-			U = V = 0.0;
+			X = Y = 0.0;
 		}
 		else
 		{
-			U /= Norm;
-			V /= Norm;
+			X /= Norm;
+			Y /= Norm;
 		}
 		return *this;
 	}
@@ -575,25 +492,25 @@ public:
 	constexpr double& operator[](int32 Index)
 	{
 PRAGMA_DISABLE_DEPRECATION_WARNINGS
-		return UV[Index];
+		return XY[Index];
 PRAGMA_ENABLE_DEPRECATION_WARNINGS
 	}
 
 	constexpr double operator[](int32 Index) const
 	{
 PRAGMA_DISABLE_DEPRECATION_WARNINGS
-		return UV[Index];
+		return XY[Index];
 PRAGMA_ENABLE_DEPRECATION_WARNINGS
 	}
 
 	bool operator==(const FPoint2D& Point) const
 	{
-		return FMath::IsNearlyEqual(U, Point.U) && FMath::IsNearlyEqual(V, Point.V);
+		return FMath::IsNearlyEqual(X, Point.X) && FMath::IsNearlyEqual(Y, Point.Y);
 	}
 
 	bool operator<(const FPoint2D& other) const
 	{
-		int32 cmp = RealCompare(U, other.U);
+		int32 cmp = RealCompare(X, other.X);
 		if (cmp < 0)
 		{
 			return true;
@@ -603,45 +520,50 @@ PRAGMA_ENABLE_DEPRECATION_WARNINGS
 			return false;
 		}
 
-		cmp = RealCompare(V, other.V);
+		cmp = RealCompare(Y, other.Y);
 		return (cmp < 0);
 	}
 
 	void Set(double NewU, double NewV = 0.)
 	{
-		U = NewU;
-		V = NewV;
+		X = NewU;
+		Y = NewV;
 	}
 
 	void Set(double* NewCoordinate)
 	{
-		U = NewCoordinate[0];
-		V = NewCoordinate[1];
+		X = NewCoordinate[0];
+		Y = NewCoordinate[1];
 	}
 
 	FPoint2D operator+(const FPoint2D& Point) const
 	{
-		return FPoint2D(U + Point.U, V + Point.V);
+		return FPoint2D(X + Point.X, Y + Point.Y);
 	}
 
 	FPoint2D operator-(const FPoint2D& Point) const
 	{
-		return FPoint2D(U - Point.U, V - Point.V);
+		return FPoint2D(X - Point.X, Y - Point.Y);
 	}
 
 	double operator*(const FPoint2D& Point) const
 	{
-		return U * Point.U + V * Point.V;
+		return X * Point.X + Y * Point.Y;
+	}
+
+	double operator|(const FPoint2D& Point) const
+	{
+		return X * Point.X + Y * Point.Y;
 	}
 
 	double operator^(const FPoint2D& Point) const
 	{
-		return (U * Point.V - V * Point.U);
+		return (X * Point.Y - Y * Point.X);
 	}
 
 	FPoint2D operator*(double Scale) const
 	{
-		return FPoint2D(U * Scale, V * Scale);
+		return FPoint2D(X * Scale, Y * Scale);
 	}
 
 	FPoint2D operator/(double InvScale) const
@@ -651,32 +573,32 @@ PRAGMA_ENABLE_DEPRECATION_WARNINGS
 			return *this;
 		}
 
-		return FPoint2D(U / InvScale, V / InvScale);
+		return FPoint2D(X / InvScale, Y / InvScale);
 	}
 
 	FPoint2D operator-() const
 	{
-		return FPoint2D(-U, -V);
+		return FPoint2D(-X, -Y);
 	}
 
 	FPoint2D& operator+=(const FPoint2D& Point)
 	{
-		U += Point.U;
-		V += Point.V;
+		X += Point.X;
+		Y += Point.Y;
 		return *this;
 	}
 
 	FPoint2D& operator-=(const FPoint2D& Point)
 	{
-		U -= Point.U;
-		V -= Point.V;
+		X -= Point.X;
+		Y -= Point.Y;
 		return *this;
 	}
 
 	FPoint2D& operator*=(double Scale)
 	{
-		U *= Scale;
-		V *= Scale;
+		X *= Scale;
+		Y *= Scale;
 		return *this;
 	}
 
@@ -687,8 +609,8 @@ PRAGMA_ENABLE_DEPRECATION_WARNINGS
 			return *this;
 		}
 
-		U /= InvScale;
-		V /= InvScale;
+		X /= InvScale;
+		Y /= InvScale;
 		return *this;
 	}
 
@@ -707,16 +629,16 @@ PRAGMA_ENABLE_DEPRECATION_WARNINGS
 
 	FPoint2D GetPerpendicularVector()
 	{
-		return FPoint2D(-V, U);
+		return FPoint2D(-Y, X);
 	}
 
 	/**
 	 * Return the projection of the point on the diagonal axis (of vector (1,1))
-	 * i.e. return U + V
+	 * i.e. return X + Y
 	 */
 	double DiagonalAxisCoordinate() const
 	{
-		return U + V;
+		return X + Y;
 	}
 
 	friend FArchive& operator<<(FArchive& Ar, FPoint2D& Point)
@@ -727,111 +649,140 @@ PRAGMA_ENABLE_DEPRECATION_WARNINGS
 
 	friend FPoint2D Abs(const FPoint2D& Point)
 	{
-		double PointU = FMath::Abs(Point.U);
-		double PointV = FMath::Abs(Point.V);
+		double PointU = FMath::Abs(Point.X);
+		double PointV = FMath::Abs(Point.Y);
 		return { PointU, PointV };
 	}
 
 	friend FPoint2D Max(const FPoint2D& PointA, const FPoint2D& PointB)
 	{
-		return { FMath::Max(PointA.U, PointB.U), FMath::Max(PointA.V, PointB.V) };
+		return { FMath::Max(PointA.X, PointB.X), FMath::Max(PointA.Y, PointB.Y) };
 	}
 
 	bool Equals(const FPoint2D& Other, double Tolerance = UE_DOUBLE_SMALL_NUMBER) const
 	{
-		return FMath::IsNearlyEqual(U, Other.U, Tolerance) && FMath::IsNearlyEqual(V, Other.V, Tolerance);
+		return FMath::IsNearlyEqual(X, Other.X, Tolerance) && FMath::IsNearlyEqual(Y, Other.Y, Tolerance);
 	}
 };
 
 inline FPoint::FPoint(const FPoint2D& Point)
 {
-	X = Point.U; Y = Point.V; Z = 0.0;
+	X = Point.X; Y = Point.Y; Z = 0.0;
 }
 
 typedef FPoint2D FSurfacicTolerance;
 
-class CADKERNEL_API FFPoint
+inline FPoint operator*(double Scale, const FPoint& Point)
+{
+	return FPoint(Point.X * Scale, Point.Y * Scale, Point.Z * Scale);
+}
+
+inline FPoint operator*(float Scale, const FPoint& Point)
+{
+	double ScaleD = static_cast<double>(Scale);
+	return FPoint(Point.X * ScaleD, Point.Y * ScaleD, Point.Z * ScaleD);
+}
+
+inline FPoint2D operator*(double Scale, const FPoint2D& Point)
+{
+	return FPoint2D(Point.X * Scale, Point.Y * Scale);
+}
+
+inline FPoint2D operator*(float Scale, const FPoint2D& Point)
+{
+	double ScaleD = static_cast<double>(Scale);
+	return FPoint2D(Point.X * ScaleD, Point.Y * ScaleD);
+}
+
+class FVectorUtil
 {
 public:
-	union
-	{
-		struct
-		{
-			float X;
-			float Y;
-			float Z;
-		};
+	static CADKERNEL_API const FVector4d FarawayPoint4D;
+	static CADKERNEL_API const FVector FarawayPoint3D;
+	static CADKERNEL_API const FVector2d FarawayPoint2D;
 
-		UE_DEPRECATED(all, "For internal use only")
-		float XYZ[3];
-	};
+	static FPoint2D FromVector(const FPoint& Src)
+	{
+		return FPoint2D(Src.X, Src.Y);
+	}
 
-	/** A zero point (0,0,0)*/
-	static const FFPoint ZeroPoint;
-	static const FFPoint FarawayPoint;
-	static const int32 Dimension;
+	static FPoint FromBuffer(const double* const Src)
+	{
+		return FPoint(Src[0], Src[1], Src[2]);
+	}
 
-	FFPoint(float InX = 0., float InY = 0., float InZ = 0.)
+	static void SetPoint(FPoint& Point, const double* const Src)
 	{
-		X = InX;
-		Y = InY;
-		Z = InZ;
+		Point.X = Src[0];
+		Point.Y = Src[1];
+		Point.Z = Src[2];
 	}
 
-	FFPoint(const FPoint& Point)
+	static void SetPoint(FPoint& Point, double X, double Y, double Z)
 	{
-		X = (float)Point.X;
-		Y = (float)Point.Y;
-		Z = (float)Point.Z;
+		Point.X = X;
+		Point.Y = Y;
+		Point.Z = Z;
 	}
 
-	constexpr float& operator[](int32 Index)
+	static void SetPoint(FPoint& Point, const FPoint2D& Point2D)
 	{
-PRAGMA_DISABLE_DEPRECATION_WARNINGS
-		return XYZ[Index];
-PRAGMA_ENABLE_DEPRECATION_WARNINGS
+		Point.X = Point2D.X;
+		Point.Y = Point2D.Y;
+		Point.Z = 0.;
 	}
 
-	constexpr float operator[](int32 Index) const
+	static void SetPoint2D(FPoint2D& Point, const double* const Src)
 	{
-PRAGMA_DISABLE_DEPRECATION_WARNINGS
-		return XYZ[Index];
-PRAGMA_ENABLE_DEPRECATION_WARNINGS
+		Point.X = Src[0];
+		Point.Y = Src[1];
 	}
 
-	FFPoint& operator=(const FPoint& Point)
+	static void SetPoint2D(FPoint2D& Point, double X, double Y)
 	{
-		X = (float)Point.X; Y = (float)Point.Y; Z = (float)Point.Z; return *this;
+		Point.X = X;
+		Point.Y = Y;
 	}
 
-	friend FArchive& operator<<(FArchive& Ar, FFPoint& Point)
+	template<typename VectorType>
+	static double ComputeCosinus(const VectorType& Vec1, const VectorType& Vec2)
 	{
-		Ar.Serialize(&Point, sizeof(FFPoint));
-		return Ar;
+		VectorType ThisNormalized(Vec1);
+		VectorType OtherNormalized(Vec2);
+
+		ThisNormalized.Normalize();
+		OtherNormalized.Normalize();
+
+		const double Cosinus = ThisNormalized | OtherNormalized;
+
+		return FMath::Max(-1.0, FMath::Min(Cosinus, 1.0));
 	}
-};
 
-inline FPoint operator*(double Scale, const FPoint& Point)
-{
-	return FPoint(Point.X * Scale, Point.Y * Scale, Point.Z * Scale);
-}
+	static double ComputeAngle(const FPoint& Vec1, const FPoint& Vec2)
+	{
+		FPoint ThisNormalized(Vec1);
+		FPoint OtherNormalized(Vec2);
 
-inline FPoint operator*(float Scale, const FPoint& Point)
-{
-	double ScaleD = static_cast<double>(Scale);
-	return FPoint(Point.X * ScaleD, Point.Y * ScaleD, Point.Z * ScaleD);
-}
+		ThisNormalized.Normalize();
+		OtherNormalized.Normalize();
 
-inline FPoint2D operator*(double Scale, const FPoint2D& Point)
-{
-	return FPoint2D(Point.U * Scale, Point.V * Scale);
-}
+		return FMath::Acos(ThisNormalized | OtherNormalized);
+	}
 
-inline FPoint2D operator*(float Scale, const FPoint2D& Point)
-{
-	double ScaleD = static_cast<double>(Scale);
-	return FPoint2D(Point.U * ScaleD, Point.V * ScaleD);
-}
+	/**
+	 * Return the projection of the point on the diagonal axis (of vector (1,1,1))
+	 * i.e. return X + Y + Z
+	 */
+	static double DiagonalAxisCoordinate(const FVector& Vec)
+	{
+		return Vec.X + Vec.Y + Vec.Z;
+	}
+
+	static double DiagonalAxisCoordinate(const FVector2d& Vec)
+	{
+		return Vec.X + Vec.Y;
+	}
+};
 
 }
 
@@ -842,11 +793,6 @@ inline uint32 GetTypeHash(const UE::CADKernel::FPoint& Point)
 
 inline uint32 GetTypeHash(const UE::CADKernel::FPoint2D& Point)
 {
-	return HashCombine(GetTypeHash(Point.U), GetTypeHash(Point.V));
-}
-
-inline uint32 GetTypeHash(const UE::CADKernel::FPointH& Point)
-{
-	return HashCombine(GetTypeHash(Point.X), HashCombine(GetTypeHash(Point.Y), HashCombine(GetTypeHash(Point.Z), GetTypeHash(Point.W))));
+	return HashCombine(GetTypeHash(Point.X), GetTypeHash(Point.Y));
 }
 
diff --git a/Engine/Source/Runtime/Datasmith/CADKernel/Base/Public/Math/SlopeUtils.h b/Engine/Source/Runtime/Datasmith/CADKernel/Base/Public/Math/SlopeUtils.h
--- a/Engine/Source/Runtime/Datasmith/CADKernel/Base/Public/Math/SlopeUtils.h
+++ b/Engine/Source/Runtime/Datasmith/CADKernel/Base/Public/Math/SlopeUtils.h
@@ -179,8 +179,8 @@ inline double SwapSlopeOrientation(double Slope)
 
 inline double ComputeSlope(const FPoint2D& StartPoint, const FPoint2D& EndPoint)
 {
-	double DeltaU = EndPoint.U - StartPoint.U;
-	double DeltaV = EndPoint.V - StartPoint.V;
+	double DeltaU = EndPoint.X - StartPoint.X;
+	double DeltaV = EndPoint.Y - StartPoint.Y;
 	double Delta;
 
 	if (FMath::Abs(DeltaU) < DOUBLE_SMALL_NUMBER && FMath::Abs(DeltaV) < DOUBLE_SMALL_NUMBER)
@@ -430,8 +430,8 @@ inline bool ArePointsInsideSectorABC(const FPoint2D& PointA, const FPoint2D& Poi
 
 	return Algo::AllOf(Points, [&](const FPoint2D* PointP) 
 	{
-		double DeltaU = PointB.U - PointP->U;
-		double DeltaV = PointB.V - PointP->V;
+		double DeltaU = PointB.X - PointP->X;
+		double DeltaV = PointB.Y - PointP->Y;
 
 		if (FMath::Abs(DeltaU) < DOUBLE_SMALL_NUMBER && FMath::Abs(DeltaV) < DOUBLE_SMALL_NUMBER)
 		{
diff --git a/Engine/Source/Runtime/Datasmith/CADKernel/Base/Public/Mesh/Structure/FaceMesh.h b/Engine/Source/Runtime/Datasmith/CADKernel/Base/Public/Mesh/Structure/FaceMesh.h
--- a/Engine/Source/Runtime/Datasmith/CADKernel/Base/Public/Mesh/Structure/FaceMesh.h
+++ b/Engine/Source/Runtime/Datasmith/CADKernel/Base/Public/Mesh/Structure/FaceMesh.h
@@ -6,7 +6,6 @@
 namespace UE::CADKernel
 {
 class FModelMesh;
-class FPoint;
 class FTopologicalEntity;
 class FTopologicalFace;
 
diff --git a/Engine/Source/Runtime/Datasmith/CADKernel/Base/Public/Mesh/Structure/ModelMesh.h b/Engine/Source/Runtime/Datasmith/CADKernel/Base/Public/Mesh/Structure/ModelMesh.h
--- a/Engine/Source/Runtime/Datasmith/CADKernel/Base/Public/Mesh/Structure/ModelMesh.h
+++ b/Engine/Source/Runtime/Datasmith/CADKernel/Base/Public/Mesh/Structure/ModelMesh.h
@@ -14,7 +14,6 @@ class FDatabase;
 class FEdgeMesh;
 class FFaceMesh;
 class FMesh;
-class FPoint;
 class FSession;
 class FVertexMesh;
 
diff --git a/Engine/Source/Runtime/Datasmith/CADKernel/Base/Public/UI/Visu.h b/Engine/Source/Runtime/Datasmith/CADKernel/Base/Public/UI/Visu.h
--- a/Engine/Source/Runtime/Datasmith/CADKernel/Base/Public/UI/Visu.h
+++ b/Engine/Source/Runtime/Datasmith/CADKernel/Base/Public/UI/Visu.h
@@ -157,9 +157,6 @@ public:
 	virtual void DrawPoint(const FPoint2D& Point, EVisuProperty InProperty = EVisuProperty::BluePoint)
 	{}
 
-	virtual void DrawPoint(const FPointH& Point, EVisuProperty InProperty = EVisuProperty::BluePoint)
-	{}
-
 	virtual void DrawElement(int32 Dimension, TArray<FPoint>& Points, EVisuProperty InProperty = EVisuProperty::Element)
 	{}
 
