Replaced FPoint and FPoint2D by FVector and  FVector2d and remove unused FFPoint and FPointH classes [WIP] #2

#rnx

diff --git a/Engine/Plugins/Enterprise/DatasmithCADImporter/Source/CADInterfaces/Private/TechSoftBridge.h b/Engine/Plugins/Enterprise/DatasmithCADImporter/Source/CADInterfaces/Private/TechSoftBridge.h
--- a/Engine/Plugins/Enterprise/DatasmithCADImporter/Source/CADInterfaces/Private/TechSoftBridge.h
+++ b/Engine/Plugins/Enterprise/DatasmithCADImporter/Source/CADInterfaces/Private/TechSoftBridge.h
@@ -19,7 +19,6 @@ class FCurve;
 class FEntity;
 class FFaceMesh;
 class FModel;
-class FPoint;
 class FRestrictionCurve;
 class FSession;
 class FShell;
diff --git a/Engine/Source/Runtime/Datasmith/CADKernel/Base/Private/Geo/Curves/PolylineCurve.h b/Engine/Source/Runtime/Datasmith/CADKernel/Base/Private/Geo/Curves/PolylineCurve.h
--- a/Engine/Source/Runtime/Datasmith/CADKernel/Base/Private/Geo/Curves/PolylineCurve.h
+++ b/Engine/Source/Runtime/Datasmith/CADKernel/Base/Private/Geo/Curves/PolylineCurve.h
@@ -131,7 +131,7 @@ public:
 	{
 		for (PointType& Pole : this->Points)
 		{
-			Pole += OffsetDirection;
+			Pole += PointType(OffsetDirection);
 		}
 	}
 
diff --git a/Engine/Source/Runtime/Datasmith/CADKernel/Base/Private/Geo/Curves/SegmentCurve.h b/Engine/Source/Runtime/Datasmith/CADKernel/Base/Private/Geo/Curves/SegmentCurve.h
--- a/Engine/Source/Runtime/Datasmith/CADKernel/Base/Private/Geo/Curves/SegmentCurve.h
+++ b/Engine/Source/Runtime/Datasmith/CADKernel/Base/Private/Geo/Curves/SegmentCurve.h
@@ -82,17 +82,29 @@ public:
 	}
 
 private:
+	void GetTangent(FPoint& Tangent) const
+	{
+		Tangent = EndPoint - StartPoint;
+	}
+
+	void GetTangent(FPoint2D& Tangent) const
+	{
+		Tangent = FPoint2D(EndPoint.X - StartPoint.X, EndPoint.Y - StartPoint.Y);
+	}
+
 	template <typename CurvePointType, typename PointType>
 	void Evaluate(double Coordinate, CurvePointType& OutPoint, int32 DerivativeOrder) const
 	{
 		OutPoint.DerivativeOrder = DerivativeOrder;
 
-		PointType Tangent = (EndPoint - StartPoint);
-		OutPoint.Point = Tangent * Coordinate + StartPoint;
+		PointType Tangent;
+		GetTangent(Tangent);
+
+		OutPoint.Point = Tangent * Coordinate + PointType(StartPoint);
 
 		if (DerivativeOrder > 0)
 		{
-			OutPoint.Gradient = Tangent;
+			OutPoint.Gradient = MoveTemp(Tangent);
 		}
 	}
 
diff --git a/Engine/Source/Runtime/Datasmith/CADKernel/Base/Private/Geo/Curves/SplineCurve.cpp b/Engine/Source/Runtime/Datasmith/CADKernel/Base/Private/Geo/Curves/SplineCurve.cpp
--- a/Engine/Source/Runtime/Datasmith/CADKernel/Base/Private/Geo/Curves/SplineCurve.cpp
+++ b/Engine/Source/Runtime/Datasmith/CADKernel/Base/Private/Geo/Curves/SplineCurve.cpp
@@ -56,9 +56,7 @@ void FSplineCurve::SetSplinePoints(const TArray<FPoint>& Points)
 	double InputKey = 0.0;
 	for (const FPoint& Point : Points)
 	{
-		Position.Points.Emplace(InputKey, Point, FPoint::ZeroPoint, FPoint::ZeroPoint, CIM_CurveAuto);
-		//Rotation.Points.Emplace(InputKey, FQuat::Identity, FQuat::Identity, FQuat::Identity, CIM_CurveAuto);
-		//Scale.Points.Emplace(InputKey, FPoint::UnitPoint, FPoint::ZeroPoint, FPoint::ZeroPoint, CIM_CurveAuto);
+		Position.Points.Emplace(InputKey, Point, FPoint::ZeroVector, FPoint::ZeroVector, CIM_CurveAuto);
 		InputKey += 1.0;
 	}
 
@@ -162,17 +160,17 @@ void FSplineCurve::UpdateSpline(bool bInClosedLoop, bool bStationaryEndpoints, i
 
 FPoint FSplineCurve::GetLocationAtSplineInputKey(double Coordinate) const
 {
-	return Position.Eval((float)Coordinate, FPoint::ZeroPoint);
+	return Position.Eval((float)Coordinate, FPoint::ZeroVector);
 }
 
 FPoint FSplineCurve::GetTangentAtSplineInputKey(double Coordinate) const
 {
-	return Position.EvalDerivative((float)Coordinate, FPoint::ZeroPoint);
+	return Position.EvalDerivative((float)Coordinate, FPoint::ZeroVector);
 }
 
 FPoint FSplineCurve::GetDirectionAtSplineInputKey(double Coordinate) const
 {
-	return Position.EvalDerivative(Coordinate, FPoint::ZeroPoint).Normalize();
+	return Position.EvalDerivative(Coordinate, FPoint::ZeroVector).Normalize();
 }
 
 void FSplineCurve::EvaluatePoint(double Coordinate, FCurvePoint& OutPoint, int32 DerivativeOrder) const
diff --git a/Engine/Source/Runtime/Datasmith/CADKernel/Base/Private/Geo/Surfaces/BezierSurface.cpp b/Engine/Source/Runtime/Datasmith/CADKernel/Base/Private/Geo/Surfaces/BezierSurface.cpp
--- a/Engine/Source/Runtime/Datasmith/CADKernel/Base/Private/Geo/Surfaces/BezierSurface.cpp
+++ b/Engine/Source/Runtime/Datasmith/CADKernel/Base/Private/Geo/Surfaces/BezierSurface.cpp
@@ -21,7 +21,7 @@ void FBezierSurface::EvaluatePoint(const FPoint2D& InPoint2D, FSurfacicPoint& Ou
 	
 	if (InDerivativeOrder > 0)
 	{
-		VCurveUGradient.Init(FPoint::ZeroPoint, VPoleNum);
+		VCurveUGradient.Init(FPoint::ZeroVector, VPoleNum);
 	}
 
 	if (InDerivativeOrder > 1)
@@ -46,11 +46,11 @@ void FBezierSurface::EvaluatePoint(const FPoint2D& InPoint2D, FSurfacicPoint& Ou
 
 			if (InDerivativeOrder > 0)
 			{
-				UAuxiliaryGradient.Init(FPoint::ZeroPoint, UPoleNum);
+				UAuxiliaryGradient.Init(FPoint::ZeroVector, UPoleNum);
 			}
 			if (InDerivativeOrder > 1)
 			{
-				UAuxiliaryLaplacian.Init(FPoint::ZeroPoint, UPoleNum);
+				UAuxiliaryLaplacian.Init(FPoint::ZeroVector, UPoleNum);
 			}
 
 			// Compute Point, Gradient and Laplacian at U coordinate with De Casteljau's algorithm, 
@@ -98,11 +98,11 @@ void FBezierSurface::EvaluatePoint(const FPoint2D& InPoint2D, FSurfacicPoint& Ou
 	TArray<FPoint> VVAuxiliaryLaplacian;
 	if (InDerivativeOrder > 0)
 	{
-		VAuxiliaryGradient.Init(FPoint::ZeroPoint, VPoleNum);
+		VAuxiliaryGradient.Init(FPoint::ZeroVector, VPoleNum);
 		if (InDerivativeOrder > 1)
 		{
-			VVAuxiliaryLaplacian.Init(FPoint::ZeroPoint, VPoleNum);
-			UVAuxiliaryLaplacian.Init(FPoint::ZeroPoint, VPoleNum);
+			VVAuxiliaryLaplacian.Init(FPoint::ZeroVector, VPoleNum);
+			UVAuxiliaryLaplacian.Init(FPoint::ZeroVector, VPoleNum);
 		}
 	}
 
diff --git a/Engine/Source/Runtime/Datasmith/CADKernel/Base/Private/Geo/Surfaces/ConeSurface.h b/Engine/Source/Runtime/Datasmith/CADKernel/Base/Private/Geo/Surfaces/ConeSurface.h
--- a/Engine/Source/Runtime/Datasmith/CADKernel/Base/Private/Geo/Surfaces/ConeSurface.h
+++ b/Engine/Source/Runtime/Datasmith/CADKernel/Base/Private/Geo/Surfaces/ConeSurface.h
@@ -58,7 +58,7 @@ protected:
 
 	void ComputeMinToleranceIso()
 	{
-		FPoint Origin = Matrix.Multiply(FPoint::ZeroPoint);
+		FPoint Origin = Matrix.Multiply(FPoint::ZeroVector);
 
 		double DeltaVR = tan(ConeAngle);
 		double Radius1 = FMath::Abs(StartRadius + Boundary[EIso::IsoV].Max * DeltaVR);
diff --git a/Engine/Source/Runtime/Datasmith/CADKernel/Base/Private/Geo/Surfaces/CylinderSurface.h b/Engine/Source/Runtime/Datasmith/CADKernel/Base/Private/Geo/Surfaces/CylinderSurface.h
--- a/Engine/Source/Runtime/Datasmith/CADKernel/Base/Private/Geo/Surfaces/CylinderSurface.h
+++ b/Engine/Source/Runtime/Datasmith/CADKernel/Base/Private/Geo/Surfaces/CylinderSurface.h
@@ -34,7 +34,7 @@ protected:
 
 	void ComputeMinToleranceIso()
 	{
-		FPoint Origin = Matrix.Multiply(FPoint::ZeroPoint);
+		FPoint Origin = Matrix.Multiply(FPoint::ZeroVector);
 
 		FPoint Point2DU{ 1 , 0, 0 };
 		double ToleranceU = Tolerance3D / Radius;
diff --git a/Engine/Source/Runtime/Datasmith/CADKernel/Base/Private/Geo/Surfaces/PlaneSurface.cpp b/Engine/Source/Runtime/Datasmith/CADKernel/Base/Private/Geo/Surfaces/PlaneSurface.cpp
--- a/Engine/Source/Runtime/Datasmith/CADKernel/Base/Private/Geo/Surfaces/PlaneSurface.cpp
+++ b/Engine/Source/Runtime/Datasmith/CADKernel/Base/Private/Geo/Surfaces/PlaneSurface.cpp
@@ -35,7 +35,7 @@ FPlaneSurface::FPlaneSurface(const double InToleranceGeometric, const FPoint& In
 FPlane FPlaneSurface::GetPlane() const
 {
 	FSurfacicPoint Point3D;
-	EvaluatePoint(FPoint2D::ZeroPoint, Point3D, 0);
+	EvaluatePoint(FVector2d::ZeroVector, Point3D, 0);
 	FPoint Normal = Matrix.Column(2);
 	return FPlane(Point3D.Point, Normal);
 }
@@ -52,9 +52,9 @@ void FPlaneSurface::EvaluatePoint(const FPoint2D& InSurfacicCoordinate, FSurfaci
 
 	if(InDerivativeOrder>1) 
 	{
-		OutPoint3D.LaplacianU = FPoint::ZeroPoint;
-		OutPoint3D.LaplacianV = FPoint::ZeroPoint;
-		OutPoint3D.LaplacianUV = FPoint::ZeroPoint;
+		OutPoint3D.LaplacianU = FPoint::ZeroVector;
+		OutPoint3D.LaplacianV = FPoint::ZeroVector;
+		OutPoint3D.LaplacianUV = FPoint::ZeroVector;
 	}
 }
 
diff --git a/Engine/Source/Runtime/Datasmith/CADKernel/Base/Private/Geo/Surfaces/PlaneSurface.h b/Engine/Source/Runtime/Datasmith/CADKernel/Base/Private/Geo/Surfaces/PlaneSurface.h
--- a/Engine/Source/Runtime/Datasmith/CADKernel/Base/Private/Geo/Surfaces/PlaneSurface.h
+++ b/Engine/Source/Runtime/Datasmith/CADKernel/Base/Private/Geo/Surfaces/PlaneSurface.h
@@ -41,7 +41,7 @@ protected:
 
 	void ComputeMinToleranceIso()
 	{
-		FPoint Origin = Matrix.Multiply(FPoint::ZeroPoint);
+		FPoint Origin = Matrix.Multiply(FPoint::ZeroVector);
 
 		FPoint Point2DU{ 1 , 0, 0 };
 		FPoint Point2DV{ 0, 1, 0 };
diff --git a/Engine/Source/Runtime/Datasmith/CADKernel/Base/Private/Geo/Surfaces/SphericalSurface.h b/Engine/Source/Runtime/Datasmith/CADKernel/Base/Private/Geo/Surfaces/SphericalSurface.h
--- a/Engine/Source/Runtime/Datasmith/CADKernel/Base/Private/Geo/Surfaces/SphericalSurface.h
+++ b/Engine/Source/Runtime/Datasmith/CADKernel/Base/Private/Geo/Surfaces/SphericalSurface.h
@@ -52,7 +52,7 @@ protected:
 	{
 		double Tolerance2D = Tolerance3D / Radius;
 
-		FPoint Origin = Matrix.Multiply(FPoint::ZeroPoint);
+		FPoint Origin = Matrix.Multiply(FPoint::ZeroVector);
 
 		FPoint Point2DU{ 1 , 0, 0 };
 		FPoint Point2DV{ 0, 1, 0 };
diff --git a/Engine/Source/Runtime/Datasmith/CADKernel/Base/Private/Geo/Surfaces/TabulatedCylinderSurface.cpp b/Engine/Source/Runtime/Datasmith/CADKernel/Base/Private/Geo/Surfaces/TabulatedCylinderSurface.cpp
--- a/Engine/Source/Runtime/Datasmith/CADKernel/Base/Private/Geo/Surfaces/TabulatedCylinderSurface.cpp
+++ b/Engine/Source/Runtime/Datasmith/CADKernel/Base/Private/Geo/Surfaces/TabulatedCylinderSurface.cpp
@@ -44,8 +44,8 @@ void FTabulatedCylinderSurface::EvaluatePoint(const FPoint2D& InSurfacicCoordina
 		if (InDerivativeOrder > 1)
 		{
 			OutPoint3D.LaplacianU = PointU.Laplacian * FMath::Square(CurveBounds.Length());
-			OutPoint3D.LaplacianV = FPoint::ZeroPoint;
-			OutPoint3D.LaplacianUV = FPoint::ZeroPoint;
+			OutPoint3D.LaplacianV = FPoint::ZeroVector;
+			OutPoint3D.LaplacianUV = FPoint::ZeroVector;
 		}
 	}
 }
diff --git a/Engine/Source/Runtime/Datasmith/CADKernel/Base/Private/Geo/Surfaces/TorusSurface.h b/Engine/Source/Runtime/Datasmith/CADKernel/Base/Private/Geo/Surfaces/TorusSurface.h
--- a/Engine/Source/Runtime/Datasmith/CADKernel/Base/Private/Geo/Surfaces/TorusSurface.h
+++ b/Engine/Source/Runtime/Datasmith/CADKernel/Base/Private/Geo/Surfaces/TorusSurface.h
@@ -57,7 +57,7 @@ protected:
 		double Tolerance2DU = Tolerance3D / MajorRadius;
 		double Tolerance2DV = Tolerance3D / MinorRadius;
 
-		FPoint Origin = Matrix.Multiply(FPoint::ZeroPoint);
+		FPoint Origin = Matrix.Multiply(FPoint::ZeroVector);
 
 		FPoint Point2DU{ 1 , 0, 0 };
 		FPoint Point2DV{ 0, 1, 0 };
diff --git a/Engine/Source/Runtime/Datasmith/CADKernel/Base/Private/Math/Point.cpp b/Engine/Source/Runtime/Datasmith/CADKernel/Base/Private/Math/Point.cpp
--- a/Engine/Source/Runtime/Datasmith/CADKernel/Base/Private/Math/Point.cpp
+++ b/Engine/Source/Runtime/Datasmith/CADKernel/Base/Private/Math/Point.cpp
@@ -8,12 +8,9 @@
 namespace UE::CADKernel
 {
 
-const FPoint FPoint::ZeroPoint(0.,0.,0.);
-const FPoint FPoint::UnitPoint(1., 1., 1.);
 const FPoint FPoint::FarawayPoint(HUGE_VALUE, HUGE_VALUE, HUGE_VALUE);
 const int32 FPoint::Dimension = 3;
 
-const FPoint2D FPoint2D::ZeroPoint(0., 0.);
 const FPoint2D FPoint2D::FarawayPoint(HUGE_VALUE, HUGE_VALUE);
 const int32 FPoint2D::Dimension = 2;
 
diff --git a/Engine/Source/Runtime/Datasmith/CADKernel/Base/Private/Mesh/Meshers/IsoTriangulator.h b/Engine/Source/Runtime/Datasmith/CADKernel/Base/Private/Mesh/Meshers/IsoTriangulator.h
--- a/Engine/Source/Runtime/Datasmith/CADKernel/Base/Private/Mesh/Meshers/IsoTriangulator.h
+++ b/Engine/Source/Runtime/Datasmith/CADKernel/Base/Private/Mesh/Meshers/IsoTriangulator.h
@@ -388,7 +388,7 @@ inline double CotangentCriteria(const FPoint2D& APoint, const FPoint2D& BPoint,
 	// the ratio between the scalar product PA.PB (=|PA| |PB| cos (A,P,B) )
 	// with the norm of the cross product |PA^PB| (=|PA| |PB| |sin(A,P,B)|)
 	// is compute. 
-	double ScalareProduct = PA * PB;
+	double ScalareProduct = PA | PB;
 	double OutNormal = PA ^ PB;
 	double NormOFPointProduct = FMath::Abs(OutNormal);
 
diff --git a/Engine/Source/Runtime/Datasmith/CADKernel/Base/Private/Mesh/Meshers/IsoTriangulator/IntersectionIsoSegmentTool.h b/Engine/Source/Runtime/Datasmith/CADKernel/Base/Private/Mesh/Meshers/IsoTriangulator/IntersectionIsoSegmentTool.h
--- a/Engine/Source/Runtime/Datasmith/CADKernel/Base/Private/Mesh/Meshers/IsoTriangulator/IntersectionIsoSegmentTool.h
+++ b/Engine/Source/Runtime/Datasmith/CADKernel/Base/Private/Mesh/Meshers/IsoTriangulator/IntersectionIsoSegmentTool.h
@@ -10,7 +10,6 @@ namespace UE::CADKernel
 {
 class FGrid;
 class FIsoNode;
-class FPoint2D;
 
 namespace IntersectionIsoSegmentTool
 {
diff --git a/Engine/Source/Runtime/Datasmith/CADKernel/Base/Private/Mesh/Meshers/IsoTriangulator/IntersectionSegmentTool.h b/Engine/Source/Runtime/Datasmith/CADKernel/Base/Private/Mesh/Meshers/IsoTriangulator/IntersectionSegmentTool.h
--- a/Engine/Source/Runtime/Datasmith/CADKernel/Base/Private/Mesh/Meshers/IsoTriangulator/IntersectionSegmentTool.h
+++ b/Engine/Source/Runtime/Datasmith/CADKernel/Base/Private/Mesh/Meshers/IsoTriangulator/IntersectionSegmentTool.h
@@ -12,7 +12,6 @@ class FGrid;
 class FIsoNode;
 class FIsoInnerNode;
 class FIsoSegment;
-class FPoint2D;
 
 enum class EConnectionType : uint8
 {
@@ -70,7 +69,7 @@ struct FSegment
 		const double ParallelCoef = AB ^ CD;
 		if (FMath::IsNearlyZero(ParallelCoef, DOUBLE_KINDA_SMALL_NUMBER))
 		{
-			const double OrientationCoef = AB * CD;
+			const double OrientationCoef = AB | CD;
 			if ((OrientationCoef >= 0) == bSameOrientation)
 			{
 				return EConnectionType::SuperimposedByOrOn;
diff --git a/Engine/Source/Runtime/Datasmith/CADKernel/Base/Private/Mesh/Meshers/IsoTriangulator/IsoCell.h b/Engine/Source/Runtime/Datasmith/CADKernel/Base/Private/Mesh/Meshers/IsoTriangulator/IsoCell.h
--- a/Engine/Source/Runtime/Datasmith/CADKernel/Base/Private/Mesh/Meshers/IsoTriangulator/IsoCell.h
+++ b/Engine/Source/Runtime/Datasmith/CADKernel/Base/Private/Mesh/Meshers/IsoTriangulator/IsoCell.h
@@ -28,7 +28,7 @@ struct FCellLoop
 
 	FCellLoop(const int32 InIndex, TArray<FLoopNode*>& InNodes, const FGrid& Grid)
 		: Id(InIndex)
-		, Barycenter(FPoint2D::ZeroPoint)
+		, Barycenter(FVector2d::ZeroVector)
 		, Nodes(InNodes)
 		, bIsOuterLoop(Nodes[0]->GetLoopIndex() == 0)
 	{
@@ -41,7 +41,7 @@ struct FCellLoop
 
 	FCellLoop(const int32 InIndex)
 		: Id(InIndex)
-		, Barycenter(FPoint::ZeroPoint)
+		, Barycenter(FPoint::ZeroVector)
 	{
 	}
 
diff --git a/Engine/Source/Runtime/Datasmith/CADKernel/Base/Private/Mesh/Meshers/IsoTriangulator/IsoNode.h b/Engine/Source/Runtime/Datasmith/CADKernel/Base/Private/Mesh/Meshers/IsoTriangulator/IsoNode.h
--- a/Engine/Source/Runtime/Datasmith/CADKernel/Base/Private/Mesh/Meshers/IsoTriangulator/IsoNode.h
+++ b/Engine/Source/Runtime/Datasmith/CADKernel/Base/Private/Mesh/Meshers/IsoTriangulator/IsoNode.h
@@ -398,7 +398,7 @@ public:
 		case EGridSpace::UniformScaled:
 			return Grid.GetLoop2DPoint(Space, LoopIndex, LocalIndex);
 		}
-		return FPoint::ZeroPoint;
+		return FPoint::ZeroVector;
 	}
 
 	/**
@@ -539,7 +539,7 @@ public:
 		case EGridSpace::UniformScaled:
 			return Grid.GetInner2DPoint(Space, LocalIndex);
 		}
-		return FPoint::ZeroPoint;
+		return FPoint::ZeroVector;
 	}
 
 	void OffsetId(int32 StartId)
diff --git a/Engine/Source/Runtime/Datasmith/CADKernel/Base/Private/Mesh/Meshers/ParametricFaceMesher.h b/Engine/Source/Runtime/Datasmith/CADKernel/Base/Private/Mesh/Meshers/ParametricFaceMesher.h
--- a/Engine/Source/Runtime/Datasmith/CADKernel/Base/Private/Mesh/Meshers/ParametricFaceMesher.h
+++ b/Engine/Source/Runtime/Datasmith/CADKernel/Base/Private/Mesh/Meshers/ParametricFaceMesher.h
@@ -40,7 +40,7 @@ struct FCrossZoneElement
 		, Tolerance3D(InTolerance3D)
 		, Segment(InSegment)
 		, OppositeVertexIndices(InOppositeVertexIndices)
-		, OppositePoint2D(FPoint2D::ZeroPoint)
+		, OppositePoint2D(FVector2d::ZeroVector)
 		, OppositeSegment(nullptr)
 		, OppositePointCoordinate(-1.)
 		, SquareThickness(-1.)
@@ -49,11 +49,11 @@ struct FCrossZoneElement
 
 	FCrossZoneElement()
 		: VertexId(-1)
-		, VertexPoint2D(FPoint2D::ZeroPoint)
+		, VertexPoint2D(FVector2d::ZeroVector)
 		, Tolerance3D(-1.)
 		, Segment(nullptr)
 		, OppositeVertexIndices(FPairOfIndex::Undefined)
-		, OppositePoint2D(FPoint2D::ZeroPoint)
+		, OppositePoint2D(FVector2d::ZeroVector)
 		, OppositeSegment(nullptr)
 		, OppositePointCoordinate(-1.)
 		, SquareThickness(-1.)
diff --git a/Engine/Source/Runtime/Datasmith/CADKernel/Base/Private/Mesh/Structure/Grid.cpp b/Engine/Source/Runtime/Datasmith/CADKernel/Base/Private/Mesh/Structure/Grid.cpp
--- a/Engine/Source/Runtime/Datasmith/CADKernel/Base/Private/Mesh/Structure/Grid.cpp
+++ b/Engine/Source/Runtime/Datasmith/CADKernel/Base/Private/Mesh/Structure/Grid.cpp
@@ -643,7 +643,7 @@ void FGrid::RemovePointsCloseToLoop()
 			FPoint2D Projection = ProjectPointOnSegment(Point2D, *Segment.StartPoint, *Segment.EndPoint, Coordinate, /*bRestrictCoodinateToInside*/ true);
 
 			// If Projected point is in the oval center on Point2D => the node is too close
-			FPoint2D ProjectionToPoint = Abs(Point2D - Projection);
+			FPoint2D ProjectionToPoint = (Point2D - Projection).GetAbs();
 
 			if (ProjectionToPoint.X > DeltaU * 0.05 || ProjectionToPoint.Y > DeltaV * 0.05)
 			{
diff --git a/Engine/Source/Runtime/Datasmith/CADKernel/Base/Private/Mesh/Structure/GridBase.h b/Engine/Source/Runtime/Datasmith/CADKernel/Base/Private/Mesh/Structure/GridBase.h
--- a/Engine/Source/Runtime/Datasmith/CADKernel/Base/Private/Mesh/Structure/GridBase.h
+++ b/Engine/Source/Runtime/Datasmith/CADKernel/Base/Private/Mesh/Structure/GridBase.h
@@ -24,8 +24,6 @@ class FIsoNode;
 class FIsoSegment;
 class FLoopNode;
 class FModelMesh;
-class FPoint;
-class FPoint2D;
 class FThinZone2DFinder;
 class FTopologicalFace;
 
diff --git a/Engine/Source/Runtime/Datasmith/CADKernel/Base/Private/Topo/TopologicalEdge.cpp b/Engine/Source/Runtime/Datasmith/CADKernel/Base/Private/Topo/TopologicalEdge.cpp
--- a/Engine/Source/Runtime/Datasmith/CADKernel/Base/Private/Topo/TopologicalEdge.cpp
+++ b/Engine/Source/Runtime/Datasmith/CADKernel/Base/Private/Topo/TopologicalEdge.cpp
@@ -1295,7 +1295,7 @@ FPoint FTopologicalEdge::GetTangentAt(const FTopologicalVertex& InVertex)
 	else
 	{
 		ensureCADKernel(false);
-		return FPoint::ZeroPoint;
+		return FPoint::ZeroVector;
 	}
 }
 
@@ -1324,7 +1324,7 @@ FPoint2D FTopologicalEdge::GetTangent2DAt(const FTopologicalVertex& InVertex)
 	else
 	{
 		ensureCADKernel(false);
-		return FPoint2D::ZeroPoint;
+		return FVector2d::ZeroVector;
 	}
 }
 
diff --git a/Engine/Source/Runtime/Datasmith/CADKernel/Base/Private/Topo/TopologicalLoop.cpp b/Engine/Source/Runtime/Datasmith/CADKernel/Base/Private/Topo/TopologicalLoop.cpp
--- a/Engine/Source/Runtime/Datasmith/CADKernel/Base/Private/Topo/TopologicalLoop.cpp
+++ b/Engine/Source/Runtime/Datasmith/CADKernel/Base/Private/Topo/TopologicalLoop.cpp
@@ -836,7 +836,7 @@ void FTopologicalLoop::RemoveDegeneratedEdges()
 		FPoint2D& DegeneratedEdgeExtremity = (bPrevious == (DegeneratedOrientedEdge.Direction == EOrientation::Front)) ? DegeneratedEdgeExtremities[0].Point2D : DegeneratedEdgeExtremities[1].Point2D;
 		FPoint2D& OtherDegeneratedEdgeExtremity = (bPrevious == (DegeneratedOrientedEdge.Direction == EOrientation::Front)) ? DegeneratedEdgeExtremities[1].Point2D : DegeneratedEdgeExtremities[0].Point2D;
 
-		double Slope = ComputeUnorientedSlope(FPoint2D::ZeroPoint, DegeneratedEdgeTangent, NearEdgeTangent);
+		double Slope = ComputeUnorientedSlope(FVector2d::ZeroVector, DegeneratedEdgeTangent, NearEdgeTangent);
 		constexpr double FlatSlope = 3.9; // ~175 deg. 
 		// a degenerated edge has to be really tangent with its neighbor edge to be merge with it
 		// otherwise it could distort the neighbor edge.
diff --git a/Engine/Source/Runtime/Datasmith/CADKernel/Base/Private/Topo/TopologicalVertex.cpp b/Engine/Source/Runtime/Datasmith/CADKernel/Base/Private/Topo/TopologicalVertex.cpp
--- a/Engine/Source/Runtime/Datasmith/CADKernel/Base/Private/Topo/TopologicalVertex.cpp
+++ b/Engine/Source/Runtime/Datasmith/CADKernel/Base/Private/Topo/TopologicalVertex.cpp
@@ -137,7 +137,7 @@ void FTopologicalVertex::UnlinkTo(FTopologicalVertex& OtherVertex)
 
 void FVertexLink::ComputeBarycenter()
 {
-	Barycenter = FPoint::ZeroPoint;
+	Barycenter = FPoint::ZeroVector;
 	for (const FTopologicalVertex* Vertex : TwinEntities)
 	{
 		Barycenter += Vertex->Coordinates;
diff --git a/Engine/Source/Runtime/Datasmith/CADKernel/Base/Public/Geo/Curves/SurfacicCurve.h b/Engine/Source/Runtime/Datasmith/CADKernel/Base/Public/Geo/Curves/SurfacicCurve.h
--- a/Engine/Source/Runtime/Datasmith/CADKernel/Base/Public/Geo/Curves/SurfacicCurve.h
+++ b/Engine/Source/Runtime/Datasmith/CADKernel/Base/Public/Geo/Curves/SurfacicCurve.h
@@ -13,7 +13,6 @@ namespace UE::CADKernel
 class FCADKernelArchive;
 class FDatabase;
 class FEntityGeom;
-class FPoint2D;
 class FSurfacicPolyline;
 struct FCurvePoint2D;
 struct FCurvePoint;
diff --git a/Engine/Source/Runtime/Datasmith/CADKernel/Base/Public/Geo/GeoPoint.h b/Engine/Source/Runtime/Datasmith/CADKernel/Base/Public/Geo/GeoPoint.h
--- a/Engine/Source/Runtime/Datasmith/CADKernel/Base/Public/Geo/GeoPoint.h
+++ b/Engine/Source/Runtime/Datasmith/CADKernel/Base/Public/Geo/GeoPoint.h
@@ -10,20 +10,20 @@ namespace UE::CADKernel
 struct FCurvePoint2D
 {
 	int32 DerivativeOrder = -1;
-	FPoint2D Point = FPoint2D::ZeroPoint;
-	FPoint2D Gradient = FPoint2D::ZeroPoint;
-	FPoint2D Laplacian = FPoint2D::ZeroPoint;
+	FPoint2D Point = FVector2d::ZeroVector;
+	FPoint2D Gradient = FVector2d::ZeroVector;
+	FPoint2D Laplacian = FVector2d::ZeroVector;
 };
 
 struct FSurfacicPoint
 {
 	int32 DerivativeOrder = -1;
-	FPoint Point = FPoint::ZeroPoint;
-	FPoint GradientU = FPoint::ZeroPoint;
-	FPoint GradientV = FPoint::ZeroPoint;
-	FPoint LaplacianU = FPoint::ZeroPoint;
-	FPoint LaplacianV = FPoint::ZeroPoint;
-	FPoint LaplacianUV = FPoint::ZeroPoint;
+	FPoint Point = FPoint::ZeroVector;
+	FPoint GradientU = FPoint::ZeroVector;
+	FPoint GradientV = FPoint::ZeroVector;
+	FPoint LaplacianU = FPoint::ZeroVector;
+	FPoint LaplacianV = FPoint::ZeroVector;
+	FPoint LaplacianUV = FPoint::ZeroVector;
 };
 
 struct FSurfacicCurvePoint
@@ -31,16 +31,16 @@ struct FSurfacicCurvePoint
 	bool bWithNormals;
 	bool bWithTangent;
 
-	FPoint2D Point2D = FPoint2D::ZeroPoint;
-	FPoint Point = FPoint::ZeroPoint;
-	FPoint Normal = FPoint::ZeroPoint;
-	FPoint Tangent = FPoint::ZeroPoint;
+	FPoint2D Point2D = FVector2d::ZeroVector;
+	FPoint Point = FPoint::ZeroVector;
+	FPoint Normal = FPoint::ZeroVector;
+	FPoint Tangent = FPoint::ZeroVector;
 };
 
 struct FSurfacicCurvePointWithTolerance
 {
-	FPoint2D Point2D = FPoint2D::ZeroPoint;
-	FPoint Point = FPoint::ZeroPoint;
+	FPoint2D Point2D = FVector2d::ZeroVector;
+	FPoint Point = FPoint::ZeroVector;
 	FSurfacicTolerance Tolerance = FPoint2D::FarawayPoint;
 };
 
@@ -49,24 +49,24 @@ typedef FSurfacicCurvePointWithTolerance FSurfacicCurveExtremities[2];
 struct FCurvePoint
 {
 	int32 DerivativeOrder = -1;
-	FPoint Point = FPoint::ZeroPoint;
-	FPoint Gradient = FPoint::ZeroPoint;
-	FPoint Laplacian = FPoint::ZeroPoint;
+	FPoint Point = FPoint::ZeroVector;
+	FPoint Gradient = FPoint::ZeroVector;
+	FPoint Laplacian = FPoint::ZeroVector;
 
 	FCurvePoint() = default;
 
 	FCurvePoint(FPoint InPoint)
 		: Point(InPoint)
 	{
-		Gradient = FPoint::ZeroPoint;
-		Laplacian = FPoint::ZeroPoint;
+		Gradient = FPoint::ZeroVector;
+		Laplacian = FPoint::ZeroVector;
 	}
 
 	void Init()
 	{
-		Point = FPoint::ZeroPoint;
-		Gradient = FPoint::ZeroPoint;
-		Laplacian = FPoint::ZeroPoint;
+		Point = FPoint::ZeroVector;
+		Gradient = FPoint::ZeroVector;
+		Laplacian = FPoint::ZeroVector;
 	}
 
 	/**
diff --git a/Engine/Source/Runtime/Datasmith/CADKernel/Base/Public/Geo/Sampling/PolylineTools.h b/Engine/Source/Runtime/Datasmith/CADKernel/Base/Public/Geo/Sampling/PolylineTools.h
--- a/Engine/Source/Runtime/Datasmith/CADKernel/Base/Public/Geo/Sampling/PolylineTools.h
+++ b/Engine/Source/Runtime/Datasmith/CADKernel/Base/Public/Geo/Sampling/PolylineTools.h
@@ -334,20 +334,20 @@ public:
 		if (FMath::IsNearlyZero(DeltaU))
 		{
 			OutPoint.Point = PolylinePoints[Index];
-			OutPoint.Gradient = FPoint::ZeroPoint;
-			OutPoint.Laplacian = FPoint::ZeroPoint;
+			OutPoint.Gradient = PointType::ZeroVector;
+			OutPoint.Laplacian = PointType::ZeroVector;
 			return;
 		}
 
 		double SectionCoordinate = (InCoordinate - PolylineCoordinates[Index]) / DeltaU;
 
-		FPoint Tangent = PolylinePoints[Index + 1] - PolylinePoints[Index];
+		PointType Tangent = PolylinePoints[Index + 1] - PolylinePoints[Index];
 		OutPoint.Point = PolylinePoints[Index] + Tangent * SectionCoordinate;
 
 		if (InDerivativeOrder > 0)
 		{
 			OutPoint.Gradient = Tangent;
-			OutPoint.Laplacian = FPoint::ZeroPoint;
+			OutPoint.Laplacian = PointType::ZeroVector;
 		}
 	}
 
@@ -372,20 +372,20 @@ public:
 				if (FMath::IsNearlyZero(DeltaU))
 				{
 					OutPoints[IPoint].Point = PolylinePoints[Index];
-					OutPoints[IPoint].Gradient = FPoint::ZeroPoint;
-					OutPoints[IPoint].Laplacian = FPoint::ZeroPoint;
+					OutPoints[IPoint].Gradient = PointType::ZeroVector;
+					OutPoints[IPoint].Laplacian = PointType::ZeroVector;
 					continue;
 				}
 
 				double SectionCoordinate = (InCoordinates[IPoint] - PolylineCoordinates[Index]) / DeltaU;
 
-				FPoint Tangent = PolylinePoints[Index + 1] - PolylinePoints[Index];
+				PointType Tangent = PolylinePoints[Index + 1] - PolylinePoints[Index];
 				OutPoints[IPoint].Point = PolylinePoints[Index] + Tangent * SectionCoordinate;
 
 				if (InDerivativeOrder > 0)
 				{
 					OutPoints[IPoint].Gradient = Tangent;
-					OutPoints[IPoint].Laplacian = FPoint::ZeroPoint;
+					OutPoints[IPoint].Laplacian = PointType::ZeroVector;
 				}
 			}
 		};
diff --git a/Engine/Source/Runtime/Datasmith/CADKernel/Base/Public/Geo/Surfaces/Surface.h b/Engine/Source/Runtime/Datasmith/CADKernel/Base/Public/Geo/Surfaces/Surface.h
--- a/Engine/Source/Runtime/Datasmith/CADKernel/Base/Public/Geo/Surfaces/Surface.h
+++ b/Engine/Source/Runtime/Datasmith/CADKernel/Base/Public/Geo/Surfaces/Surface.h
@@ -192,7 +192,7 @@ public:
 	 * For Normal(u,v) parallel Rotation Axis -> use Plan projection
 	 * For Normal(u,v) perpendicular to Rotation Axis -> use cylindrical projection
 	 */
-	virtual FPoint2D EvaluatePointInCylindricalSpace(const FPoint2D& InSurfacicCoordinate) const { return FPoint::ZeroPoint; }
+	virtual FPoint2D EvaluatePointInCylindricalSpace(const FPoint2D& InSurfacicCoordinate) const { return FVector2d::ZeroVector; }
 	virtual void EvaluatePointGridInCylindricalSpace(const FCoordinateGrid& Coordinates, TArray<FPoint2D>&) const { ; }
 
 	virtual void EvaluatePointGrid(const FCoordinateGrid& Coordinates, FSurfacicSampling& OutPoints, bool bComputeNormals = false) const;
@@ -281,7 +281,7 @@ public:
 	 * This function is useful to estimate tolerance when scales are defined in the Matrix
 	 * @param InAxis a vetor of length 1
 	 * @param InMatrix
-	 * @param InOrigin = InMatrix*FPoint::ZeroPoint
+	 * @param InOrigin = InMatrix*FPoint::ZeroVector
 	 */
 	static double ComputeScaleAlongAxis(const FPoint& InAxis, const FMatrixH& InMatrix, const FPoint& InOrigin)
 	{
diff --git a/Engine/Source/Runtime/Datasmith/CADKernel/Base/Public/Math/Geometry.h b/Engine/Source/Runtime/Datasmith/CADKernel/Base/Public/Math/Geometry.h
--- a/Engine/Source/Runtime/Datasmith/CADKernel/Base/Public/Math/Geometry.h
+++ b/Engine/Source/Runtime/Datasmith/CADKernel/Base/Public/Math/Geometry.h
@@ -38,7 +38,7 @@ CADKERNEL_API inline FPoint2D ComputeCircumCircleCenter(const FPoint2D& InPoint0
 	double D = 2. * Segment_P0_P1 ^ Segment_P0_P2;
 	if (FMath::IsNearlyZero(D, SMALL_NUMBER_SQUARE))
 	{
-		return FPoint2D::ZeroPoint;
+		return FVector2d::ZeroVector;
 	}
 
 	// CenterU  = 1/D * (Cv.|B|.|B| - By.|C|.|C|) = 1/D * CBv ^ SquareNorms 
@@ -70,7 +70,7 @@ CADKERNEL_API inline FPoint ComputeCircumCircleCenter(const FPoint& Point0, cons
 	double D = 2. * D2Point1 ^ D2Point2;
 	if (FMath::IsNearlyZero(D, SMALL_NUMBER_SQUARE))
 	{
-		return FPoint::ZeroPoint;
+		return FPoint::ZeroVector;
 	}
 
 	double SquareNormB = D2Point1.SquareLength();
@@ -93,7 +93,7 @@ CADKERNEL_API inline FPoint2D ComputeCircumCircleCenterAndSquareRadius(const FPo
 	if (FMath::IsNearlyZero(D, SMALL_NUMBER_SQUARE))
 	{
 		OutSquareRadius = 0;
-		return FPoint2D::ZeroPoint;
+		return FVector2d::ZeroVector;
 	}
 
 	double SquareNormB = Segment_P0_P1.SquareLength();
@@ -112,6 +112,7 @@ struct CADKERNEL_API TSegment
 {
 	const PointType& Point0;
 	const PointType& Point1;
+	const PointType Dummy = PointType::ZeroVector;
 
 	TSegment(const PointType& InPoint0, const PointType& InPoint1)
 		: Point0(InPoint0)
@@ -129,7 +130,7 @@ struct CADKERNEL_API TSegment
 		case 1:
 			return Point1;
 		default:
-			return PointType::ZeroPoint;
+			return Dummy;
 		}
 	}
 
@@ -173,7 +174,7 @@ struct CADKERNEL_API TTriangle
 		case 2:
 			return Point2;
 		default:
-			return PointType::ZeroPoint;
+			return PointType::ZeroVector;
 		}
 	}
 
@@ -183,7 +184,7 @@ struct CADKERNEL_API TTriangle
 		PointType Segment02 = Point2 - Point0;
 		double SquareLength01 = Segment01.SquareLength();
 		double SquareLength02 = Segment02.SquareLength();
-		double Seg01Seg02 = Segment01 * Segment02;
+		double Seg01Seg02 = Segment01 | Segment02;
 		double Det = SquareLength01 * SquareLength02 - FMath::Square(Seg01Seg02);
 
 		int32 SideIndex;
@@ -208,8 +209,8 @@ struct CADKERNEL_API TTriangle
 		{
 			// Resolve
 			PointType Segment1Point = InPoint - Point0;
-			double Segment1PointSegment01 = Segment1Point * Segment01;
-			double Segment1PointSegment02 = Segment1Point * Segment02;
+			double Segment1PointSegment01 = Segment1Point | Segment01;
+			double Segment1PointSegment02 = Segment1Point | Segment02;
 
 			OutCoordinate.X = ((Segment1PointSegment01 * SquareLength02) - (Segment1PointSegment02 * Seg01Seg02)) / Det;
 			OutCoordinate.Y = ((Segment1PointSegment02 * SquareLength01) - (Segment1PointSegment02 * Seg01Seg02)) / Det;
@@ -298,78 +299,6 @@ struct CADKERNEL_API FTriangle2D : public TTriangle<FPoint2D>
 	}
 };
 
-
-template<class PointType>
-struct CADKERNEL_API TQuadrangle
-{
-	const PointType& Point0;
-	const PointType& Point1;
-	const PointType& Point2;
-	const PointType& Point3;
-
-	TQuadrangle(const PointType& InPoint0, const PointType& InPoint1, const PointType& InPoint2, const PointType& InPoint3)
-		: Point0(InPoint0)
-		, Point1(InPoint1)
-		, Point2(InPoint2)
-		, Point3(InPoint3)
-	{
-	}
-
-	constexpr const PointType& operator[](int32 Index) const
-	{
-		ensureCADKernel(Index < 4);
-		switch (Index)
-		{
-		case 0:
-			return Point0;
-		case 1:
-			return Point1;
-		case 2:
-			return Point2;
-		case 3:
-			return Point3;
-		default:
-			return PointType::ZeroPoint;
-		}
-	}
-
-	inline FPoint ComputeNormal(const TQuadrangle<FPoint>& InQuadrangle) const
-	{
-		FPoint Normal = (InQuadrangle[1] - InQuadrangle[0]) ^ (InQuadrangle[2] - InQuadrangle[0]);
-		Normal += (InQuadrangle[2] - InQuadrangle[0]) ^ (InQuadrangle[3] - InQuadrangle[0]);
-		Normal += (InQuadrangle[1] - InQuadrangle[0]) ^ (InQuadrangle[3] - InQuadrangle[0]);
-		Normal += (InQuadrangle[1] - InQuadrangle[3]) ^ (InQuadrangle[2] - InQuadrangle[3]);
-		Normal.Normalize();
-		return Normal;
-	}
-
-	inline PointType  ProjectPoint(const PointType& InPoint, FPoint2D& OutCoordinate)
-	{
-		FPoint2D Coordinate013;
-		TTriangle<PointType> Triangle013(Point0, Point1, Point3);
-		PointType Projection013 = Triangle013.ProjectPoint(InPoint, Coordinate013);
-
-		FPoint2D Coordinate231;
-		TTriangle<PointType> Triangle231(Point2, Point3, Point1);
-		PointType Projection231 = Triangle231.ProjectPoint(InPoint, Coordinate231);
-
-		double DistanceTo013 = Projection013.Distance(InPoint);
-		double DistanceTo231 = Projection231.Distance(InPoint);
-		if (DistanceTo013 < DistanceTo231)
-		{
-			OutCoordinate = Coordinate013;
-			return Projection013;
-		}
-		else
-		{
-			OutCoordinate = { 1.0 , 1.0 };
-			OutCoordinate -= Coordinate231;
-			return Projection231;
-		}
-	}
-
-};
-
 inline FPoint ProjectPointOnPlane(const FPoint& Point, const FPoint& Origin, const FPoint& InNormal, double& OutDistance)
 {
 	FPoint Normal = InNormal;
@@ -436,7 +365,7 @@ inline PointType ProjectPointOnSegment(const PointType& Point, const PointType&
 {
 	PointType Segment = InSegmentB - InSegmentA;
 
-	double SquareLength = Segment * Segment;
+	double SquareLength = Segment | Segment;
 
 	if (SquareLength <= 0.0)
 	{
@@ -446,7 +375,7 @@ inline PointType ProjectPointOnSegment(const PointType& Point, const PointType&
 	else
 	{
 		PointType APoint = Point - InSegmentA;
-		OutCoordinate = (APoint * Segment) / SquareLength;
+		OutCoordinate = (APoint | Segment) / SquareLength;
 
 		if (bRestrictCoodinateToInside)
 		{
@@ -519,11 +448,11 @@ inline FPoint2D FindIntersectionOfSegments2D(const FSegment2D& SegmentAB, const
 	double ParallelCoef = DC ^ AB;
 	if (FMath::IsNearlyZero(ParallelCoef))
 	{
-		const double SquareAB = AB * AB;
-		double CCoordinate = (AB * AC) / SquareAB;
+		const double SquareAB = AB | AB;
+		double CCoordinate = (AB | AC) / SquareAB;
 
 		const FPoint2D AD = SegmentCD[1] - SegmentAB[0];
-		double DCoordinate = (AB * AD) / SquareAB;
+		double DCoordinate = (AB | AD) / SquareAB;
 
 		if (CCoordinate >= -DOUBLE_KINDA_SMALL_NUMBER && CCoordinate <= 1 + DOUBLE_KINDA_SMALL_NUMBER)
 		{
diff --git a/Engine/Source/Runtime/Datasmith/CADKernel/Base/Public/Math/Point.h b/Engine/Source/Runtime/Datasmith/CADKernel/Base/Public/Math/Point.h
--- a/Engine/Source/Runtime/Datasmith/CADKernel/Base/Public/Math/Point.h
+++ b/Engine/Source/Runtime/Datasmith/CADKernel/Base/Public/Math/Point.h
@@ -8,67 +8,62 @@
 namespace UE::CADKernel
 {
 
-class FPoint2D;
+struct FPoint2D;
 
-class CADKERNEL_API FPoint
+struct FPoint : public FVector
 {
-public:
-	union
-	{
-		struct
-		{
-			double X;
-			double Y;
-			double Z;
-		};
-
-		UE_DEPRECATED(all, "For internal use only")
-		double XYZ[3];
-	};
-
 public:
 
 	/** A zero point (0,0,0)*/
-	static const FPoint ZeroPoint;
-	static const FPoint UnitPoint;
 	static const FPoint FarawayPoint;
 	static const int32 Dimension;
 
 	explicit FPoint(double InCoordX = 0., double InCoordY = 0., double InCoordZ = 0.)
+		: FVector(InCoordX, InCoordY, InCoordZ)
 	{
-		X = InCoordX; Y = InCoordY; Z = InCoordZ;
 	}
 
-	FPoint(const FPoint2D& InPoint2D);
+	FPoint(const FPoint& Point)
+		: FVector(Point.X, Point.Y, Point.Z)
+	{
+	}
 
-	FPoint(const double* const InPoint)
-		: X(InPoint[0])
-		, Y(InPoint[1])
-		, Z(InPoint[2])
+	FPoint(const FVector& Point)
+		: FVector(Point.X, Point.Y, Point.Z)
 	{
 	}
+	
+	FPoint(const FPoint2D& Point);
 
-	FPoint(const FPoint& Point)
-		: X(Point.X)
-		, Y(Point.Y)
-		, Z(Point.Z)
+	FPoint(const FVector2d& InPoint2D)
+		: FVector(InPoint2D.X, InPoint2D.Y, 0.)
 	{
 	}
 
-	constexpr double& operator[](int32 Index)
+	FPoint(const double* const InPoint)
+		: FVector(InPoint[0], InPoint[1], InPoint[2])
 	{
-PRAGMA_DISABLE_DEPRECATION_WARNINGS
-		return XYZ[Index];
-PRAGMA_ENABLE_DEPRECATION_WARNINGS
 	}
 
-	constexpr double operator[](int32 Index) const
+	operator FVector3f() const
 	{
-PRAGMA_DISABLE_DEPRECATION_WARNINGS
-		return XYZ[Index];
-PRAGMA_ENABLE_DEPRECATION_WARNINGS
+		return FVector3f((float)X, (float)Y, (float)Z);
 	}
 
+//	constexpr double& operator[](int32 Index)
+//	{
+//PRAGMA_DISABLE_DEPRECATION_WARNINGS
+//		return XYZ[Index];
+//PRAGMA_ENABLE_DEPRECATION_WARNINGS
+//	}
+//
+//	constexpr double operator[](int32 Index) const
+//	{
+//PRAGMA_DISABLE_DEPRECATION_WARNINGS
+//		return XYZ[Index];
+//PRAGMA_ENABLE_DEPRECATION_WARNINGS
+//	}
+
 	bool operator==(const FPoint& Point) const
 	{
 		return FMath::IsNearlyEqual(X, Point.X) && FMath::IsNearlyEqual(Y, Point.Y) && FMath::IsNearlyEqual(Z, Point.Z);
@@ -227,16 +222,6 @@ PRAGMA_ENABLE_DEPRECATION_WARNINGS
 		return (Compare < 0);
 	}
 
-	operator FVector() const
-	{
-		return FVector(X, Y, Z);
-	}
-
-	operator FVector3f() const
-	{
-		return FVector3f((float)X, (float)Y, (float)Z);
-	}
-
 	void SetMin(const FPoint& Point)
 	{
 		X = FMath::Min(X, Point.X);
@@ -280,7 +265,7 @@ PRAGMA_ENABLE_DEPRECATION_WARNINGS
 		double Norm = Length();
 		if (FMath::IsNearlyZero(Norm))
 		{
-			*this = FPoint::UnitPoint.Normalize();
+			*this = FPoint(1.,1.,1.).Normalize();
 		}
 		else
 		{
@@ -294,7 +279,7 @@ PRAGMA_ENABLE_DEPRECATION_WARNINGS
 		double Norm = Length();
 		if (FMath::IsNearlyZero(Norm))
 		{
-			return FPoint::UnitPoint.Normalize();
+			return FPoint(1.,1.,1.).Normalize();
 		}
 		else
 		{
@@ -385,56 +370,53 @@ PRAGMA_ENABLE_DEPRECATION_WARNINGS
 	}
 };
 
-class CADKERNEL_API FPoint2D
+inline FPoint operator*(double Scale, const FPoint& Point)
 {
-public:
-	union
-	{
-		struct
-		{
-			double X;
-			double Y;
-		};
-
-		UE_DEPRECATED(all, "For internal use only")
-		double XY[2];
-	};
+	return FPoint(Point.X * Scale, Point.Y * Scale, Point.Z * Scale);
+}
 
+inline FPoint operator*(float Scale, const FPoint& Point)
+{
+	double ScaleD = static_cast<double>(Scale);
+	return FPoint(Point.X * ScaleD, Point.Y * ScaleD, Point.Z * ScaleD);
+}
 
+struct FPoint2D : public FVector2d
+{
 public:
-
 	/** A zero point (0,0,0)*/
-	static const FPoint2D ZeroPoint;
 	static const FPoint2D FarawayPoint;
 	static const int32 Dimension;
 
 	FPoint2D()
+		: FVector2d(0., 0.)
 	{
-		X = Y = 0.0;
 	}
 
-	FPoint2D(double NewU, double NewV)
-		: X(NewU)
-		, Y(NewV)
+	FPoint2D(double NewX, double NewY)
+		: FVector2d(NewX, NewY)
 	{
 	}
 
 	FPoint2D(const FPoint& Point)
-		: X(Point.X)
-		, Y(Point.Y)
+		: FVector2d(Point.X, Point.Y)
 	{
 	}
 
 	FPoint2D(const FPoint2D& Point)
-		: X(Point.X)
-		, Y(Point.Y)
+		: FVector2d(Point.X, Point.Y)
 	{
 	}
 
-	FPoint2D& operator=(const FPoint2D& Point)
+	FPoint2D(const FVector2d& Point)
+		: FVector2d(Point.X, Point.Y)
 	{
-		X = Point.X;
-		Y = Point.Y;
+	}
+
+	FPoint2D& operator+=(const FVector2d& Point)
+	{
+		X += Point.X;
+		Y += Point.Y;
 		return *this;
 	}
 
@@ -489,41 +471,6 @@ public:
 		return Normalize(Norm);
 	}
 
-	constexpr double& operator[](int32 Index)
-	{
-PRAGMA_DISABLE_DEPRECATION_WARNINGS
-		return XY[Index];
-PRAGMA_ENABLE_DEPRECATION_WARNINGS
-	}
-
-	constexpr double operator[](int32 Index) const
-	{
-PRAGMA_DISABLE_DEPRECATION_WARNINGS
-		return XY[Index];
-PRAGMA_ENABLE_DEPRECATION_WARNINGS
-	}
-
-	bool operator==(const FPoint2D& Point) const
-	{
-		return FMath::IsNearlyEqual(X, Point.X) && FMath::IsNearlyEqual(Y, Point.Y);
-	}
-
-	bool operator<(const FPoint2D& other) const
-	{
-		int32 cmp = RealCompare(X, other.X);
-		if (cmp < 0)
-		{
-			return true;
-		}
-		if (cmp > 0)
-		{
-			return false;
-		}
-
-		cmp = RealCompare(Y, other.Y);
-		return (cmp < 0);
-	}
-
 	void Set(double NewU, double NewV = 0.)
 	{
 		X = NewU;
@@ -536,84 +483,6 @@ PRAGMA_ENABLE_DEPRECATION_WARNINGS
 		Y = NewCoordinate[1];
 	}
 
-	FPoint2D operator+(const FPoint2D& Point) const
-	{
-		return FPoint2D(X + Point.X, Y + Point.Y);
-	}
-
-	FPoint2D operator-(const FPoint2D& Point) const
-	{
-		return FPoint2D(X - Point.X, Y - Point.Y);
-	}
-
-	double operator*(const FPoint2D& Point) const
-	{
-		return X * Point.X + Y * Point.Y;
-	}
-
-	double operator|(const FPoint2D& Point) const
-	{
-		return X * Point.X + Y * Point.Y;
-	}
-
-	double operator^(const FPoint2D& Point) const
-	{
-		return (X * Point.Y - Y * Point.X);
-	}
-
-	FPoint2D operator*(double Scale) const
-	{
-		return FPoint2D(X * Scale, Y * Scale);
-	}
-
-	FPoint2D operator/(double InvScale) const
-	{
-		if (FMath::IsNearlyZero(InvScale))
-		{
-			return *this;
-		}
-
-		return FPoint2D(X / InvScale, Y / InvScale);
-	}
-
-	FPoint2D operator-() const
-	{
-		return FPoint2D(-X, -Y);
-	}
-
-	FPoint2D& operator+=(const FPoint2D& Point)
-	{
-		X += Point.X;
-		Y += Point.Y;
-		return *this;
-	}
-
-	FPoint2D& operator-=(const FPoint2D& Point)
-	{
-		X -= Point.X;
-		Y -= Point.Y;
-		return *this;
-	}
-
-	FPoint2D& operator*=(double Scale)
-	{
-		X *= Scale;
-		Y *= Scale;
-		return *this;
-	}
-
-	FPoint2D& operator/=(double InvScale)
-	{
-		if (FMath::IsNearlyZero(InvScale))
-		{
-			return *this;
-		}
-
-		X /= InvScale;
-		Y /= InvScale;
-		return *this;
-	}
-
 	double ComputeCosinus(const FPoint2D& OtherVector) const
 	{
 		FPoint2D ThisNormalized = *this;
@@ -622,7 +491,7 @@ PRAGMA_ENABLE_DEPRECATION_WARNINGS
 		ThisNormalized.Normalize();
 		OtherNormalized.Normalize();
 
-		double Cosinus = ThisNormalized * OtherNormalized;
+		double Cosinus = ThisNormalized | OtherNormalized;
 
 		return FMath::Max(-1.0, FMath::Min(Cosinus, 1.0));
 	}
@@ -665,22 +534,11 @@ PRAGMA_ENABLE_DEPRECATION_WARNINGS
 	}
 };
 
-inline FPoint::FPoint(const FPoint2D& Point)
-{
-	X = Point.X; Y = Point.Y; Z = 0.0;
-}
-
 typedef FPoint2D FSurfacicTolerance;
 
-inline FPoint operator*(double Scale, const FPoint& Point)
-{
-	return FPoint(Point.X * Scale, Point.Y * Scale, Point.Z * Scale);
-}
-
-inline FPoint operator*(float Scale, const FPoint& Point)
+inline FPoint::FPoint(const FPoint2D& Point)
+	: FVector(Point, 0.)
 {
-	double ScaleD = static_cast<double>(Scale);
-	return FPoint(Point.X * ScaleD, Point.Y * ScaleD, Point.Z * ScaleD);
 }
 
 inline FPoint2D operator*(double Scale, const FPoint2D& Point)
@@ -791,8 +649,3 @@ inline uint32 GetTypeHash(const UE::CADKernel::FPoint& Point)
 	return HashCombine(GetTypeHash(Point.X), HashCombine(GetTypeHash(Point.Y), GetTypeHash(Point.Z)));
 }
 
-inline uint32 GetTypeHash(const UE::CADKernel::FPoint2D& Point)
-{
-	return HashCombine(GetTypeHash(Point.X), GetTypeHash(Point.Y));
-}
-
diff --git a/Engine/Source/Runtime/Datasmith/CADKernel/Base/Public/Topo/TopologicalLoop.h b/Engine/Source/Runtime/Datasmith/CADKernel/Base/Public/Topo/TopologicalLoop.h
--- a/Engine/Source/Runtime/Datasmith/CADKernel/Base/Public/Topo/TopologicalLoop.h
+++ b/Engine/Source/Runtime/Datasmith/CADKernel/Base/Public/Topo/TopologicalLoop.h
@@ -14,7 +14,6 @@ namespace UE::CADKernel
 {
 
 class FDatabase;
-class FPoint2D;
 class FTopologicalFace;
 class FTopologicalVertex;
 
diff --git a/Engine/Source/Runtime/Datasmith/CADKernel/Base/Public/Topo/TopologicalVertex.h b/Engine/Source/Runtime/Datasmith/CADKernel/Base/Public/Topo/TopologicalVertex.h
--- a/Engine/Source/Runtime/Datasmith/CADKernel/Base/Public/Topo/TopologicalVertex.h
+++ b/Engine/Source/Runtime/Datasmith/CADKernel/Base/Public/Topo/TopologicalVertex.h
@@ -34,13 +34,13 @@ protected:
 
 public:
 	FVertexLink()
-		: Barycenter(FPoint::ZeroPoint)
+		: Barycenter(FPoint::ZeroVector)
 	{
 	}
 
 	FVertexLink(FTopologicalVertex& Entity)
 		: TTopologicalLink<FTopologicalVertex>(Entity)
-		, Barycenter(FPoint::ZeroPoint)
+		, Barycenter(FPoint::ZeroVector)
 	{
 	}
 
