Replaced FPoint and FPoint2D by FVector and  FVector2d and remove unused FFPoint and FPointH classes - WIP - #3

diff --git a/Engine/Source/Runtime/Datasmith/CADKernel/Base/Private/Geo/Curves/PolylineCurve.h b/Engine/Source/Runtime/Datasmith/CADKernel/Base/Private/Geo/Curves/PolylineCurve.h
--- a/Engine/Source/Runtime/Datasmith/CADKernel/Base/Private/Geo/Curves/PolylineCurve.h
+++ b/Engine/Source/Runtime/Datasmith/CADKernel/Base/Private/Geo/Curves/PolylineCurve.h
@@ -41,7 +41,7 @@ protected:
 		double CurvilineLength = 0;
 		for (int32 iPoint = 1; iPoint < this->Points.Num(); iPoint++)
 		{
-			CurvilineLength += this->Points[iPoint].Distance(this->Points[iPoint - 1]);
+			CurvilineLength += PointType::Distance(this->Points[iPoint], this->Points[iPoint - 1]);
 			this->Coordinates.Add(CurvilineLength);
 		}
 
diff --git a/Engine/Source/Runtime/Datasmith/CADKernel/Base/Private/Geo/Sampler/SamplerAbstract.h b/Engine/Source/Runtime/Datasmith/CADKernel/Base/Private/Geo/Sampler/SamplerAbstract.h
--- a/Engine/Source/Runtime/Datasmith/CADKernel/Base/Private/Geo/Sampler/SamplerAbstract.h
+++ b/Engine/Source/Runtime/Datasmith/CADKernel/Base/Private/Geo/Sampler/SamplerAbstract.h
@@ -324,7 +324,7 @@ protected:
 
 	int32 CountOfNeededPointsToRespectChordError(const PointType& PointA, const PointType& PointB, double ChordError)
 	{
-		double ABLength = PointA.Distance(PointB) / 2;
+		double ABLength = PointType::Distance(PointA, PointB) / 2;
 		// Max segment length to respect the chord error
 		double LengthToHave = sqrt(FMath::Square(ABLength) * DesiredChordError / ChordError);
 		return (int32)(2 + (ABLength / LengthToHave + 0.5));
diff --git a/Engine/Source/Runtime/Datasmith/CADKernel/Base/Private/Geo/Sampling/SurfacicPolyline.cpp b/Engine/Source/Runtime/Datasmith/CADKernel/Base/Private/Geo/Sampling/SurfacicPolyline.cpp
--- a/Engine/Source/Runtime/Datasmith/CADKernel/Base/Private/Geo/Sampling/SurfacicPolyline.cpp
+++ b/Engine/Source/Runtime/Datasmith/CADKernel/Base/Private/Geo/Sampling/SurfacicPolyline.cpp
@@ -63,7 +63,7 @@ void FSurfacicPolyline::CheckIfDegenerated(const double Tolerance3D, const FSurf
 
 	FAABB2D Aabb;
 	TPolylineApproximator<FPoint2D> Approximator2D(Coordinates, Points2D);
-	Approximator2D.ComputeBoundingBox(BoundaryIndices, Boudary, Aabb);
+	Approximator2D.ComputeBoundingBox<2>(BoundaryIndices, Boudary, Aabb);
 
 	bDegeneration2D = (Aabb.GetSize(0) < ToleranceIso[EIso::IsoU] && Aabb.GetSize(1) < ToleranceIso[EIso::IsoV]);
 }
diff --git a/Engine/Source/Runtime/Datasmith/CADKernel/Base/Private/Math/Point.cpp b/Engine/Source/Runtime/Datasmith/CADKernel/Base/Private/Math/Point.cpp
--- a/Engine/Source/Runtime/Datasmith/CADKernel/Base/Private/Math/Point.cpp
+++ b/Engine/Source/Runtime/Datasmith/CADKernel/Base/Private/Math/Point.cpp
@@ -7,12 +7,8 @@
 
 namespace UE::CADKernel
 {
-
-const FPoint FPoint::FarawayPoint(HUGE_VALUE, HUGE_VALUE, HUGE_VALUE);
-const int32 FPoint::Dimension = 3;
-
-const FPoint2D FPoint2D::FarawayPoint(HUGE_VALUE, HUGE_VALUE);
-const int32 FPoint2D::Dimension = 2;
+const FVector FVectorUtil::FarawayPoint3D(HUGE_VALUE, HUGE_VALUE, HUGE_VALUE);
+const FVector2d FVectorUtil::FarawayPoint2D(HUGE_VALUE, HUGE_VALUE);
 
 double FPoint::SignedAngle(const FPoint & Other, const FPoint & Normal) const
 {
diff --git a/Engine/Source/Runtime/Datasmith/CADKernel/Base/Private/Mesh/Meshers/BowyerWatsonTriangulator.h b/Engine/Source/Runtime/Datasmith/CADKernel/Base/Private/Mesh/Meshers/BowyerWatsonTriangulator.h
--- a/Engine/Source/Runtime/Datasmith/CADKernel/Base/Private/Mesh/Meshers/BowyerWatsonTriangulator.h
+++ b/Engine/Source/Runtime/Datasmith/CADKernel/Base/Private/Mesh/Meshers/BowyerWatsonTriangulator.h
@@ -134,7 +134,7 @@ public:
 				{
 					const FPoint2D Center = TriangleSet[TriangleIndex].Center;
 
-					const double SquareDistanceToCenter = Center.SquareDistance(NewVertex);
+					const double SquareDistanceToCenter = FPoint2D::DistSquared(Center, NewVertex);
 
 					const double SquareRadiusMax = TriangleSet[TriangleIndex].SquareRadius * IncreaseFactor;
 					const double SquareRadiusMin = SquareRadiusMax * ReducingFactor;
diff --git a/Engine/Source/Runtime/Datasmith/CADKernel/Base/Private/Mesh/Meshers/CycleTriangulator.cpp b/Engine/Source/Runtime/Datasmith/CADKernel/Base/Private/Mesh/Meshers/CycleTriangulator.cpp
--- a/Engine/Source/Runtime/Datasmith/CADKernel/Base/Private/Mesh/Meshers/CycleTriangulator.cpp
+++ b/Engine/Source/Runtime/Datasmith/CADKernel/Base/Private/Mesh/Meshers/CycleTriangulator.cpp
@@ -484,7 +484,7 @@ bool FCycleTriangulator::FindTheBestCandidateNode()
 
 	const FPoint2D& StartPoint2D = FirstSideStartNode->Get2DPoint(EGridSpace::UniformScaled, Grid);
 	const FPoint2D& EndPoint2D = FirstSideEndNode->Get2DPoint(EGridSpace::UniformScaled, Grid);
-	double NewSegmentSquareLength2500 = StartPoint2D.SquareDistance(EndPoint2D) * 2500; // need to check if the length of the sides are not disproportionate
+	double NewSegmentSquareLength2500 = FPoint2D::DistSquared(StartPoint2D, EndPoint2D) * 2500; // need to check if the length of the sides are not disproportionate
 
 #ifdef DEBUG_FIND_THE_BEST_CANDIDATE_NODE
 	static int32 BestCandidateNode = 0;
@@ -554,8 +554,8 @@ bool FCycleTriangulator::FindTheBestCandidateNode()
 		}
 #endif
 
-		double CEndSquareDistance = StartPoint2D.SquareDistance(NodePoint2D);
-		double CStartSquareDistance = EndPoint2D.SquareDistance(NodePoint2D);
+		double CEndSquareDistance = FPoint2D::DistSquared(StartPoint2D, NodePoint2D);
+		double CStartSquareDistance = FPoint2D::DistSquared(EndPoint2D, NodePoint2D);
 
 		// if the previous candidate and the current candidate is Degenerated Due To Length, the previous is preferred
 		if (bIsDegeneratedDueToLength && IsDegeneratedDueToLength(CEndSquareDistance, CStartSquareDistance))
@@ -672,7 +672,7 @@ bool FCycleTriangulator::FindTheBestCandidateNode()
 	TFunction<void(FCandidateNode*, const double, const FPoint2D&, double, double) > CompareCandidateAtSegmentNode = [&CandidateNode, &Grid = Grid](FCandidateNode* NewCandidate, const double CloseLimit, const FPoint2D& SegmentPoint2D, double CandidateToSegmentSquareDistance, double SlopeDiffAtSegmentExtremity)
 	{
 		const FPoint2D& NodePoint2D = NewCandidate->Node->Get2DPoint(EGridSpace::UniformScaled, Grid);
-		const double CFirstSquareDistance = SegmentPoint2D.SquareDistance(NodePoint2D) * 2.;
+		const double CFirstSquareDistance = FPoint2D::DistSquared(SegmentPoint2D, NodePoint2D) * 2.;
 
 		if (SlopeDiffAtSegmentExtremity < CloseLimit && CFirstSquareDistance < CandidateToSegmentSquareDistance)
 		{
diff --git a/Engine/Source/Runtime/Datasmith/CADKernel/Base/Private/Mesh/Meshers/IsoTriangulator.cpp b/Engine/Source/Runtime/Datasmith/CADKernel/Base/Private/Mesh/Meshers/IsoTriangulator.cpp
--- a/Engine/Source/Runtime/Datasmith/CADKernel/Base/Private/Mesh/Meshers/IsoTriangulator.cpp
+++ b/Engine/Source/Runtime/Datasmith/CADKernel/Base/Private/Mesh/Meshers/IsoTriangulator.cpp
@@ -51,7 +51,7 @@ struct FCandidateSegment
 		: StartNode(Node1)
 		, EndNode(Node2)
 	{
-		Length = Node1.Get2DPoint(EGridSpace::UniformScaled, Grid).Distance(Node2.Get2DPoint(EGridSpace::UniformScaled, Grid));
+		Length = FPoint2D::Distance(Node1.Get2DPoint(EGridSpace::UniformScaled, Grid), Node2.Get2DPoint(EGridSpace::UniformScaled, Grid));
 	}
 };
 
@@ -1946,7 +1946,7 @@ void FCell::TryToConnectTwoSubLoopsWithShortestSegment(FCellConnexion& LoopConne
 			Grid.DisplayIsoSegment(EGridSpace::UniformScaled, *NodeA, *NodeB, 0, EVisuProperty::BlueCurve);
 #endif
 
-			double SquareDistance = ACoordinates.SquareDistance(BCoordinates);
+			double SquareDistance = FPoint2D::DistSquared(ACoordinates, BCoordinates);
 			if (SquareDistance < MinDistanceSquare)
 			{
 				MinDistanceSquare = SquareDistance;
@@ -2249,7 +2249,7 @@ void FIsoTriangulator::TryToConnectTwoSubLoopsWithTheMostIsoSegment(FCell& Cell,
 				// If the slope of the candidate segments is nearly zero, then select the shortest
 				if (MinSlope < DOUBLE_KINDA_SMALL_NUMBER)
 				{
-					double DistanceSquare = BCoordinates.SquareDistance(ACoordinates);
+					double DistanceSquare = FPoint2D::DistSquared(BCoordinates, ACoordinates);
 					if (DistanceSquare > MinLengthSquare)
 					{
 						continue;
@@ -2379,7 +2379,7 @@ void FIsoTriangulator::FindCandidateToConnectCellCornerToLoops(FCell& Cell)
 
 			const FPoint2D& ACoordinates = NodeA->Get2DPoint(EGridSpace::UniformScaled, Grid);
 
-			double SquareDistance = ACoordinates.SquareDistance(CornerPoint);
+			double SquareDistance = FPoint2D::DistSquared(ACoordinates, CornerPoint);
 			if (SquareDistance < MinDistanceSquare)
 			{
 				MinDistanceSquare = SquareDistance;
diff --git a/Engine/Source/Runtime/Datasmith/CADKernel/Base/Private/Mesh/Meshers/IsoTriangulator/IntersectionSegmentTool.h b/Engine/Source/Runtime/Datasmith/CADKernel/Base/Private/Mesh/Meshers/IsoTriangulator/IntersectionSegmentTool.h
--- a/Engine/Source/Runtime/Datasmith/CADKernel/Base/Private/Mesh/Meshers/IsoTriangulator/IntersectionSegmentTool.h
+++ b/Engine/Source/Runtime/Datasmith/CADKernel/Base/Private/Mesh/Meshers/IsoTriangulator/IntersectionSegmentTool.h
@@ -64,8 +64,8 @@ struct FSegment
 
 	static EConnectionType IsSuperimposed(const FSegment2D& SegmentAB, const FSegment2D& SegmentCD, bool bSameOrientation)
 	{
-		const FPoint2D AB = SegmentAB.GetVector().Normalize();
-		const FPoint2D CD = SegmentCD.GetVector().Normalize();
+		const FPoint2D AB = SegmentAB.GetVector().GetSafeNormal();
+		const FPoint2D CD = SegmentCD.GetVector().GetSafeNormal();
 		const double ParallelCoef = AB ^ CD;
 		if (FMath::IsNearlyZero(ParallelCoef, DOUBLE_KINDA_SMALL_NUMBER))
 		{
diff --git a/Engine/Source/Runtime/Datasmith/CADKernel/Base/Private/Mesh/Meshers/IsoTriangulator/IsoSegment.h b/Engine/Source/Runtime/Datasmith/CADKernel/Base/Private/Mesh/Meshers/IsoTriangulator/IsoSegment.h
--- a/Engine/Source/Runtime/Datasmith/CADKernel/Base/Private/Mesh/Meshers/IsoTriangulator/IsoSegment.h
+++ b/Engine/Source/Runtime/Datasmith/CADKernel/Base/Private/Mesh/Meshers/IsoTriangulator/IsoSegment.h
@@ -333,14 +333,14 @@ public:
 	{
 		const FPoint2D& Point1 = FirstNode->Get2DPoint(Space, Grid);
 		const FPoint2D& Point2 = SecondNode->Get2DPoint(Space, Grid);
-		return Point1.SquareDistance(Point2);
+		return FPoint2D::DistSquared(Point1, Point2);
 	}
 
 	double Get3DLengthSquare(const FGrid& Grid) const
 	{
 		const FPoint& Point1 = FirstNode->Get3DPoint(Grid);
 		const FPoint& Point2 = SecondNode->Get3DPoint(Grid);
-		return Point1.SquareDistance(Point2);
+		return FPoint::DistSquared(Point1, Point2);
 	}
 };
 
diff --git a/Engine/Source/Runtime/Datasmith/CADKernel/Base/Private/Mesh/Meshers/ParametricFaceMesher.cpp b/Engine/Source/Runtime/Datasmith/CADKernel/Base/Private/Mesh/Meshers/ParametricFaceMesher.cpp
--- a/Engine/Source/Runtime/Datasmith/CADKernel/Base/Private/Mesh/Meshers/ParametricFaceMesher.cpp
+++ b/Engine/Source/Runtime/Datasmith/CADKernel/Base/Private/Mesh/Meshers/ParametricFaceMesher.cpp
@@ -424,7 +424,7 @@ void FParametricFaceMesher::DefineImposedCuttingPointsBasedOnOtherSideMesh(FThin
 			double CoordSegmentU;
 			FPoint2D Projection = Segment.ProjectPoint(MeshNodeCoordinate, CoordSegmentU);
 
-			const double SquareDistance = MeshNodeCoordinate.SquareDistance(Projection);
+			const double SquareDistance = FPoint2D::DistSquared(MeshNodeCoordinate, Projection);
 			if (MinSquareThickness > SquareDistance)
 			{
 				FEdgeSegment* SegmentPtr = &Segment;
diff --git a/Engine/Source/Runtime/Datasmith/CADKernel/Base/Private/Mesh/Structure/Grid.cpp b/Engine/Source/Runtime/Datasmith/CADKernel/Base/Private/Mesh/Structure/Grid.cpp
--- a/Engine/Source/Runtime/Datasmith/CADKernel/Base/Private/Mesh/Structure/Grid.cpp
+++ b/Engine/Source/Runtime/Datasmith/CADKernel/Base/Private/Mesh/Structure/Grid.cpp
@@ -1020,7 +1020,7 @@ void FGrid::GetMeshOfThinZone(const FThinZone2D& ThinZone)
 			constexpr double Tolerance = DOUBLE_SMALL_NUMBER;
 			for (int32 Index = ThinZoneMesh.Num() - 1; Index > 0; --Index)
 			{
-				double SquareDist = ThinZoneMesh[Index - 1].SquareDistance(ThinZoneMesh[Index]);
+				double SquareDist = FPoint2D::DistSquared(ThinZoneMesh[Index - 1], ThinZoneMesh[Index]);
 				if (SquareDist < Tolerance)
 				{
 					ThinZoneMesh.RemoveAt(Index);
@@ -1028,7 +1028,7 @@ void FGrid::GetMeshOfThinZone(const FThinZone2D& ThinZone)
 			}
 			if (ThinZoneMesh.Num() > 2)
 			{
-				double SquareDist = ThinZoneMesh[0].SquareDistance(ThinZoneMesh.Last());
+				double SquareDist = FPoint2D::DistSquared(ThinZoneMesh[0], ThinZoneMesh.Last());
 				if (SquareDist < Tolerance)
 				{
 					ThinZoneMesh.RemoveAt(ThinZoneMesh.Num() - 1);
diff --git a/Engine/Source/Runtime/Datasmith/CADKernel/Base/Private/Mesh/Structure/LoopCleaner.cpp b/Engine/Source/Runtime/Datasmith/CADKernel/Base/Private/Mesh/Structure/LoopCleaner.cpp
--- a/Engine/Source/Runtime/Datasmith/CADKernel/Base/Private/Mesh/Structure/LoopCleaner.cpp
+++ b/Engine/Source/Runtime/Datasmith/CADKernel/Base/Private/Mesh/Structure/LoopCleaner.cpp
@@ -491,8 +491,8 @@ bool FLoopCleaner::TryToRemoveIntersectionOfTwoConsecutiveIntersectingSegments(c
 	}
 #endif
 
-	if ((PreviousNode->Get2DPoint(EGridSpace::UniformScaled, Grid).SquareDistance(NewCoordinate) < Tolerances.SquareGeometricTolerance2) ||
-		(NextNode->Get2DPoint(EGridSpace::UniformScaled, Grid).SquareDistance(NewCoordinate) < Tolerances.SquareGeometricTolerance2))
+	if ((FPoint2D::DistSquared(PreviousNode->Get2DPoint(EGridSpace::UniformScaled, Grid), NewCoordinate) < Tolerances.SquareGeometricTolerance2) ||
+		(FPoint2D::DistSquared(NextNode->Get2DPoint(EGridSpace::UniformScaled, Grid), NewCoordinate) < Tolerances.SquareGeometricTolerance2))
 	{
 		if(RemoveNodeOfLoop(*Node))
 		{
@@ -564,8 +564,8 @@ void FLoopCleaner::RemoveIntersectionByMovingOutsideSegmentNodeInside(const FIso
 	SegmentPoints[0] = Segment.GetFirstNode().Get2DPoint(EGridSpace::UniformScaled, Grid);
 	SegmentPoints[1] = Segment.GetSecondNode().Get2DPoint(EGridSpace::UniformScaled, Grid);
 
-	double SquareLenghtIntersectingSegment = IntersectingPoints[0].SquareDistance(IntersectingPoints[1]);
-	double SquareLenghtSegment = SegmentPoints[0].SquareDistance(SegmentPoints[1]);
+	double SquareLenghtIntersectingSegment = FPoint2D::DistSquared(IntersectingPoints[0], IntersectingPoints[1]);
+	double SquareLenghtSegment = FPoint2D::DistSquared(SegmentPoints[0], SegmentPoints[1]);
 
 	if (SquareLenghtSegment > 10 * SquareLenghtIntersectingSegment)
 	{
@@ -612,8 +612,8 @@ void FLoopCleaner::RemoveIntersectionByMovingOutsideSegmentNodeInside(const FIso
 	FLoopNode& PreviousNode = Node.GetPreviousNode();
 	FLoopNode& NextNode = Node.GetNextNode();
 
-	if ((PreviousNode.Get2DPoint(EGridSpace::UniformScaled, Grid).SquareDistance(NewCoordinate) < Tolerances.SquareGeometricTolerance2) ||
-		(NextNode.Get2DPoint(EGridSpace::UniformScaled, Grid).SquareDistance(NewCoordinate) < Tolerances.SquareGeometricTolerance2))
+	if ((FPoint2D::DistSquared(PreviousNode.Get2DPoint(EGridSpace::UniformScaled, Grid), NewCoordinate) < Tolerances.SquareGeometricTolerance2) ||
+		(FPoint2D::DistSquared(NextNode.Get2DPoint(EGridSpace::UniformScaled, Grid), NewCoordinate) < Tolerances.SquareGeometricTolerance2))
 	{
 		RemoveNodeOfLoop(Node);
 		return;
@@ -1294,7 +1294,7 @@ bool FLoopCleaner::MoveIntersectingSectionBehindOppositeSection(LoopCleanerImpl:
 		for (int32 Index = 1; Index < OppositeSectionCount; ++Index)
 		{
 			FPoint2D ProjectedPoint = ProjectPointOnSegment(PointToProject, *OppositeSectionPoint[Index - 1], *OppositeSectionPoint[Index], Coordinate);
-			double SquareDistance = PointToProject.SquareDistance(ProjectedPoint);
+			double SquareDistance = FPoint2D::DistSquared(PointToProject, ProjectedPoint);
 			if (SquareDistance < MinSquareDistance)
 			{
 				MinSquareDistance = SquareDistance;
@@ -1393,10 +1393,11 @@ bool FLoopCleaner::DisconnectCoincidentNodes(const LoopCleanerImpl::FPinchInters
 	double MediumSlope0 = TransformIntoPositiveSlope((Slope00 + Slope02) * 0.5);
 	double MediumSlope1 = TransformIntoPositiveSlope((Slope10 + Slope12) * 0.5);
 
-	FPoint2D P1 = SlopeToVector(MediumSlope1);
-	FPoint2D P0 = SlopeToVector(MediumSlope0);
-	P1.Normalize()*= Tolerances.GeometricTolerance*30;
-	P0.Normalize()*= Tolerances.GeometricTolerance*30;
+	// #cadkernel_check: Why is this constant equal to 30????
+	constexpr double ScaleFactor = 30.;
+	const double Scale = Tolerances.GeometricTolerance * ScaleFactor;
+	FPoint2D P1 = SlopeToVector(MediumSlope1).GetSafeNormal() * Scale;
+	FPoint2D P0 = SlopeToVector(MediumSlope0).GetSafeNormal() * Scale;
 
 	P0 += *Context.Points[0][1];
 	P1 += *Context.Points[1][1];
@@ -1693,8 +1694,8 @@ void FLoopCleaner::MoveNode(FLoopNode& NodeToMove, FPoint2D& NewPostion)
 	FLoopNode& PreviousNode = NodeToMove.GetPreviousNode();
 	FLoopNode& NextNode = NodeToMove.GetNextNode();
 
-	if ((PreviousNode.Get2DPoint(EGridSpace::UniformScaled, Grid).SquareDistance(NewCoordinate) < Tolerances.SquareGeometricTolerance2) ||
-		(NextNode.Get2DPoint(EGridSpace::UniformScaled, Grid).SquareDistance(NewCoordinate) < Tolerances.SquareGeometricTolerance2))
+	if ((FPoint2D::DistSquared(PreviousNode.Get2DPoint(EGridSpace::UniformScaled, Grid), NewCoordinate) < Tolerances.SquareGeometricTolerance2) ||
+		(FPoint2D::DistSquared(NextNode.Get2DPoint(EGridSpace::UniformScaled, Grid), NewCoordinate) < Tolerances.SquareGeometricTolerance2))
 	{
 		RemoveNodeOfLoop(NodeToMove);
 	}
@@ -1774,8 +1775,8 @@ void FLoopCleaner::FindLoopIntersections()
 				double SegmentIntersectionCoordinate;
 				FPoint2D IntersectionPoint = FindIntersectionOfSegments2D(SegmentPoints, IntersectedSegmentPoints, SegmentIntersectionCoordinate);
 
-				const double IntersectedSegmentPointsSquareLength = IntersectedSegmentPoints.SquareLength();
-				const double SquareDistanceToIntersectionPoint = IntersectedSegmentPoints[0].SquareDistance(IntersectionPoint);
+				const double IntersectedSegmentPointsSquareLength = IntersectedSegmentPoints.SquaredLength();
+				const double SquareDistanceToIntersectionPoint = FPoint2D::DistSquared(IntersectedSegmentPoints[0], IntersectionPoint);
 				if (SquareDistanceToIntersectionPoint > IntersectedSegmentPointsSquareLength)
 				{
 					// the IntersectionPoint is not inside IntersectedSegmentPoints
@@ -2182,8 +2183,8 @@ void FLoopCleaner::FindBestLoopExtremity()
 		FLoopNode& PreviousNode = Node->GetPreviousNode();
 		FLoopNode& NextNode = Node->GetNextNode();
 
-		if (Node->Get2DPoint(EGridSpace::UniformScaled, Grid).SquareDistance(PreviousNode.Get2DPoint(EGridSpace::UniformScaled, Grid)) < DOUBLE_SMALL_NUMBER ||
-			Node->Get2DPoint(EGridSpace::UniformScaled, Grid).SquareDistance(NextNode.Get2DPoint(EGridSpace::UniformScaled, Grid)) < DOUBLE_SMALL_NUMBER)
+		if (FPoint2D::DistSquared(Node->Get2DPoint(EGridSpace::UniformScaled, Grid), PreviousNode.Get2DPoint(EGridSpace::UniformScaled, Grid)) < DOUBLE_SMALL_NUMBER ||
+			FPoint2D::DistSquared(Node->Get2DPoint(EGridSpace::UniformScaled, Grid), NextNode.Get2DPoint(EGridSpace::UniformScaled, Grid)) < DOUBLE_SMALL_NUMBER)
 		{
 			return;
 		}
diff --git a/Engine/Source/Runtime/Datasmith/CADKernel/Base/Private/Mesh/Structure/LoopCleaner.h b/Engine/Source/Runtime/Datasmith/CADKernel/Base/Private/Mesh/Structure/LoopCleaner.h
--- a/Engine/Source/Runtime/Datasmith/CADKernel/Base/Private/Mesh/Structure/LoopCleaner.h
+++ b/Engine/Source/Runtime/Datasmith/CADKernel/Base/Private/Mesh/Structure/LoopCleaner.h
@@ -201,7 +201,7 @@ private:
 	 */
 	bool CheckAndRemoveCoincidence(const FPoint2D& Point0, const FPoint2D& Point1, FLoopNode& NodeToRemove)
 	{
-		double SquareDistance = Point0.SquareDistance(Point1);
+		double SquareDistance = FPoint2D::DistSquared(Point0, Point1);
 		if (SquareDistance < Tolerances.SquareGeometricTolerance2)
 		{
 			return RemoveNodeOfLoop(NodeToRemove);
diff --git a/Engine/Source/Runtime/Datasmith/CADKernel/Base/Private/Mesh/Structure/ThinZone2DFinder.cpp b/Engine/Source/Runtime/Datasmith/CADKernel/Base/Private/Mesh/Structure/ThinZone2DFinder.cpp
--- a/Engine/Source/Runtime/Datasmith/CADKernel/Base/Private/Mesh/Structure/ThinZone2DFinder.cpp
+++ b/Engine/Source/Runtime/Datasmith/CADKernel/Base/Private/Mesh/Structure/ThinZone2DFinder.cpp
@@ -143,7 +143,7 @@ void FThinZone2DFinder::FindCloseSegments()
 
 			// if the distance of FirstPointCandidate with the middle of current segment is biggest than MaxSpace, then the projection of the point cannot be smaller than the tolerance, 
 			{
-				const double SquareDistance = SegmentMiddle.SquareDistance(FirstPointCandidate);
+				const double SquareDistance = FPoint2D::DistSquared(SegmentMiddle, FirstPointCandidate);
 				if (SquareDistance > MaxSquareSpace)
 				{
 					continue;
@@ -171,7 +171,7 @@ void FThinZone2DFinder::FindCloseSegments()
 			}
 #endif
 
-			const double SquareDistance = SegmentMiddle.SquareDistance(Projection);
+			const double SquareDistance = FPoint2D::DistSquared(SegmentMiddle, Projection);
 			if (SquareDistance > SquareFinderTolerance)
 			{
 				continue;
@@ -262,7 +262,7 @@ void FThinZone2DFinder::FindCloseSegments(const TArray<FEdgeSegment*>& Segments,
 				double Coordinate;
 				const FPoint2D Projection = ProjectPointOnSegment(SegmentMiddle, FirstPointCandidate, SecondPointCandidate, Coordinate, true);
 
-				const double SquareDistance = SegmentMiddle.SquareDistance(Projection);
+				const double SquareDistance = FPoint2D::DistSquared(SegmentMiddle, Projection);
 				if (MinSquareThickness > SquareDistance)
 				{
 					MinSquareThickness = SquareDistance;
@@ -450,7 +450,7 @@ bool FThinZone2DFinder::CheckIfCloseSideOfThinSideIsNotDegenerated(TArray<FEdgeS
 			double Coordinate;
 			const FPoint2D Projection = ProjectPointOnSegment(Candidate, FirstPointCandidate, SecondPointCandidate, Coordinate, true);
 
-			const double SquareDistance = Candidate.SquareDistance(Projection);
+			const double SquareDistance = FPoint2D::DistSquared(Candidate, Projection);
 			if (MinSquareThickness > SquareDistance)
 			{
 				MinSquareThickness = SquareDistance;
@@ -568,7 +568,7 @@ bool FThinZone2DFinder::CheckIfCloseSideOfThinSideIsNotDegenerated(TArray<FEdgeS
 		}
 		const FPoint2D& NextPoint = Segment->GetExtemity(End);
 		FPoint2D OppositePoint = GetOppositeSegment(NextPoint);
-		const double SquarewDistance = OppositePoint.SquareDistance(FirstOppositePoint);
+		const double SquarewDistance = FPoint2D::DistSquared(OppositePoint, FirstOppositePoint);
 
 #ifdef DEBUG_CHECK_THIN_SIDE_
 		{
@@ -1409,7 +1409,7 @@ bool FThinZone2DFinder::BuildLoopSegments()
 			// Remove duplicated points (vs GeometricTolerance) of the end
 			{
 				int32 EndPointIndex = ScaledPoints.Num() - 1;
-				while (EndPointIndex > 0 && ScaledPoints[EndPointIndex].Distance(ScaledPoints[EndPointIndex - 1]) < GeometricTolerance)
+				while (EndPointIndex > 0 && FPoint2D::Distance(ScaledPoints[EndPointIndex], ScaledPoints[EndPointIndex - 1]) < GeometricTolerance)
 				{
 					EndPointIndex--;
 					ScaledPoints.RemoveAt(EndPointIndex);
@@ -1426,7 +1426,7 @@ bool FThinZone2DFinder::BuildLoopSegments()
 			{
 				for (int32 PointIndex = 1; PointIndex < ScaledPoints.Num(); PointIndex++)
 				{
-					while (PointIndex < ScaledPoints.Num() && ScaledPoints[PointIndex - 1].Distance(ScaledPoints[PointIndex]) < GeometricTolerance)
+					while (PointIndex < ScaledPoints.Num() && FPoint2D::Distance(ScaledPoints[PointIndex - 1], ScaledPoints[PointIndex]) < GeometricTolerance)
 					{
 						ScaledPoints.RemoveAt(PointIndex);
 						Coordinates.RemoveAt(PointIndex);
diff --git a/Engine/Source/Runtime/Datasmith/CADKernel/Base/Private/Topo/TopologicalLoop.cpp b/Engine/Source/Runtime/Datasmith/CADKernel/Base/Private/Topo/TopologicalLoop.cpp
--- a/Engine/Source/Runtime/Datasmith/CADKernel/Base/Private/Topo/TopologicalLoop.cpp
+++ b/Engine/Source/Runtime/Datasmith/CADKernel/Base/Private/Topo/TopologicalLoop.cpp
@@ -38,7 +38,7 @@ TSharedPtr<FTopologicalLoop> FTopologicalLoop::Make(const TArray<TSharedPtr<FTop
 			const bool Edge0IsIso = Edge0.GetCurve()->IsIso(EIso::IsoU, TolU);
 			const bool Edge1IsIso = Edge1.GetCurve()->IsIso(EIso::IsoU, TolU);
 			const FPoint2D Edge1EndPoint = Edge1.Approximate2DPoint(InEdgeDirections[0]== InEdgeDirections[1] ? Edge1.GetBoundary().GetMax() : Edge1.GetBoundary().GetMin());
-			double Distance = Edge1EndPoint.SquareDistance(StartPoint);
+			double Distance = FPoint2D::DistSquared(Edge1EndPoint, StartPoint);
 			if (Distance < FMath::Square(TolU))
 			{
 				const FPoint2D Offset(TWO_PI, 0.);
@@ -352,8 +352,8 @@ bool FTopologicalLoop::Orient()
 
 		if(Slope > 7.9 || Slope < 0.1)
 		{
-			double SquareLengthBefore = LoopSampling[Index].SquareDistance(LoopSampling[PreviousIndex]);
-			double SquareLengthAfter = LoopSampling[Index].SquareDistance(LoopSampling[NextIndex]);
+			double SquareLengthBefore = FPoint2D::DistSquared(LoopSampling[Index], LoopSampling[PreviousIndex]);
+			double SquareLengthAfter = FPoint2D::DistSquared(LoopSampling[Index], LoopSampling[NextIndex]);
 			if(SquareLengthBefore < SquareLengthAfter)
 			{
 				Index = PreviousIndex;
@@ -706,21 +706,21 @@ void FTopologicalLoop::CheckEdgesOrientation()
 	{
 		int32 PIndex = 0;
 		int32 EIndex = 0;
-		double SmallDistance = PreviousExtremities[0].Point2D.Distance(EdgeExtremities[0].Point2D);
-		double Distance = PreviousExtremities[0].Point2D.Distance(EdgeExtremities[1].Point2D);
+		double SmallDistance = FPoint2D::Distance(PreviousExtremities[0].Point2D, EdgeExtremities[0].Point2D);
+		double Distance = FPoint2D::Distance(PreviousExtremities[0].Point2D, EdgeExtremities[1].Point2D);
 		if (Distance < SmallDistance)
 		{
 			SmallDistance = Distance;
 			EIndex = 1;
 		}
-		Distance = PreviousExtremities[1].Point2D.Distance(EdgeExtremities[0].Point2D);
+		Distance = FPoint2D::Distance(PreviousExtremities[1].Point2D, EdgeExtremities[0].Point2D);
 		if (Distance < SmallDistance)
 		{
 			SmallDistance = Distance;
 			PIndex = 1;
 			EIndex = 0;
 		}
-		Distance = PreviousExtremities[1].Point2D.Distance(EdgeExtremities[1].Point2D);
+		Distance = FPoint2D::Distance(PreviousExtremities[1].Point2D, EdgeExtremities[1].Point2D);
 		if (Distance < SmallDistance)
 		{
 			SmallDistance = Distance;
@@ -767,21 +767,21 @@ void FTopologicalLoop::CheckLoopWithTwoEdgesOrientation()
 
 	int32 IndexEdge0 = 0;
 	int32 IndexEdge1 = 0;
-	double SmallDistance = Edge0Extremities[0].Point2D.Distance(Edge1Extremities[0].Point2D);
-	double Distance = Edge0Extremities[0].Point2D.Distance(Edge1Extremities[1].Point2D);
+	double SmallDistance = FPoint2D::Distance(Edge0Extremities[0].Point2D, Edge1Extremities[0].Point2D);
+	double Distance = FPoint2D::Distance(Edge0Extremities[0].Point2D, Edge1Extremities[1].Point2D);
 	if (Distance < SmallDistance)
 	{
 		SmallDistance = Distance;
 		IndexEdge1 = 1;
 	}
-	Distance = Edge0Extremities[1].Point2D.Distance(Edge1Extremities[0].Point2D);
+	Distance = FPoint2D::Distance(Edge0Extremities[1].Point2D, Edge1Extremities[0].Point2D);
 	if (Distance < SmallDistance)
 	{
 		SmallDistance = Distance;
 		IndexEdge0 = 1;
 		IndexEdge1 = 0;
 	}
-	Distance = Edge0Extremities[1].Point2D.Distance(Edge1Extremities[1].Point2D);
+	Distance = FPoint2D::Distance(Edge0Extremities[1].Point2D, Edge1Extremities[1].Point2D);
 	if (Distance < SmallDistance)
 	{
 		SmallDistance = Distance;
diff --git a/Engine/Source/Runtime/Datasmith/CADKernel/Base/Public/Geo/GeoPoint.h b/Engine/Source/Runtime/Datasmith/CADKernel/Base/Public/Geo/GeoPoint.h
--- a/Engine/Source/Runtime/Datasmith/CADKernel/Base/Public/Geo/GeoPoint.h
+++ b/Engine/Source/Runtime/Datasmith/CADKernel/Base/Public/Geo/GeoPoint.h
@@ -41,7 +41,7 @@ struct FSurfacicCurvePointWithTolerance
 {
 	FPoint2D Point2D = FVector2d::ZeroVector;
 	FPoint Point = FPoint::ZeroVector;
-	FSurfacicTolerance Tolerance = FPoint2D::FarawayPoint;
+	FSurfacicTolerance Tolerance = FVectorUtil::FarawayPoint2D;
 };
 
 typedef FSurfacicCurvePointWithTolerance FSurfacicCurveExtremities[2];
diff --git a/Engine/Source/Runtime/Datasmith/CADKernel/Base/Public/Geo/Sampling/PolylineTools.h b/Engine/Source/Runtime/Datasmith/CADKernel/Base/Public/Geo/Sampling/PolylineTools.h
--- a/Engine/Source/Runtime/Datasmith/CADKernel/Base/Public/Geo/Sampling/PolylineTools.h
+++ b/Engine/Source/Runtime/Datasmith/CADKernel/Base/Public/Geo/Sampling/PolylineTools.h
@@ -91,7 +91,7 @@ double ComputeLength(const TArray<PointType>& Polyline)
 	double Length = 0;
 	for (int32 Index = 1; Index < Polyline.Num(); ++Index)
 	{
-		Length += Polyline[Index - 1].Distance(Polyline[Index]);
+		Length += PointType::Distance(Polyline[Index - 1], Polyline[Index]);
 	}
 	return Length;
 }
@@ -111,25 +111,25 @@ TArray<double> ComputePolylineSegmentLengths(const PointType& StartPoint, const
 	if (InnerPointCount > 0)
 	{
 		ElementLength.Reserve(InnerPointCount + 1);
-		const FPoint* PrevPoint = &InnerPolyline[0];
+		const PointType* PrevPoint = &InnerPolyline[0];
 		{
-			ElementLength.Add(PrevPoint->Distance(StartPoint));
+			ElementLength.Add(PointType::Distance(*PrevPoint, StartPoint));
 		}
 
 		for (int32 Index = 1; Index < InnerPointCount; ++Index)
 		{
 			const FPoint& CurrentPoint = InnerPolyline[Index];
-			ElementLength.Add(PrevPoint->Distance(CurrentPoint));
+			ElementLength.Add(PointType::Distance(*PrevPoint, CurrentPoint));
 			PrevPoint = &CurrentPoint;
 		}
 		{
-			ElementLength.Add(PrevPoint->Distance(EndPoint));
+			ElementLength.Add(PointType::Distance(*PrevPoint, EndPoint));
 		}
 	}
 	else
 	{
 		ElementLength.Reserve(1);
-		ElementLength.Add(StartPoint.Distance(EndPoint));
+		ElementLength.Add(PointType::Distance(StartPoint, EndPoint));
 	}
 
 	return MoveTemp(ElementLength);
@@ -141,8 +141,8 @@ TArray<double> ComputePolylineSegmentLengths(const PointType& StartPoint, const
 template<typename PointType>
 void ExtendTo(TArray<PointType>& Polyline, const PointType& DesiredEnd)
 {
-	double DistanceStart = Polyline[0].SquareDistance(DesiredEnd);
-	double DistanceEnd = Polyline.Last().SquareDistance(DesiredEnd);
+	double DistanceStart = PointType::DistSquared(Polyline[0], DesiredEnd);
+	double DistanceEnd = PointType::DistSquared(Polyline.Last(), DesiredEnd);
 
 	if (DistanceStart < DistanceEnd)
 	{
@@ -155,10 +155,10 @@ void ExtendTo(TArray<PointType>& Polyline, const PointType& DesiredEnd)
 	Delta /= PolylineLength;
 
 	PolylineLength = 0;
-	PolylineLength = Polyline[1].Distance(Polyline[0]);
+	PolylineLength = PointType::Distance(Polyline[1], Polyline[0]);
 	for (int32 Index = 1; Index < Polyline.Num() - 1; ++Index)
 	{
-		double LengthNextSegment = Polyline[Index].Distance(Polyline[Index + 1]);
+		double LengthNextSegment = PointType::Distance(Polyline[Index], Polyline[Index + 1]);
 		Polyline[Index] += Delta * PolylineLength;
 		PolylineLength += LengthNextSegment;
 	}
@@ -219,16 +219,16 @@ protected:
 			PointType EndPoint = ComputePoint(BoundaryIndices[1], InBoundary.Max);
 			PointType NextPointIndice1 = PolylinePoints[BoundaryIndices[1] + 1];
 
-			EdgeLength = StartPoint.Distance(PolylinePoints[BoundaryIndices[0] + 1]);
-			double LengthOfSegments = PointIndice0.Distance(PolylinePoints[BoundaryIndices[0] + 1]);
+			EdgeLength = PointType::Distance(StartPoint, PolylinePoints[BoundaryIndices[0] + 1]);
+			double LengthOfSegments = PointType::Distance(PointIndice0, PolylinePoints[BoundaryIndices[0] + 1]);
 
-			LastEdgeSegmentLength = EndPoint.Distance(PolylinePoints[BoundaryIndices[1]]);
-			const double LastSegmentLength = NextPointIndice1.Distance(PolylinePoints[BoundaryIndices[1]]);
+			LastEdgeSegmentLength = PointType::Distance(EndPoint, PolylinePoints[BoundaryIndices[1]]);
+			const double LastSegmentLength = PointType::Distance(NextPointIndice1, PolylinePoints[BoundaryIndices[1]]);
 
 			OutCurvilinearCoordinates.Add(LengthOfSegments);
 			for (int32 Index = BoundaryIndices[0] + 1; Index < BoundaryIndices[1]; ++Index)
 			{
-				const double SegLength = PolylinePoints[Index].Distance(PolylinePoints[Index + 1]);
+				const double SegLength = PointType::Distance(PolylinePoints[Index], PolylinePoints[Index + 1]);
 				EdgeLength += SegLength;
 				LengthOfSegments += SegLength;
 				OutCurvilinearCoordinates.Add(LengthOfSegments);
@@ -241,8 +241,8 @@ protected:
 		{
 			PointType StartPoint = ComputePoint(BoundaryIndices[0], InBoundary.Min);
 			PointType EndPoint = ComputePoint(BoundaryIndices[1], InBoundary.Max);
-			EdgeLength = StartPoint.Distance(EndPoint);
-			const double SegLength = PolylinePoints[BoundaryIndices[0]].Distance(PolylinePoints[BoundaryIndices[0] + 1]);
+			EdgeLength = PointType::Distance(StartPoint, EndPoint);
+			const double SegLength = PointType::Distance(PolylinePoints[BoundaryIndices[0]], PolylinePoints[BoundaryIndices[0] + 1]);
 			OutCurvilinearCoordinates.Add(SegLength);
 		}
 		return EdgeLength;
@@ -459,7 +459,7 @@ public:
 		const double CurveLength = ComputeCurvilinearCoordinatesOfPolyline(InBoundary, CurvilinearCoordinates, BoundaryIndices);
 
 		const PointType StartPoint = ComputePoint(BoundaryIndices[0], InBoundary.Min);
-		double FromStartSegmentLength = PolylinePoints[BoundaryIndices[0]].Distance(StartPoint);
+		double FromStartSegmentLength = PointType::Distance(PolylinePoints[BoundaryIndices[0]], StartPoint);
 
 		if (FMath::IsNearlyZero(DesiredSegmentLength, DOUBLE_KINDA_SMALL_NUMBER) || CurveLength < 2. * DesiredSegmentLength)
 		{
@@ -778,14 +778,14 @@ public:
 		{
 			PointType StartPoint = ComputePoint(BoundaryIndex[0], InBoundary.Min);
 			PointType EndPoint = ComputePoint(BoundaryIndex[1], InBoundary.Max);
-			Length = StartPoint.Distance(PolylinePoints[BoundaryIndex[0] + 1]);
-			Length += EndPoint.Distance(PolylinePoints[BoundaryIndex[1]]);
+			Length = PointType::Distance(StartPoint, PolylinePoints[BoundaryIndex[0] + 1]);
+			Length += PointType::Distance(EndPoint, PolylinePoints[BoundaryIndex[1]]);
 
 			if (BoundaryIndex[1] > BoundaryIndex[0] + 1)
 			{
 				for (int32 Index = BoundaryIndex[0] + 1; Index < BoundaryIndex[1]; ++Index)
 				{
-					Length += PolylinePoints[Index].Distance(PolylinePoints[Index + 1]);
+					Length += PointType::Distance(PolylinePoints[Index], PolylinePoints[Index + 1]);
 				}
 			}
 		}
@@ -793,7 +793,7 @@ public:
 		{
 			PointType StartPoint = ComputePoint(BoundaryIndex[0], InBoundary.Min);
 			PointType EndPoint = ComputePoint(BoundaryIndex[1], InBoundary.Max);
-			Length = StartPoint.Distance(EndPoint);
+			Length = PointType::Distance(StartPoint, EndPoint);
 		}
 		return Length;
 	}
@@ -805,7 +805,8 @@ public:
 		return ComputeLengthOfSubPolyline(BoundaryIndices, InBoundary);
 	}
 
-	void ComputeBoundingBox(const int BoundaryIndex[2], const FLinearBoundary& InBoundary, TAABB<PointType>& Aabb)
+	template<int Dimension = 3>
+	void ComputeBoundingBox(const int BoundaryIndex[2], const FLinearBoundary& InBoundary, TAABB<PointType, Dimension>& Aabb)
 	{
 		Aabb.Empty();
 		if (BoundaryIndex[1] > BoundaryIndex[0])
diff --git a/Engine/Source/Runtime/Datasmith/CADKernel/Base/Public/Geo/Sampling/SurfacicPolyline.h b/Engine/Source/Runtime/Datasmith/CADKernel/Base/Public/Geo/Sampling/SurfacicPolyline.h
--- a/Engine/Source/Runtime/Datasmith/CADKernel/Base/Public/Geo/Sampling/SurfacicPolyline.h
+++ b/Engine/Source/Runtime/Datasmith/CADKernel/Base/Public/Geo/Sampling/SurfacicPolyline.h
@@ -110,12 +110,12 @@ public:
 		double Distance3D = Points3D[Index].Distance(Points3D[Index + 1]);
 		if (FMath::IsNearlyZero(Distance3D, (double)DOUBLE_SMALL_NUMBER))
 		{
-			return FPoint2D::FarawayPoint;
+			return FVectorUtil::FarawayPoint2D;
 		}
 		else
 		{
 			FSurfacicTolerance Tolerance2D = Points2D[Index] - Points2D[Index + 1];
-			return Max(Abs(Tolerance2D) * Tolerance3D / Distance3D, MinToleranceIso);
+			return FVector2D::Max(Tolerance2D.GetAbs() * Tolerance3D / Distance3D, MinToleranceIso);
 		}
 	};
 
diff --git a/Engine/Source/Runtime/Datasmith/CADKernel/Base/Public/Math/Aabb.h b/Engine/Source/Runtime/Datasmith/CADKernel/Base/Public/Math/Aabb.h
--- a/Engine/Source/Runtime/Datasmith/CADKernel/Base/Public/Math/Aabb.h
+++ b/Engine/Source/Runtime/Datasmith/CADKernel/Base/Public/Math/Aabb.h
@@ -18,7 +18,7 @@ ENUM_CLASS_FLAGS(EAABBBoundary);
 
 namespace UE::CADKernel
 {
-template<class PointType>
+template<class PointType, int Dimension = 3>
 class TAABB
 {
 
@@ -28,8 +28,6 @@ protected:
 
 public:
 	TAABB()
-		: MinCorner(PointType::FarawayPoint)
-		, MaxCorner(-PointType::FarawayPoint)
 	{
 	}
 
@@ -48,7 +46,7 @@ public:
 
 	bool IsValid() const
 	{
-		for (int32 Axis = 0; Axis < PointType::Dimension; Axis++)
+		for (int32 Axis = 0; Axis < Dimension; Axis++)
 		{
 			if (MinCorner[Axis] > MaxCorner[Axis])
 			{
@@ -60,13 +58,11 @@ public:
 
 	void Empty()
 	{
-		MinCorner = PointType::FarawayPoint;
-		MaxCorner = -PointType::FarawayPoint;
 	}
 
 	bool Contains(const PointType& Point) const
 	{
-		for (int32 Axis = 0; Axis < PointType::Dimension; Axis++)
+		for (int32 Axis = 0; Axis < Dimension; Axis++)
 		{
 			if ((Point[Axis] < MinCorner[Axis]) || (Point[Axis] > MaxCorner[Axis]))
 			{
@@ -78,7 +74,7 @@ public:
 
 	void SetMinSize(double MinSize)
 	{
-		for (int32 Axis = 0; Axis < PointType::Dimension; Axis++)
+		for (int32 Axis = 0; Axis < Dimension; Axis++)
 		{
 			double AxisSize = GetSize(Axis);
 			if (AxisSize < MinSize)
@@ -93,7 +89,7 @@ public:
 	double GetMaxSize() const
 	{
 		double MaxSideSize = 0;
-		for (int32 Index = 0; Index < PointType::Dimension; Index++)
+		for (int32 Index = 0; Index < Dimension; Index++)
 		{
 			double Size = GetSize(Index);
 			if (Size > MaxSideSize)
@@ -111,7 +107,7 @@ public:
 
 	double DiagonalLength() const
 	{
-		return MaxCorner.Distance(MinCorner);
+		return PointType::Distance(MaxCorner, MinCorner);
 	}
 
 	PointType Diagonal() const
@@ -136,7 +132,7 @@ public:
 
 	TAABB& operator+= (const double* Point)
 	{
-		for (int32 Index = 0; Index < PointType::Dimension; Index++)
+		for (int32 Index = 0; Index < Dimension; Index++)
 		{
 			if (Point[Index] < MinCorner[Index])
 			{
@@ -152,7 +148,7 @@ public:
 
 	TAABB& operator+= (const PointType& Point)
 	{
-		for (int32 Index = 0; Index < PointType::Dimension; Index++)
+		for (int32 Index = 0; Index < Dimension; Index++)
 		{
 			if (Point[Index] < MinCorner[Index])
 			{
@@ -179,7 +175,7 @@ public:
 
 	void Offset(double Offset)
 	{
-		for (int32 Index = 0; Index < PointType::Dimension; Index++)
+		for (int32 Index = 0; Index < Dimension; Index++)
 		{
 			MinCorner[Index] -= Offset;
 			MaxCorner[Index] += Offset;
@@ -209,6 +205,35 @@ public:
 	}
 };
 
+template<>
+FORCEINLINE TAABB<FPoint2D, 2>::TAABB()
+	: MinCorner(FVectorUtil::FarawayPoint2D)
+	, MaxCorner(-FVectorUtil::FarawayPoint2D)
+{
+}
+
+template<>
+FORCEINLINE TAABB<FPoint, 3>::TAABB()
+	: MinCorner(FVectorUtil::FarawayPoint3D)
+	, MaxCorner(-FVectorUtil::FarawayPoint3D)
+{
+}
+
+template<>
+FORCEINLINE void TAABB<FPoint2D, 2>::Empty()
+{
+	MinCorner = FVectorUtil::FarawayPoint2D;
+	MaxCorner = -FVectorUtil::FarawayPoint2D;
+}
+
+template<>
+FORCEINLINE void TAABB<FPoint, 3>::Empty()
+{
+	MinCorner = FVectorUtil::FarawayPoint3D;
+	MaxCorner = -FVectorUtil::FarawayPoint3D;
+}
+
+
 class FAABB : public TAABB<FPoint>
 {
 
@@ -233,16 +258,16 @@ public:
 	}
 };
 
-class CADKERNEL_API FAABB2D : public TAABB<FPoint2D>
+class CADKERNEL_API FAABB2D : public TAABB<FPoint2D, 2>
 {
 public:
 	FAABB2D()
-		: TAABB<FPoint2D>()
+		: TAABB<FPoint2D, 2>()
 	{
 	}
 
 	FAABB2D(const FPoint& InMinCorner, const FPoint& InMaxCorner)
-		: TAABB<FPoint2D>(InMinCorner, InMaxCorner)
+		: TAABB<FPoint2D, 2>(InMinCorner, InMaxCorner)
 	{
 	}
 
diff --git a/Engine/Source/Runtime/Datasmith/CADKernel/Base/Public/Math/Geometry.h b/Engine/Source/Runtime/Datasmith/CADKernel/Base/Public/Math/Geometry.h
--- a/Engine/Source/Runtime/Datasmith/CADKernel/Base/Public/Math/Geometry.h
+++ b/Engine/Source/Runtime/Datasmith/CADKernel/Base/Public/Math/Geometry.h
@@ -43,8 +43,8 @@ CADKERNEL_API inline FPoint2D ComputeCircumCircleCenter(const FPoint2D& InPoint0
 
 	// CenterU  = 1/D * (Cv.|B|.|B| - By.|C|.|C|) = 1/D * CBv ^ SquareNorms 
 	// CenterV  = 1/D * (Bu.|B|.|B| - Cu.|C|.|C|) = -1/D * SquareNorms ^ CBu
-	double SquareNormB = Segment_P0_P1.SquareLength();
-	double SquareNormC = Segment_P0_P2.SquareLength();
+	double SquareNormB = Segment_P0_P1.SquaredLength();
+	double SquareNormC = Segment_P0_P2.SquaredLength();
 	double CenterU = (SquareNormB * Segment_P0_P2.Y - SquareNormC * Segment_P0_P1.Y) / D;
 	double CenterV = (SquareNormC * Segment_P0_P1.X - SquareNormB * Segment_P0_P2.X) / D;
 
@@ -73,8 +73,8 @@ CADKERNEL_API inline FPoint ComputeCircumCircleCenter(const FPoint& Point0, cons
 		return FPoint::ZeroVector;
 	}
 
-	double SquareNormB = D2Point1.SquareLength();
-	double SquareNormC = D2Point2.SquareLength();
+	double SquareNormB = D2Point1.SquaredLength();
+	double SquareNormC = D2Point2.SquaredLength();
 
 	double CenterU = (SquareNormB * D2Point2.Y - SquareNormC * D2Point1.Y) / D;
 	double CenterV = (SquareNormC * D2Point1.X - SquareNormB * D2Point2.X) / D;
@@ -96,13 +96,13 @@ CADKERNEL_API inline FPoint2D ComputeCircumCircleCenterAndSquareRadius(const FPo
 		return FVector2d::ZeroVector;
 	}
 
-	double SquareNormB = Segment_P0_P1.SquareLength();
-	double SquareNormC = Segment_P0_P2.SquareLength();
+	double SquareNormB = Segment_P0_P1.SquaredLength();
+	double SquareNormC = Segment_P0_P2.SquaredLength();
 	double CenterU = (SquareNormB * Segment_P0_P2.Y - SquareNormC * Segment_P0_P1.Y) / D;
 	double CenterV = (SquareNormC * Segment_P0_P1.X - SquareNormB * Segment_P0_P2.X) / D;
 
 	FPoint2D Center(CenterU, CenterV);
-	OutSquareRadius = Center.SquareLength();
+	OutSquareRadius = Center.SquaredLength();
 
 	return Center + InPoint0;
 }
@@ -134,9 +134,9 @@ struct CADKERNEL_API TSegment
 		}
 	}
 
-	double SquareLength() const
+	double SquaredLength() const
 	{
-		return Point0.SquareDistance(Point1);
+		return FPoint2D::DistSquared(Point0, Point1);
 	}
 
 	PointType GetVector() const
@@ -182,8 +182,8 @@ struct CADKERNEL_API TTriangle
 	{
 		PointType Segment01 = Point1 - Point0;
 		PointType Segment02 = Point2 - Point0;
-		double SquareLength01 = Segment01.SquareLength();
-		double SquareLength02 = Segment02.SquareLength();
+		double SquareLength01 = Segment01.SquaredLength();
+		double SquareLength02 = Segment02.SquaredLength();
 		double Seg01Seg02 = Segment01 | Segment02;
 		double Det = SquareLength01 * SquareLength02 - FMath::Square(Seg01Seg02);
 
@@ -200,7 +200,7 @@ struct CADKERNEL_API TTriangle
 			}
 
 			PointType Segment12 = Point2 - Point1;
-			if (Segment12.SquareLength() > MaxLength)
+			if (Segment12.SquaredLength() > MaxLength)
 			{
 				SideIndex = Side12;
 			}
@@ -329,7 +329,7 @@ template<class PointType>
 inline double DistanceOfPointToSegment(const PointType& Point, const PointType& SegmentPoint1, const PointType& SegmentPoint2)
 {
 	double Coordinate;
-	return ProjectPointOnSegment(Point, SegmentPoint1, SegmentPoint2, Coordinate, /*bRestrictCoodinateToInside*/ true).Distance(Point);
+	return PointType::Distance(ProjectPointOnSegment(Point, SegmentPoint1, SegmentPoint2, Coordinate, /*bRestrictCoodinateToInside*/ true), Point);
 }
 
 /**
@@ -340,7 +340,7 @@ template<class PointType>
 inline double SquareDistanceOfPointToSegment(const PointType& Point, const PointType& SegmentPoint1, const PointType& SegmentPoint2)
 {
 	double Coordinate;
-	return ProjectPointOnSegment(Point, SegmentPoint1, SegmentPoint2, Coordinate, /*bRestrictCoodinateToInside*/ true).SquareDistance(Point);
+	return FPoint2D::DistSquared(ProjectPointOnSegment(Point, SegmentPoint1, SegmentPoint2, Coordinate, /*bRestrictCoodinateToInside*/ true), Point);
 }
 
 /**
@@ -365,9 +365,9 @@ inline PointType ProjectPointOnSegment(const PointType& Point, const PointType&
 {
 	PointType Segment = InSegmentB - InSegmentA;
 
-	double SquareLength = Segment | Segment;
+	double SquaredLength = Segment | Segment;
 
-	if (SquareLength <= 0.0)
+	if (SquaredLength <= 0.0)
 	{
 		OutCoordinate = 0.0;
 		return InSegmentA;
@@ -375,7 +375,7 @@ inline PointType ProjectPointOnSegment(const PointType& Point, const PointType&
 	else
 	{
 		PointType APoint = Point - InSegmentA;
-		OutCoordinate = (APoint | Segment) / SquareLength;
+		OutCoordinate = (APoint | Segment) / SquaredLength;
 
 		if (bRestrictCoodinateToInside)
 		{
@@ -406,16 +406,16 @@ inline double CoordinateOfProjectedPointOnSegment(const PointType& Point, const
 {
 	PointType Segment = InSegmentB - InSegmentA;
 
-	double SquareLength = Segment * Segment;
+	double SquaredLength = Segment * Segment;
 
-	if (SquareLength <= 0.0)
+	if (SquaredLength <= 0.0)
 	{
 		return 0.0;
 	}
 	else
 	{
 		PointType APoint = Point - InSegmentA;
-		double Coordinate = (APoint * Segment) / SquareLength;
+		double Coordinate = (APoint * Segment) / SquaredLength;
 
 		if (bRestrictCoodinateToInside)
 		{
@@ -526,8 +526,8 @@ CADKERNEL_API bool DoIntersectInside(const FSegment2D& SegmentAB, const FSegment
 
 inline bool AreParallel(const FSegment2D& SegmentAB, const FSegment2D& SegmentCD)
 {
-	const FPoint2D AB = SegmentAB.GetVector().Normalize();
-	const FPoint2D CD = SegmentCD.GetVector().Normalize();
+	const FPoint2D AB = SegmentAB.GetVector().GetSafeNormal();
+	const FPoint2D CD = SegmentCD.GetVector().GetSafeNormal();
 	const double ParallelCoef = AB ^ CD;
 	return (FMath::IsNearlyZero(ParallelCoef, DOUBLE_KINDA_SMALL_NUMBER));
 };
diff --git a/Engine/Source/Runtime/Datasmith/CADKernel/Base/Public/Math/Point.h b/Engine/Source/Runtime/Datasmith/CADKernel/Base/Public/Math/Point.h
--- a/Engine/Source/Runtime/Datasmith/CADKernel/Base/Public/Math/Point.h
+++ b/Engine/Source/Runtime/Datasmith/CADKernel/Base/Public/Math/Point.h
@@ -13,11 +13,6 @@ struct FPoint2D;
 struct FPoint : public FVector
 {
 public:
-
-	/** A zero point (0,0,0)*/
-	static const FPoint FarawayPoint;
-	static const int32 Dimension;
-
 	explicit FPoint(double InCoordX = 0., double InCoordY = 0., double InCoordZ = 0.)
 		: FVector(InCoordX, InCoordY, InCoordZ)
 	{
@@ -33,8 +28,6 @@ public:
 	{
 	}
 	
-	FPoint(const FPoint2D& Point);
-
 	FPoint(const FVector2d& InPoint2D)
 		: FVector(InPoint2D.X, InPoint2D.Y, 0.)
 	{
@@ -50,19 +43,15 @@ public:
 		return FVector3f((float)X, (float)Y, (float)Z);
 	}
 
-//	constexpr double& operator[](int32 Index)
-//	{
-//PRAGMA_DISABLE_DEPRECATION_WARNINGS
-//		return XYZ[Index];
-//PRAGMA_ENABLE_DEPRECATION_WARNINGS
-//	}
-//
-//	constexpr double operator[](int32 Index) const
-//	{
-//PRAGMA_DISABLE_DEPRECATION_WARNINGS
-//		return XYZ[Index];
-//PRAGMA_ENABLE_DEPRECATION_WARNINGS
-//	}
+	static double Distance(const FPoint& A, const FPoint& B)
+	{
+		return FVector::Distance(A, B);
+	}
+
+	static double DistSquared(const FPoint& A, const FPoint& B)
+	{
+		return FVector::DistSquared(A, B);
+	}
 
 	bool operator==(const FPoint& Point) const
 	{
@@ -370,24 +359,9 @@ public:
 	}
 };
 
-inline FPoint operator*(double Scale, const FPoint& Point)
-{
-	return FPoint(Point.X * Scale, Point.Y * Scale, Point.Z * Scale);
-}
-
-inline FPoint operator*(float Scale, const FPoint& Point)
-{
-	double ScaleD = static_cast<double>(Scale);
-	return FPoint(Point.X * ScaleD, Point.Y * ScaleD, Point.Z * ScaleD);
-}
-
 struct FPoint2D : public FVector2d
 {
 public:
-	/** A zero point (0,0,0)*/
-	static const FPoint2D FarawayPoint;
-	static const int32 Dimension;
-
 	FPoint2D()
 		: FVector2d(0., 0.)
 	{
@@ -413,63 +387,32 @@ public:
 	{
 	}
 
-	FPoint2D& operator+=(const FVector2d& Point)
+	static double Distance(const FPoint2D& A, const FPoint2D& B)
 	{
-		X += Point.X;
-		Y += Point.Y;
-		return *this;
+		return FVector2d::Distance(A, B);
 	}
 
-	FPoint2D Rotate(double Theta)
+	static double DistSquared(const FPoint2D& A, const FPoint2D& B)
 	{
-		return FPoint2D(X * cos(Theta) - Y * sin(Theta), X * sin(Theta) + Y * cos(Theta));
+		return FVector2d::DistSquared(A, B);
 	}
 
+	// TO REMOVE
+
 	FPoint2D Middle(const FPoint2D& Point) const
 	{
 		return FPoint(0.5 * (X + Point.X), 0.5 * (Y + Point.Y));
 	}
 
-	double Distance(const FPoint2D& Point) const
-	{
-		return FMath::Sqrt(SquareDistance(Point));
-	}
+	//double Distance(const FPoint2D& Point) const
+	//{
+	//	return FMath::Sqrt(SquareDistance(Point));
+	//}
 
-	double SquareDistance(const FPoint2D& Point) const
-	{
-		return FMath::Square(Point.X - X) + FMath::Square(Point.Y - Y);
-	}
-
-	double Length() const
-	{
-		return FMath::Sqrt(SquareLength());
-	}
-
-	double SquareLength() const
-	{
-		return FMath::Square(X) + FMath::Square(Y);
-	}
-
-	FPoint2D& Normalize(double& Norm)
-	{
-		Norm = Length();
-		if (FMath::IsNearlyZero(Norm))
-		{
-			X = Y = 0.0;
-		}
-		else
-		{
-			X /= Norm;
-			Y /= Norm;
-		}
-		return *this;
-	}
-
-	FPoint2D& Normalize()
-	{
-		double Norm = 0;
-		return Normalize(Norm);
-	}
+	//double SquareDistance(const FPoint2D& Point) const
+	//{
+	//	return FMath::Square(Point.X - X) + FMath::Square(Point.Y - Y);
+	//}
 
 	void Set(double NewU, double NewV = 0.)
 	{
@@ -516,18 +459,6 @@ public:
 		return Ar;
 	}
 
-	friend FPoint2D Abs(const FPoint2D& Point)
-	{
-		double PointU = FMath::Abs(Point.X);
-		double PointV = FMath::Abs(Point.Y);
-		return { PointU, PointV };
-	}
-
-	friend FPoint2D Max(const FPoint2D& PointA, const FPoint2D& PointB)
-	{
-		return { FMath::Max(PointA.X, PointB.X), FMath::Max(PointA.Y, PointB.Y) };
-	}
-
 	bool Equals(const FPoint2D& Other, double Tolerance = UE_DOUBLE_SMALL_NUMBER) const
 	{
 		return FMath::IsNearlyEqual(X, Other.X, Tolerance) && FMath::IsNearlyEqual(Y, Other.Y, Tolerance);
@@ -536,26 +467,9 @@ public:
 
 typedef FPoint2D FSurfacicTolerance;
 
-inline FPoint::FPoint(const FPoint2D& Point)
-	: FVector(Point, 0.)
-{
-}
-
-inline FPoint2D operator*(double Scale, const FPoint2D& Point)
-{
-	return FPoint2D(Point.X * Scale, Point.Y * Scale);
-}
-
-inline FPoint2D operator*(float Scale, const FPoint2D& Point)
-{
-	double ScaleD = static_cast<double>(Scale);
-	return FPoint2D(Point.X * ScaleD, Point.Y * ScaleD);
-}
-
 class FVectorUtil
 {
 public:
-	static CADKERNEL_API const FVector4d FarawayPoint4D;
 	static CADKERNEL_API const FVector FarawayPoint3D;
 	static CADKERNEL_API const FVector2d FarawayPoint2D;
 
diff --git a/Engine/Source/Runtime/Datasmith/CADKernel/Base/Public/Mesh/Structure/EdgeSegment.h b/Engine/Source/Runtime/Datasmith/CADKernel/Base/Public/Mesh/Structure/EdgeSegment.h
--- a/Engine/Source/Runtime/Datasmith/CADKernel/Base/Public/Mesh/Structure/EdgeSegment.h
+++ b/Engine/Source/Runtime/Datasmith/CADKernel/Base/Public/Mesh/Structure/EdgeSegment.h
@@ -76,7 +76,7 @@ public:
 		ClosedSegment = nullptr;
 
 		SquareDistanceToClosedSegment = HUGE_VAL;
-		Length = USSPoints[ELimit::Start].Distance(USSPoints[ELimit::End]);
+		Length = FPoint2D::Distance(USSPoints[ELimit::Start], USSPoints[ELimit::End]);
 
 		Id = ++LastId;
 		ChainIndex = Ident::Undefined;
