Replaced FPoint and FPoint2D by FVector and  FVector2d and remove unused FFPoint and FPointH classes - WIP - #4

#rnx

diff --git a/Engine/Plugins/Enterprise/DatasmithCADImporter/Source/CADInterfaces/Private/TechSoftBridge.cpp b/Engine/Plugins/Enterprise/DatasmithCADImporter/Source/CADInterfaces/Private/TechSoftBridge.cpp
--- a/Engine/Plugins/Enterprise/DatasmithCADImporter/Source/CADInterfaces/Private/TechSoftBridge.cpp
+++ b/Engine/Plugins/Enterprise/DatasmithCADImporter/Source/CADInterfaces/Private/TechSoftBridge.cpp
@@ -132,7 +132,7 @@ void FillPointArray(const int32 UCount, const int32 VCount, const A3DVector3dDat
 		int32 Index = Undex;
 		for (int32 Vndex = 0; Vndex < VCount; ++Vndex, Index += UCount, ++PointIndex)
 		{
-			OutPointsArray[Index].Set((double*)&(Points[PointIndex].m_dX));
+			OutPointsArray[Index].Set(Points[PointIndex].m_dX, Points[PointIndex].m_dY, Points[PointIndex].m_dZ);
 		}
 	}
 
diff --git a/Engine/Source/Runtime/Datasmith/CADKernel/Base/Private/Geo/Curves/SegmentCurve.h b/Engine/Source/Runtime/Datasmith/CADKernel/Base/Private/Geo/Curves/SegmentCurve.h
--- a/Engine/Source/Runtime/Datasmith/CADKernel/Base/Private/Geo/Curves/SegmentCurve.h
+++ b/Engine/Source/Runtime/Datasmith/CADKernel/Base/Private/Geo/Curves/SegmentCurve.h
@@ -69,8 +69,8 @@ public:
 
 	virtual void ExtendTo(const FPoint& DesiredPosition) override
 	{
-		double DistanceToStartPoint = DesiredPosition.SquareDistance(StartPoint);
-		double DistanceToEndPoint = DesiredPosition.SquareDistance(EndPoint);
+		double DistanceToStartPoint = FVector::DistSquared(DesiredPosition, StartPoint);
+		double DistanceToEndPoint = FVector::DistSquared(DesiredPosition, EndPoint);
 		if (DistanceToEndPoint < DistanceToStartPoint)
 		{
 			EndPoint = DesiredPosition;
diff --git a/Engine/Source/Runtime/Datasmith/CADKernel/Base/Private/Geo/Curves/SplineCurve.cpp b/Engine/Source/Runtime/Datasmith/CADKernel/Base/Private/Geo/Curves/SplineCurve.cpp
--- a/Engine/Source/Runtime/Datasmith/CADKernel/Base/Private/Geo/Curves/SplineCurve.cpp
+++ b/Engine/Source/Runtime/Datasmith/CADKernel/Base/Private/Geo/Curves/SplineCurve.cpp
@@ -170,7 +170,7 @@ FPoint FSplineCurve::GetTangentAtSplineInputKey(double Coordinate) const
 
 FPoint FSplineCurve::GetDirectionAtSplineInputKey(double Coordinate) const
 {
-	return Position.EvalDerivative(Coordinate, FPoint::ZeroVector).Normalize();
+	return Position.EvalDerivative(Coordinate, FPoint::ZeroVector).GetSafeNormal();
 }
 
 void FSplineCurve::EvaluatePoint(double Coordinate, FCurvePoint& OutPoint, int32 DerivativeOrder) const
diff --git a/Engine/Source/Runtime/Datasmith/CADKernel/Base/Private/Geo/Sampler/SamplerOnParam.h b/Engine/Source/Runtime/Datasmith/CADKernel/Base/Private/Geo/Sampler/SamplerOnParam.h
--- a/Engine/Source/Runtime/Datasmith/CADKernel/Base/Private/Geo/Sampler/SamplerOnParam.h
+++ b/Engine/Source/Runtime/Datasmith/CADKernel/Base/Private/Geo/Sampler/SamplerOnParam.h
@@ -77,10 +77,10 @@ protected:
 			double CurvilinearAbscise;
 			FPoint ProjectedPoint = ProjectPointOnSegment(PointM, APoint, BPoint, CurvilinearAbscise, /*bRestrictCoodinateToInside*/ true);
 
-			double Chord = ProjectedPoint.Distance(PointM);
+			double Chord = FVector::Distance(ProjectedPoint, PointM);
 			double AMCoordinate = this->CandidatePoints.Coordinates[Index] - ACoordinate;
 			FPoint MApproximate = APoint + (AMCoordinate / ABCoordinate) * ABVector;
-			double ParametrizationError = MApproximate.Distance(ProjectedPoint);
+			double ParametrizationError = FVector::Distance(MApproximate, ProjectedPoint);
 
 			if (Chord > MaxChord)
 			{
diff --git a/Engine/Source/Runtime/Datasmith/CADKernel/Base/Private/Math/BSpline.cpp b/Engine/Source/Runtime/Datasmith/CADKernel/Base/Private/Math/BSpline.cpp
--- a/Engine/Source/Runtime/Datasmith/CADKernel/Base/Private/Math/BSpline.cpp
+++ b/Engine/Source/Runtime/Datasmith/CADKernel/Base/Private/Math/BSpline.cpp
@@ -594,7 +594,7 @@ void EvaluatePoint(const FNURBSCurve& Nurbs, double Coordinate, FCurvePoint& Out
 	}
 	else
 	{
-		OutPoint.Point.Set(ValueH);
+		OutPoint.Point.Set(ValueH[0], ValueH[1], ValueH[2]);
 	}
 
 	if (DerivativeOrder > 0)
@@ -607,7 +607,7 @@ void EvaluatePoint(const FNURBSCurve& Nurbs, double Coordinate, FCurvePoint& Out
 		}
 		else
 		{
-			OutPoint.Gradient.Set(GradientH);
+			OutPoint.Gradient.Set(GradientH[0], GradientH[1], GradientH[2]);
 		}
 
 		if (DerivativeOrder > 1)
@@ -620,7 +620,7 @@ void EvaluatePoint(const FNURBSCurve& Nurbs, double Coordinate, FCurvePoint& Out
 			}
 			else
 			{
-				OutPoint.Laplacian.Set(LaplacianH);
+				OutPoint.Laplacian.Set(LaplacianH[0], LaplacianH[1], LaplacianH[2]);
 			}
 		}
 	}
@@ -639,37 +639,41 @@ void Evaluate2DPoint(const FNURBSCurve& Nurbs, double Coordinate, FCurvePoint2D&
 	int32 PoleDimension = Nurbs.GetDimension() + (Nurbs.IsRational() ? 1 : 0);
 	BSpline::Interpolate1DBSpline(Nurbs.GetDegree(), Nurbs.GetPoleCount(), Nurbs.GetNodalVector().GetData(), PoleDimension, Nurbs.GetHPoles().GetData(), 1, &Coordinate, DerivativeOrder, ValueH, GradientH, LaplacianH);
 
+	const double InvWeight = 1. / ValueH[2];
+
 	if (Nurbs.IsRational())
 	{
-		OutPoint.Point.Set(ValueH[0] / ValueH[2], ValueH[1] / ValueH[2]);
+		OutPoint.Point.Set(ValueH[0] * InvWeight, ValueH[1] * InvWeight);
 	}
 	else
 	{
-		OutPoint.Point.Set(ValueH);
+		OutPoint.Point.Set(ValueH[0], ValueH[1]);
 	}
 
 	if (DerivativeOrder > 0)
 	{
 		if (Nurbs.IsRational())
 		{
-			OutPoint.Gradient.X = (GradientH[0] - GradientH[2] * ValueH[0] / ValueH[2]) / ValueH[2];
-			OutPoint.Gradient.Y = (GradientH[1] - GradientH[2] * ValueH[1] / ValueH[2]) / ValueH[2];
+			OutPoint.Gradient.X = (GradientH[0] - GradientH[2] * ValueH[0] * InvWeight) * InvWeight;
+			OutPoint.Gradient.Y = (GradientH[1] - GradientH[2] * ValueH[1] * InvWeight) * InvWeight;
 		}
 		else
 		{
-			OutPoint.Gradient.Set(GradientH);
+			OutPoint.Gradient.Set(GradientH[0], GradientH[1]);
 		}
 
 		if (DerivativeOrder > 1)
 		{
 			if (Nurbs.IsRational())
 			{
-				OutPoint.Laplacian.X = (LaplacianH[0] * ValueH[2] - LaplacianH[2] * ValueH[0] - 2.0 * GradientH[2] * (GradientH[0] - GradientH[2] * ValueH[0] / ValueH[2])) / (ValueH[2] * ValueH[2]);
-				OutPoint.Laplacian.Y = (LaplacianH[1] * ValueH[2] - LaplacianH[2] * ValueH[1] - 2.0 * GradientH[2] * (GradientH[1] - GradientH[2] * ValueH[1] / ValueH[2])) / (ValueH[2] * ValueH[2]);
+				const double InvWeightSquared = InvWeight * InvWeight;
+
+				OutPoint.Laplacian.X = (LaplacianH[0] * ValueH[2] - LaplacianH[2] * ValueH[0] - 2.0 * GradientH[2] * (GradientH[0] - GradientH[2] * ValueH[0] * InvWeight)) * InvWeightSquared;
+				OutPoint.Laplacian.Y = (LaplacianH[1] * ValueH[2] - LaplacianH[2] * ValueH[1] - 2.0 * GradientH[2] * (GradientH[1] - GradientH[2] * ValueH[1] * InvWeight)) * InvWeightSquared;
 			}
 			else
 			{
-				OutPoint.Laplacian.Set(LaplacianH);
+				OutPoint.Laplacian.Set(LaplacianH[0], LaplacianH[1]);
 			}
 		}
 	}
@@ -749,19 +753,19 @@ void EvaluatePoint(const FNURBSSurface& Nurbs, const FPoint2D& InPoint2D, FSurfa
 	}
 	else
 	{
-		OutPoint3D.Point.Set(Point4);
+		OutPoint3D.Point.Set(Point4[0], Point4[1], Point4[2]);
 
 		if (InDerivativeOrder > 0)
 		{
-			OutPoint3D.GradientU.Set(GradientU);
-			OutPoint3D.GradientV.Set(GradientV);
-		}
+			OutPoint3D.GradientU.Set(GradientU[0], GradientU[1], GradientU[2]);
+			OutPoint3D.GradientV.Set(GradientV[0], GradientV[1], GradientV[2]);
 
-		if (InDerivativeOrder > 1)
-		{
-			OutPoint3D.LaplacianU.Set(LaplacianU);
-			OutPoint3D.LaplacianV.Set(LaplacianU);
-			OutPoint3D.LaplacianUV.Set(LaplacianUV);
+			if (InDerivativeOrder > 1)
+			{
+				OutPoint3D.LaplacianU.Set(LaplacianU[0], LaplacianU[1], LaplacianU[2]);
+				OutPoint3D.LaplacianV.Set(LaplacianV[0], LaplacianV[1], LaplacianV[2]);
+				OutPoint3D.LaplacianUV.Set(LaplacianUV[0], LaplacianUV[1], LaplacianUV[2]);
+			}
 		}
 	}
 }
diff --git a/Engine/Source/Runtime/Datasmith/CADKernel/Base/Private/Math/Geometry.cpp b/Engine/Source/Runtime/Datasmith/CADKernel/Base/Private/Math/Geometry.cpp
--- a/Engine/Source/Runtime/Datasmith/CADKernel/Base/Private/Math/Geometry.cpp
+++ b/Engine/Source/Runtime/Datasmith/CADKernel/Base/Private/Math/Geometry.cpp
@@ -91,19 +91,19 @@ bool ConfirmIntersectionWhenNearlyCoincident(const FIntersectionContext& Context
 
 double ComputeCurvature(const FPoint& Gradient, const FPoint& Laplacian)
 {
-	const FPoint GradientCopy = Gradient.Normalize();
-	const FPoint LaplacianCopy = Laplacian.Normalize();
+	const FPoint GradientCopy = Gradient.GetSafeNormal();
+	const FPoint LaplacianCopy = Laplacian.GetSafeNormal();
 	const FPoint Normal = GradientCopy ^ LaplacianCopy;
 	return (Normal.Length() * Laplacian.Length()) / (Gradient.Length() * Gradient.Length());
 }
 
 double ComputeCurvature(const FPoint& Normal, const FPoint& Gradient, const FPoint& Laplacian)
 {
-	const FPoint GradientCopy = Gradient.Normalize();
-	const FPoint LaplacianCopy = Laplacian.Normalize();
-	const FPoint NormalCopy = Normal.Normalize();
+	const FPoint GradientCopy = Gradient.GetSafeNormal();
+	const FPoint LaplacianCopy = Laplacian.GetSafeNormal();
+	const FPoint NormalCopy = Normal.GetSafeNormal();
 	const FPoint Coef = (LaplacianCopy ^ GradientCopy) ^ NormalCopy;
-	return (Coef.Length() * Laplacian.Length()) / Gradient.SquareLength();
+	return (Coef.Length() * Laplacian.Length()) / Gradient.SquaredLength();
 }
 
 void FindLoopIntersectionsWithIso(const EIso Iso, const double IsoParameter, const TArray<TArray<FPoint2D>>& Loops, TArray<double>& OutIntersections)
diff --git a/Engine/Source/Runtime/Datasmith/CADKernel/Base/Private/Math/Point.cpp b/Engine/Source/Runtime/Datasmith/CADKernel/Base/Private/Math/Point.cpp
--- a/Engine/Source/Runtime/Datasmith/CADKernel/Base/Private/Math/Point.cpp
+++ b/Engine/Source/Runtime/Datasmith/CADKernel/Base/Private/Math/Point.cpp
@@ -10,30 +10,30 @@ namespace UE::CADKernel
 const FVector FVectorUtil::FarawayPoint3D(HUGE_VALUE, HUGE_VALUE, HUGE_VALUE);
 const FVector2d FVectorUtil::FarawayPoint2D(HUGE_VALUE, HUGE_VALUE);
 
-double FPoint::SignedAngle(const FPoint & Other, const FPoint & Normal) const
-{
-	FPoint Vector1 = *this; 
-	FPoint Vector2 = Other; 
-	FPoint Vector3 = Normal; 
-
-	Vector1.Normalize();
-	Vector2.Normalize();
-	Vector3.Normalize();
-
-	double ScalarProduct = Vector1 * Vector2;
-
-	if (ScalarProduct >= 1 - DOUBLE_SMALL_NUMBER)
-	{
-		return 0.;
-	}
-
-	if (ScalarProduct <= -1 + DOUBLE_SMALL_NUMBER)
-	{
-		return DOUBLE_PI;
-	}
-
-	return MixedTripleProduct(Vector1, Vector2, Vector3) > 0 ? acos(ScalarProduct) : -acos(ScalarProduct);
-}
+//double FPoint::SignedAngle(const FPoint & Other, const FPoint & Normal) const
+//{
+//	FPoint Vector1 = *this; 
+//	FPoint Vector2 = Other; 
+//	FPoint Vector3 = Normal; 
+//
+//	Vector1.Normalize();
+//	Vector2.Normalize();
+//	Vector3.Normalize();
+//
+//	double ScalarProduct = Vector1 * Vector2;
+//
+//	if (ScalarProduct >= 1 - DOUBLE_SMALL_NUMBER)
+//	{
+//		return 0.;
+//	}
+//
+//	if (ScalarProduct <= -1 + DOUBLE_SMALL_NUMBER)
+//	{
+//		return DOUBLE_PI;
+//	}
+//
+//	return MixedTripleProduct(Vector1, Vector2, Vector3) > 0 ? acos(ScalarProduct) : -acos(ScalarProduct);
+//}
 
 
 }
\ No newline at end of file
diff --git a/Engine/Source/Runtime/Datasmith/CADKernel/Base/Private/Mesh/Criteria/CriteriaGrid.cpp b/Engine/Source/Runtime/Datasmith/CADKernel/Base/Private/Mesh/Criteria/CriteriaGrid.cpp
--- a/Engine/Source/Runtime/Datasmith/CADKernel/Base/Private/Mesh/Criteria/CriteriaGrid.cpp
+++ b/Engine/Source/Runtime/Datasmith/CADKernel/Base/Private/Mesh/Criteria/CriteriaGrid.cpp
@@ -60,17 +60,17 @@ void FCriteriaGrid::ComputeFaceMinMaxThicknessAlongIso()
 		{
 			const FPoint& StartPoint = GetPointAlong(0, VIndex);
 			const FPoint* PreviousPoint = &GetIntermediatePointAlong(0, VIndex);
-			double Length = StartPoint.Distance(*PreviousPoint);
+			double Length = FVector::Distance(StartPoint, *PreviousPoint);
 
 			for (int32 UIndex = 1; UIndex < IsoUCount - 1; ++UIndex)
 			{
 				const FPoint& Point = GetPointAlong(UIndex, VIndex);
 
 				PreviousPoint = &GetIntermediatePointAlong(UIndex, VIndex);
-				Length += Point.Distance(*PreviousPoint);
+				Length += FVector::Distance(Point, *PreviousPoint);
 			}
 			const FPoint& LastPoint = GetPointAlong(IsoUCount - 1, VIndex);
-			Length += LastPoint.Distance(*PreviousPoint);
+			Length += FVector::Distance(LastPoint, *PreviousPoint);
 
 			this->FaceMinMax[Iso].ExtendTo(Length);
 		}
diff --git a/Engine/Source/Runtime/Datasmith/CADKernel/Base/Private/Mesh/Criteria/SizeCriterion.cpp b/Engine/Source/Runtime/Datasmith/CADKernel/Base/Private/Mesh/Criteria/SizeCriterion.cpp
--- a/Engine/Source/Runtime/Datasmith/CADKernel/Base/Private/Mesh/Criteria/SizeCriterion.cpp
+++ b/Engine/Source/Runtime/Datasmith/CADKernel/Base/Private/Mesh/Criteria/SizeCriterion.cpp
@@ -61,7 +61,7 @@ void FSizeCriterion::ApplyOnParameters(const TArray<double>& Coordinates, const
 	{
 		const int32 PreviousIndex = Index - 1;
 		const double DeltaCoordinate = Coordinates[Index] - Coordinates[PreviousIndex];
-		const double ChordLength = Points[2 * PreviousIndex].Point.Distance(Points[2 * Index].Point);
+		const double ChordLength = FVector::Distance(Points[2 * PreviousIndex].Point, Points[2 * Index].Point);
 
 		const double NewDeltaU = ComputeSizeCriterionValue(DeltaCoordinate, ChordLength);
 		UpdateDeltaU(NewDeltaU, DeltaUMaxs[Index - 1], DeltaUMins[Index - 1]);
diff --git a/Engine/Source/Runtime/Datasmith/CADKernel/Base/Private/Mesh/Meshers/IsoTriangulator.h b/Engine/Source/Runtime/Datasmith/CADKernel/Base/Private/Mesh/Meshers/IsoTriangulator.h
--- a/Engine/Source/Runtime/Datasmith/CADKernel/Base/Private/Mesh/Meshers/IsoTriangulator.h
+++ b/Engine/Source/Runtime/Datasmith/CADKernel/Base/Private/Mesh/Meshers/IsoTriangulator.h
@@ -364,9 +364,9 @@ inline double CotangentCriteria(const FPoint& APoint, const FPoint& BPoint, cons
 	// the ratio between the scalar product PA.PB (=|PA| |PB| cos (A,P,B) )
 	// with the norm of the cross product |PA^PB| (=|PA| |PB| |sin(A,P,B)|)
 	// is compute. 
-	double ScalareProduct = PA * PB;
+	double ScalareProduct = PA | PB;
 	OutNormal = PA ^ PB;
-	double NormOFScalarProduct = sqrt(OutNormal * OutNormal);
+	double NormOFScalarProduct = sqrt(OutNormal | OutNormal);
 
 	// PPoint is aligned with (A,B)
 	if (NormOFScalarProduct < DOUBLE_SMALL_NUMBER)
diff --git a/Engine/Source/Runtime/Datasmith/CADKernel/Base/Private/Mesh/Meshers/IsoTriangulatorRemoveIntersection.cpp b/Engine/Source/Runtime/Datasmith/CADKernel/Base/Private/Mesh/Meshers/IsoTriangulatorRemoveIntersection.cpp
--- a/Engine/Source/Runtime/Datasmith/CADKernel/Base/Private/Mesh/Meshers/IsoTriangulatorRemoveIntersection.cpp
+++ b/Engine/Source/Runtime/Datasmith/CADKernel/Base/Private/Mesh/Meshers/IsoTriangulatorRemoveIntersection.cpp
@@ -63,7 +63,7 @@ bool FIsoTriangulator::RemovePickToOutside(const TArray<FLoopNode*>& NodesOfLoop
 	TmpNode = GetNext(EndNode);
 	FPoint2D SecondIntersection = IntersectingPoint(NextIntersection.Value, EndNode, TmpNode);
 
-	FPoint2D MiddlePoint = FirstIntersection.Middle(SecondIntersection);
+	FPoint2D MiddlePoint = (FirstIntersection + SecondIntersection) * 0.5;
 
 	TmpNode = GetNext(StartNode);
 	while (TmpNode && (TmpNode != EndNode) && !TmpNode->IsDelete())
@@ -89,7 +89,7 @@ bool FIsoTriangulator::RemovePickToOutside(const TArray<FLoopNode*>& NodesOfLoop
 	}
 
 	MoveDirection /= Length;
-	MoveDirection = MoveDirection.GetPerpendicularVector();
+	MoveDirection = FPoint2D(-MoveDirection.Y, MoveDirection.X);
 	MoveDirection *= GeometricTolerance;
 	MiddlePoint += MoveDirection;
 
@@ -439,7 +439,7 @@ bool FIsoTriangulator::TryToRemoveIntersectionOfTwoConsecutiveIntersectingSegmen
 
 	FPoint2D SegmentTangent = IntersectingSegment2D.Point1 - IntersectingSegment2D.Point0;
 	SegmentTangent.Normalize();
-	FPoint2D MoveDirection = SegmentTangent.GetPerpendicularVector();
+	FPoint2D MoveDirection = FPoint2D(-SegmentTangent.Y, SegmentTangent.X);
 
 	MoveDirection *= GeometricTolerance;
 
diff --git a/Engine/Source/Runtime/Datasmith/CADKernel/Base/Private/Mesh/Structure/Grid.cpp b/Engine/Source/Runtime/Datasmith/CADKernel/Base/Private/Mesh/Structure/Grid.cpp
--- a/Engine/Source/Runtime/Datasmith/CADKernel/Base/Private/Mesh/Structure/Grid.cpp
+++ b/Engine/Source/Runtime/Datasmith/CADKernel/Base/Private/Mesh/Structure/Grid.cpp
@@ -705,8 +705,8 @@ void SlightlyDisplacedPolyline(TArray<FPoint2D>& D2Points, const FSurfacicBounda
 	FPoint2D Normal;
 	for (int32 Index = 0; Index < D2Points.Num() - 1; ++Index)
 	{
-		FPoint2D Tangent = D2Points[Index + 1] - D2Points[Index];
-		Normal = Tangent.GetPerpendicularVector();
+		const FPoint2D Tangent = D2Points[Index + 1] - D2Points[Index];
+		Normal = FPoint2D(-Tangent.Y, Tangent.X);
 
 		ScaleAndSwap(Normal, D2Points[Index], Boundary);
 		D2Points[Index] += Normal;
diff --git a/Engine/Source/Runtime/Datasmith/CADKernel/Base/Private/Mesh/Structure/GridBase.cpp b/Engine/Source/Runtime/Datasmith/CADKernel/Base/Private/Mesh/Structure/GridBase.cpp
--- a/Engine/Source/Runtime/Datasmith/CADKernel/Base/Private/Mesh/Structure/GridBase.cpp
+++ b/Engine/Source/Runtime/Datasmith/CADKernel/Base/Private/Mesh/Structure/GridBase.cpp
@@ -87,7 +87,7 @@ bool FGridBase::ScaleGrid()
 		double Length = 0;
 		for (int32 IndexV = 1; IndexV < CuttingCount[EIso::IsoV]; ++IndexV)
 		{
-			Length += Points3D[Index].Distance(Points3D[Index + CuttingCount[EIso::IsoU]]);
+			Length += FVector::Distance(Points3D[Index], Points3D[Index + CuttingCount[EIso::IsoU]]);
 			Index += CuttingCount[EIso::IsoU];
 		}
 		LengthsV[IndexU] = Length;
@@ -100,7 +100,7 @@ bool FGridBase::ScaleGrid()
 		double Length = 0;
 		for (int32 IndexU = 1; IndexU < CuttingCount[EIso::IsoU]; ++IndexU)
 		{
-			Length += Points3D[Index].Distance(Points3D[Index + 1]);
+			Length += FVector::Distance(Points3D[Index], Points3D[Index + 1]);
 			Index++;
 		}
 		Index++;
@@ -157,7 +157,7 @@ bool FGridBase::ScaleGrid()
 			for (int32 IPointV = StartIndexUp; IPointV < CuttingCount[EIso::IsoV]; ++IPointV)
 			{
 				NumUV = IPointV * CuttingCount[EIso::IsoU] + IPointU;
-				Length += LastPoint.Distance(Points3D[NumUV]);
+				Length += FVector::Distance(LastPoint, Points3D[NumUV]);
 				Points2D[EGridSpace::Scaled][NumUV].Set(Points2D[EGridSpace::UniformScaled][NumUV].X, Length);
 				LastPoint = Points3D[NumUV];
 			}
@@ -167,7 +167,7 @@ bool FGridBase::ScaleGrid()
 			for (int32 IPointV = StartIndexDown; IPointV >= 0; --IPointV)
 			{
 				NumUV = IPointV * CuttingCount[EIso::IsoU] + IPointU;
-				Length -= LastPoint.Distance(Points3D[NumUV]);
+				Length -= FVector::Distance(LastPoint, Points3D[NumUV]);
 				Points2D[EGridSpace::Scaled][NumUV].Set(Points2D[EGridSpace::UniformScaled][NumUV].X, Length);
 				LastPoint = Points3D[NumUV];
 			}
@@ -195,7 +195,7 @@ bool FGridBase::ScaleGrid()
 			for (int32 IPointU = StartIndexUp; IPointU < CuttingCount[EIso::IsoU]; ++IPointU)
 			{
 				NumUV = IPointV * CuttingCount[EIso::IsoU] + IPointU;
-				Length += LastPoint.Distance(Points3D[NumUV]);
+				Length += FVector::Distance(LastPoint, Points3D[NumUV]);
 				Points2D[EGridSpace::Scaled][NumUV].Set(Length, Points2D[EGridSpace::UniformScaled][NumUV].Y);
 				LastPoint = Points3D[NumUV];
 			}
@@ -205,7 +205,7 @@ bool FGridBase::ScaleGrid()
 			for (int32 IPointU = StartIndexDown; IPointU >= 0; --IPointU)
 			{
 				NumUV = IPointV * CuttingCount[EIso::IsoU] + IPointU;
-				Length -= LastPoint.Distance(Points3D[NumUV]);
+				Length -= FVector::Distance(LastPoint, Points3D[NumUV]);
 				Points2D[EGridSpace::Scaled][NumUV].Set(Length, Points2D[EGridSpace::UniformScaled][NumUV].Y);
 				LastPoint = Points3D[NumUV];
 			}
diff --git a/Engine/Source/Runtime/Datasmith/CADKernel/Base/Private/Mesh/Structure/LoopCleaner.cpp b/Engine/Source/Runtime/Datasmith/CADKernel/Base/Private/Mesh/Structure/LoopCleaner.cpp
--- a/Engine/Source/Runtime/Datasmith/CADKernel/Base/Private/Mesh/Structure/LoopCleaner.cpp
+++ b/Engine/Source/Runtime/Datasmith/CADKernel/Base/Private/Mesh/Structure/LoopCleaner.cpp
@@ -477,7 +477,7 @@ bool FLoopCleaner::TryToRemoveIntersectionOfTwoConsecutiveIntersectingSegments(c
 
 	FPoint2D SegmentTangent = IntersectingSegment2D.Point1 - IntersectingSegment2D.Point0;
 	SegmentTangent.Normalize();
-	FPoint2D MoveDirection = SegmentTangent.GetPerpendicularVector();
+	FPoint2D MoveDirection = FPoint2D(-SegmentTangent.Y, SegmentTangent.X);
 
 	MoveDirection *= Tolerances.GeometricTolerance;
 
@@ -527,7 +527,7 @@ void FLoopCleaner::OffsetSegment(FIsoSegment& Segment, FSegment2D& Segment2D, FS
 {
 	FPoint2D SegmentTangent = IntersectingSegment2D.Point1 - IntersectingSegment2D.Point0;
 	SegmentTangent.Normalize();
-	FPoint2D MoveDirection = SegmentTangent.GetPerpendicularVector();
+	FPoint2D MoveDirection = FPoint2D(-SegmentTangent.Y, SegmentTangent.X);
 	MoveDirection *= Tolerances.GeometricTolerance;
 
 	FPoint2D NewPoint0 = Segment2D.Point0 + MoveDirection;
@@ -541,7 +541,7 @@ void FLoopCleaner::OffsetNode(FLoopNode& Node, FSegment2D& IntersectingSegment2D
 {
 	FPoint2D SegmentTangent = IntersectingSegment2D.Point1 - IntersectingSegment2D.Point0;
 	SegmentTangent.Normalize();
-	FPoint2D MoveDirection = SegmentTangent.GetPerpendicularVector();
+	FPoint2D MoveDirection = FPoint2D(-SegmentTangent.Y, SegmentTangent.X);
 	MoveDirection *= Tolerances.GeometricTolerance;
 
 	FPoint2D NewPoint = Node.Get2DPoint(EGridSpace::UniformScaled, Grid) + MoveDirection;
@@ -603,7 +603,7 @@ void FLoopCleaner::RemoveIntersectionByMovingOutsideSegmentNodeInside(const FIso
 
 	FPoint2D MoveDirection = IntersectingPoints[1] - IntersectingPoints[0];
 	MoveDirection.Normalize();
-	MoveDirection = MoveDirection.GetPerpendicularVector();
+	MoveDirection = FPoint2D(-MoveDirection.Y, MoveDirection.X);
 	MoveDirection *= bIsSameInnerLoop ? -Tolerances.GeometricTolerance : Tolerances.GeometricTolerance;
 
 	FPoint2D NewCoordinate = ProjectedPoint + MoveDirection;
@@ -825,7 +825,7 @@ bool FLoopCleaner::RemoveOutgoingLoop(const TPair<double, double>& Intersection,
 	TmpNode = GetNext(EndNode);
 	FPoint2D SecondIntersection = IntersectingPoint(NextIntersection.Value, EndNode, TmpNode);
 
-	FPoint2D MiddlePoint = FirstIntersection.Middle(SecondIntersection);
+	FPoint2D MiddlePoint = (FirstIntersection + SecondIntersection) * 0.5;
 
 	TmpNode = GetNext(StartNode);
 	while (TmpNode && (TmpNode != EndNode) && !TmpNode->IsDelete())
@@ -851,7 +851,7 @@ bool FLoopCleaner::RemoveOutgoingLoop(const TPair<double, double>& Intersection,
 	}
 
 	MoveDirection /= Length;
-	MoveDirection = MoveDirection.GetPerpendicularVector();
+	MoveDirection = FPoint2D(-MoveDirection.Y, MoveDirection.X);
 	MoveDirection *= Tolerances.GeometricTolerance;
 	MiddlePoint += MoveDirection;
 
@@ -936,7 +936,7 @@ bool FLoopCleaner::RemoveIntersectionsOfSubLoop(int32 IntersectionIndex, int32 I
 
 		FPoint2D MoveDirection = Point1 - Point0;
 		MoveDirection.Normalize();
-		MoveDirection = MoveDirection.GetPerpendicularVector();
+		MoveDirection = FPoint2D(-MoveDirection.Y, MoveDirection.X);
 		MoveDirection *= Tolerances.GeometricTolerance;
 
 		FPoint2D NewCoordinate = Point1 + MoveDirection;
diff --git a/Engine/Source/Runtime/Datasmith/CADKernel/Base/Private/Mesh/Structure/ModelMeshAnalyzer.cpp b/Engine/Source/Runtime/Datasmith/CADKernel/Base/Private/Mesh/Structure/ModelMeshAnalyzer.cpp
--- a/Engine/Source/Runtime/Datasmith/CADKernel/Base/Private/Mesh/Structure/ModelMeshAnalyzer.cpp
+++ b/Engine/Source/Runtime/Datasmith/CADKernel/Base/Private/Mesh/Structure/ModelMeshAnalyzer.cpp
@@ -197,7 +197,7 @@ double FModelMeshAnalyzer::ComputeMaxAngle() const
 		FPoint Normal0 = Edge.Triangles[0]->ComputeNormal(NodeCoordinates);
 		FPoint Normal1 = Edge.Triangles[1]->ComputeNormal(NodeCoordinates);
 
-		double Angle = Normal0.ComputeAngle(Normal1);
+		double Angle = FVectorUtil::ComputeAngle(Normal0, Normal1);
 		if (Angle > MaxAngle)
 		{
 			MaxAngle = Angle;
diff --git a/Engine/Source/Runtime/Datasmith/CADKernel/Base/Private/Mesh/Structure/ThinZone2DFinder.cpp b/Engine/Source/Runtime/Datasmith/CADKernel/Base/Private/Mesh/Structure/ThinZone2DFinder.cpp
--- a/Engine/Source/Runtime/Datasmith/CADKernel/Base/Private/Mesh/Structure/ThinZone2DFinder.cpp
+++ b/Engine/Source/Runtime/Datasmith/CADKernel/Base/Private/Mesh/Structure/ThinZone2DFinder.cpp
@@ -97,7 +97,7 @@ void FThinZone2DFinder::FindCloseSegments()
 		//     MaxSpacePlusLength = MaxSpace + MaxSegmentLength;
 
 		const double SegmentMinMin = Segment->GetAxeMin() - MaxSpacePlusLength;
-		const double SegmentMiddleMax = SegmentMiddle.DiagonalAxisCoordinate() + MaxSpacePlusLength;
+		const double SegmentMiddleMax = FVectorUtil::DiagonalAxisCoordinate(SegmentMiddle) + MaxSpacePlusLength;
 		double MinSquareThickness = HUGE_VALUE;
 
 #ifdef DEBUG_FIND_CLOSE_SEGMENTS_
diff --git a/Engine/Source/Runtime/Datasmith/CADKernel/Base/Private/Topo/FaceAnalyzer.cpp b/Engine/Source/Runtime/Datasmith/CADKernel/Base/Private/Topo/FaceAnalyzer.cpp
--- a/Engine/Source/Runtime/Datasmith/CADKernel/Base/Private/Topo/FaceAnalyzer.cpp
+++ b/Engine/Source/Runtime/Datasmith/CADKernel/Base/Private/Topo/FaceAnalyzer.cpp
@@ -69,7 +69,7 @@ void FFaceAnalyzer::FindClosedSegments(Topo::FThinFaceContext& Context)
 			}
 
 			const FPoint& CenterPointCandidate = Candidate->GetMiddle();
-			const double SquareDistanceToCenter = CenterPointCandidate.SquareDistance(Point);
+			const double SquareDistanceToCenter = FVector::DistSquared(CenterPointCandidate, Point);
 
 			// if the distance of Point with the middle of candidate segment is biggest than MaxSpace, then the projection of the point cannot be smaller than the tolerance, 
 			if (SquareDistanceToCenter > MinSquareDistance)
@@ -83,7 +83,7 @@ void FFaceAnalyzer::FindClosedSegments(Topo::FThinFaceContext& Context)
 			double Coordinate;
 			FPoint Projection = ProjectPointOnSegment(Point, FirstPointCandidate, SecondPointCandidate, Coordinate, true);
 
-			double SquareDistance = Point.SquareDistance(Projection);
+			double SquareDistance = FVector::DistSquared(Point, Projection);
 			if (SquareDistance > MinSquareDistance)
 			{
 				continue;
diff --git a/Engine/Source/Runtime/Datasmith/CADKernel/Base/Private/Topo/FaceAnalyzer.h b/Engine/Source/Runtime/Datasmith/CADKernel/Base/Private/Topo/FaceAnalyzer.h
--- a/Engine/Source/Runtime/Datasmith/CADKernel/Base/Private/Topo/FaceAnalyzer.h
+++ b/Engine/Source/Runtime/Datasmith/CADKernel/Base/Private/Topo/FaceAnalyzer.h
@@ -64,15 +64,15 @@ public:
 		Coordinates[ELimit::End] = InEndU;
 		Points[ELimit::Start] = InStartPoint;
 		Points[ELimit::End] = InEndPoint;
-		Middle = Points[ELimit::Start].Middle(Points[ELimit::End]);
+		Middle = (Points[ELimit::Start] + Points[ELimit::End]) * 0.5;
 		Vector = Points[ELimit::End] - Points[ELimit::Start];
 
 		ClosedSegment = nullptr;
 
 		SquareDistanceToClosedSegment = -1.;
-		Length = Points[ELimit::Start].Distance(Points[ELimit::End]);
+		Length = FVector::Distance(Points[ELimit::Start], Points[ELimit::End]);
 
-		MiddleAxis = Middle.DiagonalAxisCoordinate();
+		MiddleAxis = FVectorUtil::DiagonalAxisCoordinate(Middle);
 
 #ifdef CADKERNEL_DEV
 		Id = LastId++;
@@ -121,7 +121,7 @@ public:
 	 */
 	double ComputeCosAngleOf(const FEdgeSegment* Segment)
 	{
-		return GetVector().ComputeCosinus(Segment->GetVector());
+		return FVectorUtil::ComputeCosinus(GetVector() , Segment->GetVector());
 	}
 
 	FEdgeSegment* GetClosedSegment() const
diff --git a/Engine/Source/Runtime/Datasmith/CADKernel/Base/Private/Topo/TopologicalEdge.cpp b/Engine/Source/Runtime/Datasmith/CADKernel/Base/Private/Topo/TopologicalEdge.cpp
--- a/Engine/Source/Runtime/Datasmith/CADKernel/Base/Private/Topo/TopologicalEdge.cpp
+++ b/Engine/Source/Runtime/Datasmith/CADKernel/Base/Private/Topo/TopologicalEdge.cpp
@@ -84,7 +84,7 @@ bool FTopologicalEdge::CheckVertices()
 	double ToleranceGeo = GetTolerance3D();
 	TFunction<bool(TSharedPtr<FTopologicalVertex>, FPoint)> CheckExtremityGap = [&](TSharedPtr<FTopologicalVertex> Vertex, FPoint Point)
 	{
-		double GapToVertex = Vertex->GetCoordinates().Distance(Point);
+		double GapToVertex = FVector::Distance(Vertex->GetCoordinates(), Point);
 		return (GapToVertex < ToleranceGeo);
 	};
 
@@ -205,8 +205,8 @@ bool FTopologicalEdge::IsTangentAtExtremitiesWith(const FTopologicalEdge& Edge)
 		GetTangentsAtExtremities(StartTangent, EndTangent, Orientation);
 		Edge.GetTangentsAtExtremities(EdgeStartTangent, EdgeEndTangent, EOrientation::Front);
 
-		double StartAngle = EdgeStartTangent.ComputeCosinus(StartTangent);
-		double EndAngle = EdgeEndTangent.ComputeCosinus(EndTangent);
+		double StartAngle = FVectorUtil::ComputeCosinus(EdgeStartTangent, StartTangent);
+		double EndAngle = FVectorUtil::ComputeCosinus(EdgeEndTangent, EndTangent);
 
 		if (StartAngle >= UE_DOUBLE_HALF_SQRT_3 && EndAngle >= UE_DOUBLE_HALF_SQRT_3)
 		{
@@ -284,11 +284,11 @@ void FTopologicalEdge::LinkIfCoincident(FTopologicalEdge& Twin, double EdgeLengt
 	const FPoint& Edge2Vertex2 = Twin.GetEndVertex()->GetBarycenter();
 
 	// Define the orientation
-	const double SquareDistanceE1V1_E2V1 = GetStartVertex()->IsLinkedTo(Twin.GetStartVertex()) ? 0. : Edge1Vertex1.SquareDistance(Edge2Vertex1);
-	const double SquareDistanceE1V2_E2V2 = GetEndVertex()->IsLinkedTo(Twin.GetEndVertex()) ? 0. : Edge1Vertex2.SquareDistance(Edge2Vertex2);
+	const double SquareDistanceE1V1_E2V1 = GetStartVertex()->IsLinkedTo(Twin.GetStartVertex()) ? 0. : FVector::DistSquared(Edge1Vertex1, Edge2Vertex1);
+	const double SquareDistanceE1V2_E2V2 = GetEndVertex()->IsLinkedTo(Twin.GetEndVertex()) ? 0. : FVector::DistSquared(Edge1Vertex2, Edge2Vertex2);
 
-	const double SquareDistanceE1V1_E2V2 = GetStartVertex()->IsLinkedTo(Twin.GetEndVertex()) ? 0. : Edge1Vertex1.SquareDistance(Edge2Vertex2);
-	const double SquareDistanceE1V2_E2V1 = GetEndVertex()->IsLinkedTo(Twin.GetStartVertex()) ? 0. : Edge1Vertex2.SquareDistance(Edge2Vertex1);
+	const double SquareDistanceE1V1_E2V2 = GetStartVertex()->IsLinkedTo(Twin.GetEndVertex()) ? 0. : FVector::DistSquared(Edge1Vertex1, Edge2Vertex2);
+	const double SquareDistanceE1V2_E2V1 = GetEndVertex()->IsLinkedTo(Twin.GetStartVertex()) ? 0. : FVector::DistSquared(Edge1Vertex2, Edge2Vertex1);
 
 	const double SquareDistanceSameOrientation = SquareDistanceE1V1_E2V1 + SquareDistanceE1V2_E2V2;
 	const double SquareDistanceReverseOrientation = SquareDistanceE1V1_E2V2 + SquareDistanceE1V2_E2V1;
@@ -816,8 +816,8 @@ bool FTopologicalEdge::IsSameDirection(const FTopologicalEdge& Edge) const
 		FPoint EndTangent;
 		GetTangentsAtExtremities(StartTangent, EndTangent, true);
 
-		double StartAngle = StartTangent.ComputeCosinus(EdgeStartTangent);
-		double EndAngle = EndTangent.ComputeCosinus(EdgeEndTangent);
+		double StartAngle = FVectorUtil::ComputeCosinus(StartTangent, EdgeStartTangent);
+		double EndAngle = FVectorUtil::ComputeCosinus(EndTangent, EdgeEndTangent);
 
 		if (StartAngle >= 0 && EndAngle >= 0)
 		{
@@ -1228,7 +1228,7 @@ bool FTopologicalEdge::ExtendTo(bool bIsStartExtremity, const FPoint2D& NewExtre
 	{
 		FPoint ProjectedPoint;
 		double UProjectedPoint = ProjectPoint(NewVertex->GetCoordinates(), ProjectedPoint);
-		if (ProjectedPoint.Distance(NewVertex->GetCoordinates()) > GetTolerance3D())
+		if (FVector::Distance(ProjectedPoint, NewVertex->GetCoordinates()) > GetTolerance3D())
 		{
 			return false;
 		}
@@ -1266,7 +1266,7 @@ void FTopologicalEdge::ComputeEdge2DProperties(FEdge2DProperties& EdgeCharacteri
 		{
 			Slope = 4. - Slope;
 		}
-		EdgeCharacteristics.Add(Slope, Polyline3D[Index].Distance(Polyline3D[Index + 1]));
+		EdgeCharacteristics.Add(Slope, FVector::Distance(Polyline3D[Index], Polyline3D[Index + 1]));
 	}
 }
 
diff --git a/Engine/Source/Runtime/Datasmith/CADKernel/Base/Private/Topo/TopologicalFace.cpp b/Engine/Source/Runtime/Datasmith/CADKernel/Base/Private/Topo/TopologicalFace.cpp
--- a/Engine/Source/Runtime/Datasmith/CADKernel/Base/Private/Topo/TopologicalFace.cpp
+++ b/Engine/Source/Runtime/Datasmith/CADKernel/Base/Private/Topo/TopologicalFace.cpp
@@ -180,20 +180,20 @@ void FTopologicalFace::ApplyNaturalLoops(const FSurfacicBoundary& Boundaries)
 	FPoint EndPoint;
 
 	// Build 4 bounding edges of the surface
-	StartPoint.Set(Boundaries[EIso::IsoU].Min, Boundaries[EIso::IsoV].Min);
-	EndPoint.Set(Boundaries[EIso::IsoU].Max, Boundaries[EIso::IsoV].Min);
+	StartPoint.Set(Boundaries[EIso::IsoU].Min, Boundaries[EIso::IsoV].Min, 0.);
+	EndPoint.Set(Boundaries[EIso::IsoU].Max, Boundaries[EIso::IsoV].Min, 0.);
 	BuildEdge(StartPoint, EndPoint);
 
-	StartPoint.Set(Boundaries[EIso::IsoU].Max, Boundaries[EIso::IsoV].Min);
-	EndPoint.Set(Boundaries[EIso::IsoU].Max, Boundaries[EIso::IsoV].Max);
+	StartPoint.Set(Boundaries[EIso::IsoU].Max, Boundaries[EIso::IsoV].Min, 0.);
+	EndPoint.Set(Boundaries[EIso::IsoU].Max, Boundaries[EIso::IsoV].Max, 0.);
 	BuildEdge(StartPoint, EndPoint);
 
-	StartPoint.Set(Boundaries[EIso::IsoU].Max, Boundaries[EIso::IsoV].Max);
-	EndPoint.Set(Boundaries[EIso::IsoU].Min, Boundaries[EIso::IsoV].Max);
+	StartPoint.Set(Boundaries[EIso::IsoU].Max, Boundaries[EIso::IsoV].Max, 0.);
+	EndPoint.Set(Boundaries[EIso::IsoU].Min, Boundaries[EIso::IsoV].Max, 0.);
 	BuildEdge(StartPoint, EndPoint);
 
-	StartPoint.Set(Boundaries[EIso::IsoU].Min, Boundaries[EIso::IsoV].Max);
-	EndPoint.Set(Boundaries[EIso::IsoU].Min, Boundaries[EIso::IsoV].Min);
+	StartPoint.Set(Boundaries[EIso::IsoU].Min, Boundaries[EIso::IsoV].Max, 0.);
+	EndPoint.Set(Boundaries[EIso::IsoU].Min, Boundaries[EIso::IsoV].Min, 0.);
 	BuildEdge(StartPoint, EndPoint);
 	if (Edges.IsEmpty())
 	{
diff --git a/Engine/Source/Runtime/Datasmith/CADKernel/Base/Private/Topo/TopologicalLoop.cpp b/Engine/Source/Runtime/Datasmith/CADKernel/Base/Private/Topo/TopologicalLoop.cpp
--- a/Engine/Source/Runtime/Datasmith/CADKernel/Base/Private/Topo/TopologicalLoop.cpp
+++ b/Engine/Source/Runtime/Datasmith/CADKernel/Base/Private/Topo/TopologicalLoop.cpp
@@ -904,7 +904,7 @@ void FTopologicalLoop::EnsureLogicalClosing(const double Tolerance3D)
 		OrientedEdge.Entity->GetExtremities(EdgeExtremities);
 
 		const int32 ExtremityIndex = OrientedEdge.Direction == EOrientation::Front ? 0 : 1;
-		const double SquareGap3D = EdgeExtremities[ExtremityIndex].Point.SquareDistance(PreviousExtremity.Point);
+		const double SquareGap3D = FVector::DistSquared(EdgeExtremities[ExtremityIndex].Point, PreviousExtremity.Point);
 
 		TSharedRef<FTopologicalVertex> PreviousEdgeEndVertex = PreviousEdge.Direction == EOrientation::Front ? PreviousEdge.Entity->GetEndVertex() : PreviousEdge.Entity->GetStartVertex();
 		TSharedRef<FTopologicalVertex> EdgeStartVertex = OrientedEdge.Direction == EOrientation::Front ? OrientedEdge.Entity->GetStartVertex() : OrientedEdge.Entity->GetEndVertex();
@@ -918,7 +918,7 @@ void FTopologicalLoop::EnsureLogicalClosing(const double Tolerance3D)
 
 			FPoint Gap = EdgeExtremities[ExtremityIndex].Point - PreviousExtremity.Point;
 
-			double CosAngle = Gap.ComputeCosinus(EdgeTangent);
+			double CosAngle = FVectorUtil::ComputeCosinus(Gap, EdgeTangent);
 			constexpr double CosTangentLimit = 0.9; // ~25 deg : 25 deg is not big angle between the extremity curve tangent and the missing segment to close the gap.  
 			if (CosAngle > CosTangentLimit) 
 			{
@@ -926,7 +926,7 @@ void FTopologicalLoop::EnsureLogicalClosing(const double Tolerance3D)
 			}
 			else
 			{
-				CosAngle = Gap.ComputeCosinus(PreviousTangent);
+				CosAngle = FVectorUtil::ComputeCosinus(Gap, PreviousTangent);
 				if (CosAngle < -CosTangentLimit) // ~25 deg
 				{
 					PreviousEdge.Entity->ExtendTo(PreviousEdge.Direction == EOrientation::Back, EdgeExtremities[ExtremityIndex].Point2D, EdgeStartVertex);
@@ -973,7 +973,7 @@ void FTopologicalLoop::EnsureLogicalClosing(const double Tolerance3D)
 			const FPoint2D PreviousTangent = PreviousEdge.Entity->GetTangent2DAt(*PreviousEdgeEndVertex);
 			const FPoint2D EdgeTangent = OrientedEdge.Entity->GetTangent2DAt(*EdgeStartVertex);
 
-			double CosAngle = Gap2D.ComputeCosinus(EdgeTangent);
+			double CosAngle = FVectorUtil::ComputeCosinus(Gap2D, EdgeTangent);
 			constexpr double CosFlatTangent = 0.98;  // ~10 deg: In 2D, the distortion due to the parametric space, degenerated area and else imposes to be more careful before extending a curve.
 
 			//   ----------------------*         ------------**--------*
@@ -987,7 +987,7 @@ void FTopologicalLoop::EnsureLogicalClosing(const double Tolerance3D)
 			}
 			else
 			{
-				CosAngle = Gap2D.ComputeCosinus(PreviousTangent);
+				CosAngle = FVectorUtil::ComputeCosinus(Gap2D, PreviousTangent);
 				if (FMath::Abs(CosAngle) > CosFlatTangent)
 				{
 					PreviousEdge.Entity->ExtendTo(PreviousEdge.Direction == EOrientation::Back, EdgeExtremities[ExtremityIndex].Point2D, EdgeStartVertex);
@@ -1057,7 +1057,7 @@ void FTopologicalLoop::FindBreaks(TArray<TSharedPtr<FTopologicalVertex>>& OutBre
 		// if both edge are border, the rupture is not evaluate. 
 		if (bIsSurface || bPreviousIsSurface)
 		{
-			double CosAngle = StartTangentEdge.ComputeCosinus(EndTangentPreviousEdge);
+			double CosAngle = FVectorUtil::ComputeCosinus(StartTangentEdge, EndTangentPreviousEdge);
 
 #ifdef FIND_BREAKS
 			{
diff --git a/Engine/Source/Runtime/Datasmith/CADKernel/Base/Private/Topo/Topomaker.cpp b/Engine/Source/Runtime/Datasmith/CADKernel/Base/Private/Topo/Topomaker.cpp
--- a/Engine/Source/Runtime/Datasmith/CADKernel/Base/Private/Topo/Topomaker.cpp
+++ b/Engine/Source/Runtime/Datasmith/CADKernel/Base/Private/Topo/Topomaker.cpp
@@ -71,7 +71,7 @@ void MergeCoincidentVertices(TArray<FTopologicalVertex*>& VerticesToMerge, doubl
 
 	for (FTopologicalVertex* Vertex : VerticesToMerge)
 	{
-		VerticesWeight.Add(Vertex->GetCoordinates().DiagonalAxisCoordinate());
+		VerticesWeight.Add(FVectorUtil::DiagonalAxisCoordinate(Vertex->GetCoordinates()));
 	}
 
 	for (int32 Index = 0; Index < VerticesToMerge.Num(); ++Index)
@@ -170,13 +170,13 @@ FTopologicalVertex* SplitAndLink(FTopologicalVertex& StartVertex, FTopologicalEd
 	FPoint ProjectedPoint;
 	double UProjectedPoint = EdgeToSplit.ProjectPoint(VertexToLink->GetBarycenter(), ProjectedPoint);
 
-	double SquareDistanceToProjectedPoint = ProjectedPoint.SquareDistance(VertexToLink->GetBarycenter());
+	double SquareDistanceToProjectedPoint = FVector::DistSquared(ProjectedPoint, VertexToLink->GetBarycenter());
 	if (SquareDistanceToProjectedPoint > SquareSewTolerance)
 	{
 		return nullptr;
 	}
 
-	double SquareDistanceToOtherPoint = ProjectedPoint.SquareDistance(EndVertex->GetBarycenter());
+	double SquareDistanceToOtherPoint = FVector::DistSquared(ProjectedPoint, EndVertex->GetBarycenter());
 	if (FMath::IsNearlyEqual(UProjectedPoint, EdgeToSplit.GetBoundary().GetMax(), UE_DOUBLE_SMALL_NUMBER) || SquareDistanceToOtherPoint < SquareMinEdgeLength)
 	{
 		// the new point is closed to the extremity, a degenerated edge will be created, so the edges are joined
@@ -185,7 +185,7 @@ FTopologicalVertex* SplitAndLink(FTopologicalVertex& StartVertex, FTopologicalEd
 		return &VertexToLink->GetLinkActiveEntity().Get();
 	}
 
-	SquareDistanceToProjectedPoint = ProjectedPoint.SquareDistance(StartVertex.GetBarycenter());
+	SquareDistanceToProjectedPoint = FVector::DistSquared(ProjectedPoint, StartVertex.GetBarycenter());
 	if (FMath::IsNearlyEqual(UProjectedPoint, EdgeToSplit.GetBoundary().GetMin(), UE_DOUBLE_SMALL_NUMBER) || SquareDistanceToProjectedPoint < SquareMinEdgeLength)
 	{
 		VertexToLink->Link(StartVertex);
@@ -338,7 +338,7 @@ FTopologicalVertex* StitchParallelEdgesFrom(FTopologicalVertex* Vertex, double S
 			FPoint StartTangentEdge = Edge->GetTangentAt(*Vertex);
 			FPoint StartTangentOtherEdge = SecondEdge->GetTangentAt(*Vertex);
 
-			double CosAngle = StartTangentEdge.ComputeCosinus(StartTangentOtherEdge);
+			double CosAngle = FVectorUtil::ComputeCosinus(StartTangentEdge, StartTangentOtherEdge);
 			if (CosAngle < UE_DOUBLE_HALF_SQRT_3) // cos(30 deg)
 			{
 				continue;
@@ -665,7 +665,7 @@ void FTopomaker::Sew()
 	// basic case: merge pair of edges connected together at both extremities i.e. pair of edges with same extremity active vertices.
 	TopomakerTools::MergeCoincidentEdges(BorderVertices, EdgeLengthTolerance);
 
-	MergeUnconnectedSuccessiveEdges();
+	//MergeUnconnectedSuccessiveEdges();
 
 	// advance case: two partially coincident edges connected at one extremity i.e. coincident along the shortest edge
 	BorderVertices.Reset();
@@ -831,7 +831,7 @@ void FTopomaker::MergeUnconnectedSuccessiveEdges()
 						FPoint StartTangentEdge = ConnectedEdges[0]->GetTangentAt(*EndVertex);
 						FPoint StartTangentOtherEdge = ConnectedEdges[1]->GetTangentAt(*EndVertex);
 
-						double CosAngle = StartTangentEdge.ComputeCosinus(StartTangentOtherEdge);
+						double CosAngle = FVectorUtil::ComputeCosinus(StartTangentEdge, StartTangentOtherEdge);
 						if (CosAngle < -UE_DOUBLE_HALF_SQRT_3) // Cos(30 deg)
 						{
 							bEdgeIsNotTheLast = true;
diff --git a/Engine/Source/Runtime/Datasmith/CADKernel/Base/Public/Geo/Sampling/PolylineTools.h b/Engine/Source/Runtime/Datasmith/CADKernel/Base/Public/Geo/Sampling/PolylineTools.h
--- a/Engine/Source/Runtime/Datasmith/CADKernel/Base/Public/Geo/Sampling/PolylineTools.h
+++ b/Engine/Source/Runtime/Datasmith/CADKernel/Base/Public/Geo/Sampling/PolylineTools.h
@@ -609,7 +609,7 @@ public:
 			for (int32 Index = StartIndex; Index < EndIndex; ++Index)
 			{
 				PointType ProjectPoint = ProjectPointOnSegment(InPointToProject, PolylinePoints[Index], PolylinePoints[Index + 1], ParamU, true);
-				double SquareDistance = ProjectPoint.SquareDistance(InPointToProject);
+				double SquareDistance = FVector::DistSquared(ProjectPoint, InPointToProject);
 				if (SquareDistance > MinDistance + SquareTol)
 				{
 					break;
diff --git a/Engine/Source/Runtime/Datasmith/CADKernel/Base/Public/Geo/Sampling/SurfacicPolyline.h b/Engine/Source/Runtime/Datasmith/CADKernel/Base/Public/Geo/Sampling/SurfacicPolyline.h
--- a/Engine/Source/Runtime/Datasmith/CADKernel/Base/Public/Geo/Sampling/SurfacicPolyline.h
+++ b/Engine/Source/Runtime/Datasmith/CADKernel/Base/Public/Geo/Sampling/SurfacicPolyline.h
@@ -107,7 +107,7 @@ public:
 
 	FSurfacicTolerance ComputeTolerance(const double Tolerance3D, const FSurfacicTolerance& MinToleranceIso, const int32 Index) const
 	{
-		double Distance3D = Points3D[Index].Distance(Points3D[Index + 1]);
+		double Distance3D = FVector::Distance(Points3D[Index], Points3D[Index + 1]);
 		if (FMath::IsNearlyZero(Distance3D, (double)DOUBLE_SMALL_NUMBER))
 		{
 			return FVectorUtil::FarawayPoint2D;
@@ -121,7 +121,7 @@ public:
 
 	double ComputeLinearToleranceAt(const double Tolerance3D, const double MinLinearTolerance, const int32 Index) const
 	{
-		double Distance3D = Points3D[Index].Distance(Points3D[Index + 1]);
+		double Distance3D = FVector::Distance(Points3D[Index], Points3D[Index + 1]);
 		if (FMath::IsNearlyZero(Distance3D, (double)DOUBLE_SMALL_NUMBER))
 		{
 			return (Coordinates.Last() - Coordinates[0]) / 10.;
diff --git a/Engine/Source/Runtime/Datasmith/CADKernel/Base/Public/Geo/Surfaces/Surface.h b/Engine/Source/Runtime/Datasmith/CADKernel/Base/Public/Geo/Surfaces/Surface.h
--- a/Engine/Source/Runtime/Datasmith/CADKernel/Base/Public/Geo/Surfaces/Surface.h
+++ b/Engine/Source/Runtime/Datasmith/CADKernel/Base/Public/Geo/Surfaces/Surface.h
@@ -286,7 +286,7 @@ public:
 	static double ComputeScaleAlongAxis(const FPoint& InAxis, const FMatrixH& InMatrix, const FPoint& InOrigin)
 	{
 		FPoint Point = InMatrix.Multiply(InAxis);
-		double Length = InOrigin.Distance(Point);
+		double Length = FVector::Distance(InOrigin, Point);
 		return Length;
 	};
 
diff --git a/Engine/Source/Runtime/Datasmith/CADKernel/Base/Public/Math/Geometry.h b/Engine/Source/Runtime/Datasmith/CADKernel/Base/Public/Math/Geometry.h
--- a/Engine/Source/Runtime/Datasmith/CADKernel/Base/Public/Math/Geometry.h
+++ b/Engine/Source/Runtime/Datasmith/CADKernel/Base/Public/Math/Geometry.h
@@ -53,6 +53,7 @@ CADKERNEL_API inline FPoint2D ComputeCircumCircleCenter(const FPoint2D& InPoint0
 
 CADKERNEL_API inline FPoint ComputeCircumCircleCenter(const FPoint& Point0, const FPoint& Point1, const FPoint& Point2)
 {
+	FPoint Test = Point1 - Point0;
 	FPoint AxisZ = (Point1 - Point0) ^ (Point2 - Point0);
 	AxisZ.Normalize();
 
@@ -306,13 +307,13 @@ inline FPoint ProjectPointOnPlane(const FPoint& Point, const FPoint& Origin, con
 	Normal.Normalize();
 
 	FPoint OP = Point - Origin;
-	OutDistance = OP * Normal;
+	OutDistance = OP | Normal;
 
 	return Point - (Normal * OutDistance);
 }
 
 /**
- * @return InSegmentA + (InSegmentB - InSegmentA) * InCoordinate
+ * @return InSegmentA + (InSegmentB - InSegmentA) | InCoordinate
  */
 template<class PointType>
 inline PointType PointOnSegment(const PointType& InSegmentA, const PointType& InSegmentB, double InCoordinate)
@@ -406,7 +407,7 @@ inline double CoordinateOfProjectedPointOnSegment(const PointType& Point, const
 {
 	PointType Segment = InSegmentB - InSegmentA;
 
-	double SquaredLength = Segment * Segment;
+	double SquaredLength = Segment | Segment;
 
 	if (SquaredLength <= 0.0)
 	{
@@ -415,7 +416,7 @@ inline double CoordinateOfProjectedPointOnSegment(const PointType& Point, const
 	else
 	{
 		PointType APoint = Point - InSegmentA;
-		double Coordinate = (APoint * Segment) / SquaredLength;
+		double Coordinate = (APoint | Segment) / SquaredLength;
 
 		if (bRestrictCoodinateToInside)
 		{
@@ -459,7 +460,7 @@ inline FPoint2D FindIntersectionOfSegments2D(const FSegment2D& SegmentAB, const
 			if (DCoordinate >= -DOUBLE_KINDA_SMALL_NUMBER && DCoordinate <= 1 + DOUBLE_KINDA_SMALL_NUMBER)
 			{
 				OutABIntersectionCoordinate = (DCoordinate + CCoordinate) * 0.5;
-				return SegmentCD[0].Middle(SegmentCD[1]);
+				return (SegmentCD[0] + SegmentCD[1]) * 0.5;
 			}
 
 			CCoordinate = FMath::Clamp(CCoordinate, 0., 1.);
@@ -475,7 +476,7 @@ inline FPoint2D FindIntersectionOfSegments2D(const FSegment2D& SegmentAB, const
 		else
 		{
 			OutABIntersectionCoordinate = 0.5;
-			return SegmentAB[0].Middle(SegmentAB[1]);
+			return (SegmentAB[0] + SegmentAB[1]) * 0.5;
 		}
 	}
 
@@ -532,16 +533,5 @@ inline bool AreParallel(const FSegment2D& SegmentAB, const FSegment2D& SegmentCD
 	return (FMath::IsNearlyZero(ParallelCoef, DOUBLE_KINDA_SMALL_NUMBER));
 };
 
-inline double ComputeCosinus(FVector Vector, FVector OtherVector)
-{
-	Vector.Normalize();
-	OtherVector.Normalize();
-
-	double Cosinus = Vector.X * OtherVector.X + Vector.Y * OtherVector.Y + Vector.Z * OtherVector.Z;
-
-	return FMath::Max(-1.0, FMath::Min(Cosinus, 1.0));
-}
-
-
 } // namespace UE::CADKernel
 
diff --git a/Engine/Source/Runtime/Datasmith/CADKernel/Base/Public/Math/Point.h b/Engine/Source/Runtime/Datasmith/CADKernel/Base/Public/Math/Point.h
--- a/Engine/Source/Runtime/Datasmith/CADKernel/Base/Public/Math/Point.h
+++ b/Engine/Source/Runtime/Datasmith/CADKernel/Base/Public/Math/Point.h
@@ -10,6 +10,7 @@ namespace UE::CADKernel
 
 struct FPoint2D;
 
+#if 1
 struct FPoint : public FVector
 {
 public:
@@ -53,15 +54,13 @@ public:
 		return FVector::DistSquared(A, B);
 	}
 
-	bool operator==(const FPoint& Point) const
+	friend FArchive& operator<<(FArchive& Ar, FPoint& Point)
 	{
-		return FMath::IsNearlyEqual(X, Point.X) && FMath::IsNearlyEqual(Y, Point.Y) && FMath::IsNearlyEqual(Z, Point.Z);
+		Ar.Serialize(&Point, sizeof(FPoint));
+		return Ar;
 	}
 
-	bool operator!=(const FPoint& Point) const
-	{
-		return !(*this == Point);
-	};
+	// TO REMOVE
 
 	FPoint operator+(const FPoint& Point) const
 	{
@@ -88,11 +87,6 @@ public:
 		return FPoint(Y * Point.Z - Z * Point.Y, Z * Point.X - X * Point.Z, X * Point.Y - Y * Point.X);
 	}
 
-	double operator*(const FPoint& Point) const
-	{
-		return X * Point.X + Y * Point.Y + Z * Point.Z;
-	}
-
 	double operator|(const FPoint& Point) const
 	{
 		return X * Point.X + Y * Point.Y + Z * Point.Z;
@@ -165,13 +159,6 @@ public:
 		return *this;
 	}
 
-	void NonUniformScale(const FPoint& Scale)
-	{
-		X *= Scale.X;
-		Y *= Scale.Y;
-		Z *= Scale.Z;
-	}
-
 	FPoint& operator/=(double InvScale)
 	{
 		if (FMath::IsNearlyZero(InvScale))
@@ -185,178 +172,106 @@ public:
 		return *this;
 	}
 
-	bool operator<(const FPoint& Other) const
-	{
-		int32 Compare = RealCompare(X, Other.X);
-		if (Compare < 0)
-		{
-			return true;
-		}
-		if (Compare > 0)
-		{
-			return false;
-		}
-
-		Compare = RealCompare(Y, Other.Y);
-		if (Compare < 0)
-		{
-			return true;
-		}
-		if (Compare > 0)
-		{
-			return false;
-		}
-
-		Compare = RealCompare(Z, Other.Z);
-		return (Compare < 0);
-	}
-
-	void SetMin(const FPoint& Point)
-	{
-		X = FMath::Min(X, Point.X);
-		Y = FMath::Min(Y, Point.Y);
-		Z = FMath::Min(Z, Point.Z);
-	}
-
-	void SetMax(const FPoint& Point)
-	{
-		X = FMath::Max(X, Point.X);
-		Y = FMath::Max(Y, Point.Y);
-		Z = FMath::Max(Z, Point.Z);
-	}
-
-	void Set(double InCoordX = 0., double InCoordY = 0., double InCoordZ = 0.)
-	{
-		X = InCoordX;
-		Y = InCoordY;
-		Z = InCoordZ;
-	}
-
-	void Set(double* InCoordinates)
-	{
-		X = InCoordinates[0];
-		Y = InCoordinates[1];
-		Z = InCoordinates[2];
-	}
-
-	double Length() const
-	{
-		return FMath::Sqrt((X * X) + (Y * Y) + (Z * Z));
-	}
-
-	double SquareLength() const
-	{
-		return (X * X) + (Y * Y) + (Z * Z);
-	}
-
-	FPoint& Normalize()
-	{
-		double Norm = Length();
-		if (FMath::IsNearlyZero(Norm))
-		{
-			*this = FPoint(1.,1.,1.).Normalize();
-		}
-		else
-		{
-			*this /= Norm;
-		}
-		return *this;
-	}
-
-	FPoint Normalize() const 
-	{
-		double Norm = Length();
-		if (FMath::IsNearlyZero(Norm))
-		{
-			return FPoint(1.,1.,1.).Normalize();
-		}
-		else
-		{
-			return *this / Norm;
-		}
-	}
-
-	FPoint Middle(const FPoint& Point) const
-	{
-		return FPoint(0.5 * (X + Point.X), 0.5 * (Y + Point.Y), 0.5 * (Z + Point.Z));
-	}
-
-	double Distance(const FPoint& Point) const
-	{
-		return FMath::Sqrt(FMath::Square(Point.X - X) + FMath::Square(Point.Y - Y) + FMath::Square(Point.Z - Z));
-	}
-
-	double SquareDistance(const FPoint& Point) const
-	{
-		return FMath::Square(Point.X - X) + FMath::Square(Point.Y - Y) + FMath::Square(Point.Z - Z);
-	}
-
-	double ComputeCosinus(const FPoint& OtherVector) const
-	{
-		FPoint ThisNormalized = *this;
-		FPoint OtherNormalized = OtherVector;
-
-		ThisNormalized.Normalize();
-		OtherNormalized.Normalize();
-
-		double Cosinus = ThisNormalized * OtherNormalized;
-
-		return FMath::Max(-1.0, FMath::Min(Cosinus, 1.0));
-	}
+	//bool operator<(const FPoint& Other) const
+	//{
+	//	int32 Compare = RealCompare(X, Other.X);
+	//	if (Compare < 0)
+	//	{
+	//		return true;
+	//	}
+	//	if (Compare > 0)
+	//	{
+	//		return false;
+	//	}
+
+	//	Compare = RealCompare(Y, Other.Y);
+	//	if (Compare < 0)
+	//	{
+	//		return true;
+	//	}
+	//	if (Compare > 0)
+	//	{
+	//		return false;
+	//	}
+
+	//	Compare = RealCompare(Z, Other.Z);
+	//	return (Compare < 0);
+	//}
 
-	double ComputeSinus(const FPoint& OtherVector) const
-	{
-		FPoint ThisNormalized = *this;
-		FPoint OtherNormalized = OtherVector;
+	//void SetMin(const FPoint& Point)
+	//{
+	//	X = FMath::Min(X, Point.X);
+	//	Y = FMath::Min(Y, Point.Y);
+	//	Z = FMath::Min(Z, Point.Z);
+	//}
 
-		ThisNormalized.Normalize();
-		OtherNormalized.Normalize();
+	//void SetMax(const FPoint& Point)
+	//{
+	//	X = FMath::Max(X, Point.X);
+	//	Y = FMath::Max(Y, Point.Y);
+	//	Z = FMath::Max(Z, Point.Z);
+	//}
 
-		FPoint SinusPoint = ThisNormalized ^ OtherNormalized;
-		double Sinus = SinusPoint.Length();
-		return FMath::Min(Sinus, 1.0);
-	}
+	//void Set(double InCoordX = 0., double InCoordY = 0., double InCoordZ = 0.)
+	//{
+	//	X = InCoordX;
+	//	Y = InCoordY;
+	//	Z = InCoordZ;
+	//}
 
-	double ComputeAngle(const FPoint& OtherVector) const
-	{
-		double CosAngle = ComputeCosinus(OtherVector);
-		return acos(CosAngle);
-	}
+	//void Set(double* InCoordinates)
+	//{
+	//	X = InCoordinates[0];
+	//	Y = InCoordinates[1];
+	//	Z = InCoordinates[2];
+	//}
 
-	double SignedAngle(const FPoint& Vector2, const FPoint& Normal) const;
+	//double SquareLength() const
+	//{
+	//	return (X * X) + (Y * Y) + (Z * Z);
+	//}
 
-	/**
-	 * Return the projection of the point on the diagonal axis (of vector (1,1,1))
-	 * i.e. return X + Y + Z
-	 */
-	double DiagonalAxisCoordinate() const
-	{
-		return X + Y + Z;
-	}
+	//FPoint& Normalize()
+	//{
+	//	double Norm = Length();
+	//	if (FMath::IsNearlyZero(Norm))
+	//	{
+	//		*this = FPoint(1.,1.,1.).Normalize();
+	//	}
+	//	else
+	//	{
+	//		*this /= Norm;
+	//	}
+	//	return *this;
+	//}
 
-	static double MixedTripleProduct(const FPoint& VectorA, const FPoint& VectorB, const FPoint& VectorC)
-	{
-		return VectorA * (VectorB ^ VectorC);
-	}
+	//FPoint Normalize() const 
+	//{
+	//	double Norm = Length();
+	//	if (FMath::IsNearlyZero(Norm))
+	//	{
+	//		return FPoint(1.,1.,1.).Normalize();
+	//	}
+	//	else
+	//	{
+	//		return *this / Norm;
+	//	}
+	//}
 
-	friend FArchive& operator<<(FArchive& Ar, FPoint& Point)
-	{
-		Ar.Serialize(&Point, sizeof(FPoint));
-		return Ar;
-	}
+	//FPoint Middle(const FPoint& Point) const
+	//{
+	//	return FPoint(0.5 * (X + Point.X), 0.5 * (Y + Point.Y), 0.5 * (Z + Point.Z));
+	//}
 
-	friend FPoint Abs(const FPoint& Point)
-	{
-		double PointX = FMath::Abs(Point.X);
-		double PointY = FMath::Abs(Point.Y);
-		double PointZ = FMath::Abs(Point.Z);
-		return FPoint(PointX, PointY, PointZ);
-	}
+	//double Distance(const FPoint& Point) const
+	//{
+	//	return FMath::Sqrt(FMath::Square(Point.X - X) + FMath::Square(Point.Y - Y) + FMath::Square(Point.Z - Z));
+	//}
 
-	bool Equals(const FPoint& Other, double Tolerance = UE_DOUBLE_SMALL_NUMBER) const
-	{
-		return FMath::IsNearlyEqual(X, Other.X, Tolerance) && FMath::IsNearlyEqual(Y, Other.Y, Tolerance) && FMath::IsNearlyEqual(Z, Other.Z, Tolerance);
-	}
+	//double SquareDistance(const FPoint& Point) const
+	//{
+	//	return FMath::Square(Point.X - X) + FMath::Square(Point.Y - Y) + FMath::Square(Point.Z - Z);
+	//}
 };
 
 struct FPoint2D : public FVector2d
@@ -397,73 +312,18 @@ public:
 		return FVector2d::DistSquared(A, B);
 	}
 
-	// TO REMOVE
-
-	FPoint2D Middle(const FPoint2D& Point) const
-	{
-		return FPoint(0.5 * (X + Point.X), 0.5 * (Y + Point.Y));
-	}
-
-	//double Distance(const FPoint2D& Point) const
-	//{
-	//	return FMath::Sqrt(SquareDistance(Point));
-	//}
-
-	//double SquareDistance(const FPoint2D& Point) const
-	//{
-	//	return FMath::Square(Point.X - X) + FMath::Square(Point.Y - Y);
-	//}
-
-	void Set(double NewU, double NewV = 0.)
-	{
-		X = NewU;
-		Y = NewV;
-	}
-
-	void Set(double* NewCoordinate)
-	{
-		X = NewCoordinate[0];
-		Y = NewCoordinate[1];
-	}
-
-	double ComputeCosinus(const FPoint2D& OtherVector) const
-	{
-		FPoint2D ThisNormalized = *this;
-		FPoint2D OtherNormalized = OtherVector;
-
-		ThisNormalized.Normalize();
-		OtherNormalized.Normalize();
-
-		double Cosinus = ThisNormalized | OtherNormalized;
-
-		return FMath::Max(-1.0, FMath::Min(Cosinus, 1.0));
-	}
-
-	FPoint2D GetPerpendicularVector()
-	{
-		return FPoint2D(-Y, X);
-	}
-
-	/**
-	 * Return the projection of the point on the diagonal axis (of vector (1,1))
-	 * i.e. return X + Y
-	 */
-	double DiagonalAxisCoordinate() const
-	{
-		return X + Y;
-	}
-
 	friend FArchive& operator<<(FArchive& Ar, FPoint2D& Point)
 	{
 		Ar.Serialize(&Point, sizeof(FPoint2D));
 		return Ar;
 	}
 
-	bool Equals(const FPoint2D& Other, double Tolerance = UE_DOUBLE_SMALL_NUMBER) const
+	void Set(double NewX, double NewY)
 	{
-		return FMath::IsNearlyEqual(X, Other.X, Tolerance) && FMath::IsNearlyEqual(Y, Other.Y, Tolerance);
+		FVector2d::Set(NewX, NewY);
 	}
 };
+#endif
 
 typedef FPoint2D FSurfacicTolerance;
 
diff --git a/Engine/Source/Runtime/Datasmith/CADKernel/Base/Public/Mesh/Criteria/Criterion.h b/Engine/Source/Runtime/Datasmith/CADKernel/Base/Public/Mesh/Criteria/Criterion.h
--- a/Engine/Source/Runtime/Datasmith/CADKernel/Base/Public/Mesh/Criteria/Criterion.h
+++ b/Engine/Source/Runtime/Datasmith/CADKernel/Base/Public/Mesh/Criteria/Criterion.h
@@ -78,13 +78,13 @@ public:
 		FPoint MiddleVec = PointMiddle - PointPoint;
 
 		double Sag = 0.0;
-		double NormSqrVec = ChordVec * ChordVec;
+		double NormSqrVec = ChordVec | ChordVec;
 		Length = sqrt(NormSqrVec);
 
 		if (NormSqrVec > DOUBLE_SMALL_NUMBER)
 		{
 			FPoint VecSag = ChordVec ^ MiddleVec;
-			double NormSqrSag = VecSag * VecSag;
+			double NormSqrSag = VecSag | VecSag;
 			Sag = NormSqrSag / NormSqrVec;
 
 			if (Sag < SMALL_NUMBER_SQUARE)
diff --git a/Engine/Source/Runtime/Datasmith/CADKernel/Base/Public/Topo/TopologicalLoop.h b/Engine/Source/Runtime/Datasmith/CADKernel/Base/Public/Topo/TopologicalLoop.h
--- a/Engine/Source/Runtime/Datasmith/CADKernel/Base/Public/Topo/TopologicalLoop.h
+++ b/Engine/Source/Runtime/Datasmith/CADKernel/Base/Public/Topo/TopologicalLoop.h
@@ -35,6 +35,25 @@ public:
 		return Entity == Edge.Entity;
 	}
 
+	TSharedRef<FTopologicalVertex> GetStartVertex()
+	{
+		return Direction == EOrientation::Front ? Entity->GetStartVertex() : Entity->GetEndVertex();
+	}
+
+	TSharedRef<FTopologicalVertex> GetEndVertex()
+	{
+		return Direction == EOrientation::Front ? Entity->GetEndVertex() : Entity->GetStartVertex();
+	}
+
+	double GetStartCoordinate()
+	{
+		return Direction == EOrientation::Front ? Entity->GetBoundary().GetMin() : Entity->GetBoundary().GetMax();
+	}
+
+	double GetEndCoordinate()
+	{
+		return Direction == EOrientation::Front ? Entity->GetBoundary().GetMax() : Entity->GetBoundary().GetMin();
+	}
 };
 
 class CADKERNEL_API FTopologicalLoop : public FTopologicalEntity
diff --git a/Engine/Source/Runtime/Datasmith/CADKernel/Base/Public/Topo/TopologicalVertex.h b/Engine/Source/Runtime/Datasmith/CADKernel/Base/Public/Topo/TopologicalVertex.h
--- a/Engine/Source/Runtime/Datasmith/CADKernel/Base/Public/Topo/TopologicalVertex.h
+++ b/Engine/Source/Runtime/Datasmith/CADKernel/Base/Public/Topo/TopologicalVertex.h
@@ -174,22 +174,22 @@ public:
 
 	double Distance(const FTopologicalVertex& OtherVertex) const
 	{
-		return Coordinates.Distance(OtherVertex.Coordinates);
+		return FVector::Distance(Coordinates, OtherVertex.Coordinates);
 	}
 
 	double SquareDistance(const FTopologicalVertex& OtherVertex) const
 	{
-		return Coordinates.SquareDistance(OtherVertex.Coordinates);
+		return FVector::DistSquared(Coordinates, OtherVertex.Coordinates);
 	}
 
 	double SquareDistanceBetweenBarycenters(const FTopologicalVertex& OtherVertex) const
 	{
-		return GetLink()->GetBarycenter().SquareDistance(OtherVertex.GetLink()->GetBarycenter());
+		return FVector::DistSquared(GetLink()->GetBarycenter(), OtherVertex.GetLink()->GetBarycenter());
 	}
 
 	double SquareDistance(const FPoint& Point) const
 	{
-		return Coordinates.SquareDistance(Point);
+		return FVector::DistSquared(Coordinates, Point);
 	}
 
 	CADKERNEL_API FVertexMesh& GetOrCreateMesh(FModelMesh& MeshModel);
