Replaced FPoint and FPoint2D by FVector and  FVector2d and remove unused FFPoint and FPointH classes - Final

diff --git a/Engine/Source/Runtime/Datasmith/CADKernel/Base/Private/Geo/Surfaces/OffsetSurface.cpp b/Engine/Source/Runtime/Datasmith/CADKernel/Base/Private/Geo/Surfaces/OffsetSurface.cpp
--- a/Engine/Source/Runtime/Datasmith/CADKernel/Base/Private/Geo/Surfaces/OffsetSurface.cpp
+++ b/Engine/Source/Runtime/Datasmith/CADKernel/Base/Private/Geo/Surfaces/OffsetSurface.cpp
@@ -41,7 +41,7 @@ FInfoEntity& FOffsetSurface::GetInfo(FInfoEntity& Info) const
 }
 #endif
 
-void FOffsetSurface::EvaluatePoint(const FPoint2D& InSurfacicCoordinate, FSurfacicPoint& OutPoint3D, int32 InDerivativeOrder) const
+void FOffsetSurface::EvaluatePoint(const FVector2d& InSurfacicCoordinate, FSurfacicPoint& OutPoint3D, int32 InDerivativeOrder) const
 {
 	if (InDerivativeOrder == 0)
 	{
@@ -49,12 +49,12 @@ void FOffsetSurface::EvaluatePoint(const FPoint2D& InSurfacicCoordinate, FSurfac
 	}
 	BaseSurface->EvaluatePoint(InSurfacicCoordinate, OutPoint3D, InDerivativeOrder);
 
-	FPoint Normal = OutPoint3D.GradientU ^ OutPoint3D.GradientV;
+	FVector Normal = OutPoint3D.GradientU ^ OutPoint3D.GradientV;
 	Normal.Normalize();
 	OutPoint3D.Point += Normal*Offset;
 }
 
-void FOffsetSurface::EvaluatePoints(const TArray<FPoint2D>& InSurfacicCoordinates, TArray<FSurfacicPoint>& OutPoint3D, int32 InDerivativeOrder) const
+void FOffsetSurface::EvaluatePoints(const TArray<FVector2d>& InSurfacicCoordinates, TArray<FSurfacicPoint>& OutPoint3D, int32 InDerivativeOrder) const
 {
 	if (InDerivativeOrder == 0)
 	{
@@ -65,7 +65,7 @@ void FOffsetSurface::EvaluatePoints(const TArray<FPoint2D>& InSurfacicCoordinate
 
 	for (FSurfacicPoint& Point : OutPoint3D)
 	{
-		FPoint Normal = Point.GradientU ^ Point.GradientV;
+		FVector Normal = Point.GradientU ^ Point.GradientV;
 		Normal.Normalize();
 		Point.Point += Normal * Offset;
 	}
@@ -78,7 +78,7 @@ void FOffsetSurface::EvaluatePointGrid(const FCoordinateGrid& Coords, FSurfacicS
 
 	for (int32 Index = 0; Index < OutPoints.Count(); ++Index)
 	{
-		OutPoints.Points3D[Index] += OutPoints.Normals[Index] * Offset;
+		OutPoints.Points3D[Index] += FVector(OutPoints.Normals[Index]) * Offset;
 	}
 }
 
diff --git a/Engine/Source/Runtime/Datasmith/CADKernel/Base/Private/Geo/Surfaces/OffsetSurface.h b/Engine/Source/Runtime/Datasmith/CADKernel/Base/Private/Geo/Surfaces/OffsetSurface.h
--- a/Engine/Source/Runtime/Datasmith/CADKernel/Base/Private/Geo/Surfaces/OffsetSurface.h
+++ b/Engine/Source/Runtime/Datasmith/CADKernel/Base/Private/Geo/Surfaces/OffsetSurface.h
@@ -83,8 +83,8 @@ public:
 	/**
 	 * The Gradient and Laplacian are approximated by their value on the carrier surface.
 	 */
-	virtual void EvaluatePoint(const FPoint2D& InSurfacicCoordinate, FSurfacicPoint& OutPoint3D, int32 InDerivativeOrder = 0) const override;
-	virtual void EvaluatePoints(const TArray<FPoint2D>& InSurfacicCoordinates, TArray<FSurfacicPoint>& OutPoint3D, int32 InDerivativeOrder = 0) const override;
+	virtual void EvaluatePoint(const FVector2d& InSurfacicCoordinate, FSurfacicPoint& OutPoint3D, int32 InDerivativeOrder = 0) const override;
+	virtual void EvaluatePoints(const TArray<FVector2d>& InSurfacicCoordinates, TArray<FSurfacicPoint>& OutPoint3D, int32 InDerivativeOrder = 0) const override;
 	virtual void EvaluatePointGrid(const FCoordinateGrid& Coordinates, FSurfacicSampling& OutPoints, bool bComputeNormals = false) const override;
 
 	virtual void Presample(const FSurfacicBoundary& InBoundaries, FCoordinateGrid& OutCoordinates) override;
diff --git a/Engine/Source/Runtime/Datasmith/CADKernel/Base/Private/Geo/Surfaces/PlaneSurface.cpp b/Engine/Source/Runtime/Datasmith/CADKernel/Base/Private/Geo/Surfaces/PlaneSurface.cpp
--- a/Engine/Source/Runtime/Datasmith/CADKernel/Base/Private/Geo/Surfaces/PlaneSurface.cpp
+++ b/Engine/Source/Runtime/Datasmith/CADKernel/Base/Private/Geo/Surfaces/PlaneSurface.cpp
@@ -21,7 +21,7 @@ FPlaneSurface::FPlaneSurface(const double InToleranceGeometric, const FMatrixH&
 	ComputeMinToleranceIso();
 }
 
-FPlaneSurface::FPlaneSurface(const double InToleranceGeometric, const FPoint& InPosition, FPoint InNormal, const FSurfacicBoundary& InBoundary)
+FPlaneSurface::FPlaneSurface(const double InToleranceGeometric, const FVector& InPosition, FVector InNormal, const FSurfacicBoundary& InBoundary)
 	: FSurface(InToleranceGeometric, InBoundary)
 {
 	InNormal.Normalize();
@@ -36,11 +36,11 @@ FPlane FPlaneSurface::GetPlane() const
 {
 	FSurfacicPoint Point3D;
 	EvaluatePoint(FVector2d::ZeroVector, Point3D, 0);
-	FPoint Normal = Matrix.Column(2);
+	FVector Normal = Matrix.Column(2);
 	return FPlane(Point3D.Point, Normal);
 }
 
-void FPlaneSurface::EvaluatePoint(const FPoint2D& InSurfacicCoordinate, FSurfacicPoint& OutPoint3D, int32 InDerivativeOrder) const
+void FPlaneSurface::EvaluatePoint(const FVector2d& InSurfacicCoordinate, FSurfacicPoint& OutPoint3D, int32 InDerivativeOrder) const
 {
 	OutPoint3D.DerivativeOrder = InDerivativeOrder;
 	OutPoint3D.Point = Matrix.Multiply(InSurfacicCoordinate);
@@ -52,13 +52,13 @@ void FPlaneSurface::EvaluatePoint(const FPoint2D& InSurfacicCoordinate, FSurfaci
 
 	if(InDerivativeOrder>1) 
 	{
-		OutPoint3D.LaplacianU = FPoint::ZeroVector;
-		OutPoint3D.LaplacianV = FPoint::ZeroVector;
-		OutPoint3D.LaplacianUV = FPoint::ZeroVector;
+		OutPoint3D.LaplacianU = FVector::ZeroVector;
+		OutPoint3D.LaplacianV = FVector::ZeroVector;
+		OutPoint3D.LaplacianUV = FVector::ZeroVector;
 	}
 }
 
-void FPlaneSurface::EvaluatePoints(const TArray<FPoint2D>& InSurfacicCoordinates, TArray<FSurfacicPoint>& OutPoint3D, int32 InDerivativeOrder) const
+void FPlaneSurface::EvaluatePoints(const TArray<FVector2d>& InSurfacicCoordinates, TArray<FSurfacicPoint>& OutPoint3D, int32 InDerivativeOrder) const
 {
 	int32 PointNum = InSurfacicCoordinates.Num();
 	OutPoint3D.SetNum(PointNum);
@@ -75,8 +75,8 @@ void FPlaneSurface::EvaluatePoints(const TArray<FPoint2D>& InSurfacicCoordinates
 
 	if (InDerivativeOrder > 0)
 	{
-		FPoint GradientU(Matrix.Column(0));
-		FPoint GradientV(Matrix.Column(1));
+		FVector GradientU(Matrix.Column(0));
+		FVector GradientV(Matrix.Column(1));
 
 		for (int32 Index = 0; Index < PointNum; ++Index)
 		{
@@ -99,21 +99,21 @@ void FPlaneSurface::EvaluatePointGrid(const FCoordinateGrid& Coordinates, FSurfa
 
 	OutPoints.Set2DCoordinates(Coordinates);
 
-	for (FPoint Point : OutPoints.Points2D)
+	for (FVector2d& Point : OutPoints.Points2D)
 	{
 		OutPoints.Points3D.Emplace(Matrix.Multiply(Point));
 	}
 
 	if(bComputeNormals)
 	{
-		FPoint Normal(Matrix.Column(2));
-		OutPoints.Normals.Init(Normal, PointNum);
+		FVector Normal(Matrix.Column(2));
+		OutPoints.Normals.Init((FVector3f)Normal, PointNum);
 	}
 }
 
-FPoint FPlaneSurface::ProjectPoint(const FPoint& Point, FPoint* OutProjectedPoint) const
+FVector FPlaneSurface::ProjectPoint(const FVector& Point, FVector* OutProjectedPoint) const
 {
-	FPoint PointCoordinate = InverseMatrix.Multiply(Point);
+	FVector PointCoordinate = InverseMatrix.Multiply(Point);
 	PointCoordinate.Z = 0.0;
 
 	if(OutProjectedPoint)
@@ -124,7 +124,7 @@ FPoint FPlaneSurface::ProjectPoint(const FPoint& Point, FPoint* OutProjectedPoin
 	return PointCoordinate;
 }
 
-void FPlaneSurface::ProjectPoints(const TArray<FPoint>& Points, TArray<FPoint>* PointCoordinates, TArray<FPoint>* OutProjectedPointS) const
+void FPlaneSurface::ProjectPoints(const TArray<FVector>& Points, TArray<FVector>* PointCoordinates, TArray<FVector>* OutProjectedPointS) const
 {
 	PointCoordinates->Reserve(Points.Num());
 	if(OutProjectedPointS) 
@@ -132,15 +132,15 @@ void FPlaneSurface::ProjectPoints(const TArray<FPoint>& Points, TArray<FPoint>*
 		OutProjectedPointS->Reserve(Points.Num());
 	}
 
-	for (const FPoint& Point : Points)
+	for (const FVector& Point : Points)
 	{
-		FPoint& PointCoordinate = PointCoordinates->Emplace_GetRef(InverseMatrix.Multiply(Point));
+		FVector& PointCoordinate = PointCoordinates->Emplace_GetRef(InverseMatrix.Multiply(Point));
 		PointCoordinate.Z = 0.0;
 	}
 
 	if (OutProjectedPointS)
 	{
-		for (const FPoint& PointCoordinate : *PointCoordinates)
+		for (const FVector& PointCoordinate : *PointCoordinates)
 		{
 			OutProjectedPointS->Emplace(Matrix.Multiply(PointCoordinate));
 		}
diff --git a/Engine/Source/Runtime/Datasmith/CADKernel/Base/Private/Geo/Surfaces/PlaneSurface.h b/Engine/Source/Runtime/Datasmith/CADKernel/Base/Private/Geo/Surfaces/PlaneSurface.h
--- a/Engine/Source/Runtime/Datasmith/CADKernel/Base/Private/Geo/Surfaces/PlaneSurface.h
+++ b/Engine/Source/Runtime/Datasmith/CADKernel/Base/Private/Geo/Surfaces/PlaneSurface.h
@@ -25,7 +25,7 @@ protected:
 	 * The surface is placed at its final position and orientation by the Matrix,
 	 * The matrix is calculated from the normal plane at its final position and its distance from the origin along the normal.
 	 */
-	FPlaneSurface(const double InToleranceGeometric, double InDistanceFromOrigin, FPoint InNormal, const FSurfacicBoundary& InBoundary)
+	FPlaneSurface(const double InToleranceGeometric, double InDistanceFromOrigin, FVector InNormal, const FSurfacicBoundary& InBoundary)
 		: FPlaneSurface(InToleranceGeometric, InNormal* InDistanceFromOrigin, InNormal, InBoundary)
 	{
 	}
@@ -35,16 +35,16 @@ protected:
 	 * The surface is placed at its final position and orientation by the Matrix
 	 * The matrix is calculated from the plan origin at its final position and its final normal
 	 */
-	FPlaneSurface(const double InToleranceGeometric, const FPoint& Position, FPoint Normal, const FSurfacicBoundary& InBoundary);
+	FPlaneSurface(const double InToleranceGeometric, const FVector& Position, FVector Normal, const FSurfacicBoundary& InBoundary);
 
 	FPlaneSurface() = default;
 
 	void ComputeMinToleranceIso()
 	{
-		FPoint Origin = Matrix.Multiply(FPoint::ZeroVector);
+		FVector Origin = Matrix.Multiply(FVector::ZeroVector);
 
-		FPoint Point2DU{ 1 , 0, 0 };
-		FPoint Point2DV{ 0, 1, 0 };
+		FVector Point2DU{ 1 , 0, 0 };
+		FVector Point2DV{ 0, 1, 0 };
 
 		double ToleranceU = Tolerance3D / ComputeScaleAlongAxis(Point2DU, Matrix, Origin);
 		double ToleranceV = Tolerance3D / ComputeScaleAlongAxis(Point2DV, Matrix, Origin);
@@ -79,12 +79,12 @@ public:
 
 	virtual TSharedPtr<FEntityGeom> ApplyMatrix(const FMatrixH& InMatrix) const override;
 
-	virtual void EvaluatePoint(const FPoint2D& InSurfacicCoordinate, FSurfacicPoint& OutPoint3D, int32 InDerivativeOrder = 0) const override;
-	virtual void EvaluatePoints(const TArray<FPoint2D>& InSurfacicCoordinates, TArray<FSurfacicPoint>& OutPoint3D, int32 InDerivativeOrder = 0) const override;
+	virtual void EvaluatePoint(const FVector2d& InSurfacicCoordinate, FSurfacicPoint& OutPoint3D, int32 InDerivativeOrder = 0) const override;
+	virtual void EvaluatePoints(const TArray<FVector2d>& InSurfacicCoordinates, TArray<FSurfacicPoint>& OutPoint3D, int32 InDerivativeOrder = 0) const override;
 	virtual void EvaluatePointGrid(const FCoordinateGrid& Coordinates, FSurfacicSampling& OutPoints, bool bComputeNormals = false) const override;
 
-	virtual FPoint ProjectPoint(const FPoint& InPoint, FPoint* OutProjectedPoint = nullptr) const;// override;
-	virtual void ProjectPoints(const TArray<FPoint>& InPoints, TArray<FPoint>* OutProjectedPointCoordinates, TArray<FPoint>* OutProjectedPoints = nullptr) const;// override;
+	virtual FVector ProjectPoint(const FVector& InPoint, FVector* OutProjectedPoint = nullptr) const;// override;
+	virtual void ProjectPoints(const TArray<FVector>& InPoints, TArray<FVector>* OutProjectedPointCoordinates, TArray<FVector>* OutProjectedPoints = nullptr) const;// override;
 
 	virtual void Presample(const FSurfacicBoundary& InBoundaries, FCoordinateGrid& OutCoordinates) override
 	{
diff --git a/Engine/Source/Runtime/Datasmith/CADKernel/Base/Private/Geo/Surfaces/RevolutionSurface.cpp b/Engine/Source/Runtime/Datasmith/CADKernel/Base/Private/Geo/Surfaces/RevolutionSurface.cpp
--- a/Engine/Source/Runtime/Datasmith/CADKernel/Base/Private/Geo/Surfaces/RevolutionSurface.cpp
+++ b/Engine/Source/Runtime/Datasmith/CADKernel/Base/Private/Geo/Surfaces/RevolutionSurface.cpp
@@ -24,7 +24,7 @@ FRevolutionSurface::FRevolutionSurface(const double InToleranceGeometric, TShare
 	ComputeDefaultMinToleranceIso();
 }
 
-void FRevolutionSurface::EvaluatePoint(const FPoint2D& InSurfacicCoordinate, FSurfacicPoint& OutPoint3D, int32 InDerivativeOrder) const
+void FRevolutionSurface::EvaluatePoint(const FVector2d& InSurfacicCoordinate, FSurfacicPoint& OutPoint3D, int32 InDerivativeOrder) const
 {
 	FCurvePoint GeneratrixPoint;
 	Generatrix->EvaluatePoint(InSurfacicCoordinate.X, GeneratrixPoint, InDerivativeOrder);
@@ -36,7 +36,7 @@ void FRevolutionSurface::EvaluatePoint(const FPoint2D& InSurfacicCoordinate, FSu
 
 	if (InDerivativeOrder > 0) 
 	{
-		FPoint Vector = OutPoint3D.Point - Axis->GetStartPoint();
+		FVector Vector = OutPoint3D.Point - Axis->GetStartPoint();
 		OutPoint3D.GradientU = Matrix.MultiplyVector(GeneratrixPoint.Gradient);
 
 		OutPoint3D.GradientV = RotationAxis ^ Vector;
@@ -61,7 +61,7 @@ void FRevolutionSurface::EvaluatePointGrid(const FCoordinateGrid& Coordinates, F
 
 	int32 DerivativeOrder = bComputeNormals ? 1 : 0;
 	TArray<FCurvePoint> GeneratrixPoints;
-	TArray<FPoint> GeneratrixGradientU;
+	TArray<FVector> GeneratrixGradientU;
 	Generatrix->EvaluatePoints(Coordinates[EIso::IsoU], GeneratrixPoints, DerivativeOrder);
 
 	TArray<FMatrixH> Matrix;
@@ -86,10 +86,10 @@ void FRevolutionSurface::EvaluatePointGrid(const FCoordinateGrid& Coordinates, F
 		{
 			for (FCurvePoint GeneratrixPoint : GeneratrixPoints)
 			{
-				FPoint GradientU = Matrix[Vndex].MultiplyVector(GeneratrixPoint.Gradient);
+				FVector GradientU = Matrix[Vndex].MultiplyVector(GeneratrixPoint.Gradient);
 
-				FPoint Vector = OutPoints.Points3D[Index] - Axis->GetStartPoint();
-				FPoint GradientV = RotationAxis ^ Vector;
+				FVector Vector = OutPoints.Points3D[Index] - Axis->GetStartPoint();
+				FVector GradientV = RotationAxis ^ Vector;
 
 				OutPoints.Normals.Emplace(GradientU ^ GradientV);
 				Index++;
diff --git a/Engine/Source/Runtime/Datasmith/CADKernel/Base/Private/Geo/Surfaces/RevolutionSurface.h b/Engine/Source/Runtime/Datasmith/CADKernel/Base/Private/Geo/Surfaces/RevolutionSurface.h
--- a/Engine/Source/Runtime/Datasmith/CADKernel/Base/Private/Geo/Surfaces/RevolutionSurface.h
+++ b/Engine/Source/Runtime/Datasmith/CADKernel/Base/Private/Geo/Surfaces/RevolutionSurface.h
@@ -17,7 +17,7 @@ protected:
 	TSharedPtr<FSegmentCurve> Axis;
 	TSharedPtr<FCurve> Generatrix;
 
-	FPoint RotationAxis;
+	FVector RotationAxis;
 
 	/**
 	 * A surface of revolution is defined by an axis of rotation (Axis) and a generatrix.
@@ -68,7 +68,7 @@ public:
 
 	virtual TSharedPtr<FEntityGeom> ApplyMatrix(const FMatrixH& InMatrix) const override;
 
-	virtual void EvaluatePoint(const FPoint2D& InSurfacicCoordinate, FSurfacicPoint& OutPoint3D, int32 InDerivativeOrder = 0) const override;
+	virtual void EvaluatePoint(const FVector2d& InSurfacicCoordinate, FSurfacicPoint& OutPoint3D, int32 InDerivativeOrder = 0) const override;
 	virtual void EvaluatePointGrid(const FCoordinateGrid& Coordinates, FSurfacicSampling& OutPoints, bool bComputeNormals = false) const override;
 
 	virtual void LinesNotDerivables(const FSurfacicBoundary& Bounds, int32 InDerivativeOrder, FCoordinateGrid& OutCoordinates) const override;
diff --git a/Engine/Source/Runtime/Datasmith/CADKernel/Base/Private/Geo/Surfaces/RuledSurface.cpp b/Engine/Source/Runtime/Datasmith/CADKernel/Base/Private/Geo/Surfaces/RuledSurface.cpp
--- a/Engine/Source/Runtime/Datasmith/CADKernel/Base/Private/Geo/Surfaces/RuledSurface.cpp
+++ b/Engine/Source/Runtime/Datasmith/CADKernel/Base/Private/Geo/Surfaces/RuledSurface.cpp
@@ -32,7 +32,7 @@ void FRuledSurface::LinesNotDerivables(const FSurfacicBoundary& InBoundary, int3
 	ArrayUtils::Complete(OutCoordinates[EIso::IsoU], LinesNotDerivablesCurve1, GetIsoTolerances()[IsoU]);
 }
 
-void FRuledSurface::EvaluatePoint(const FPoint2D& InPoint2D, FSurfacicPoint& OutPoint3D, int32 InDerivativeOrder) const
+void FRuledSurface::EvaluatePoint(const FVector2d& InPoint2D, FSurfacicPoint& OutPoint3D, int32 InDerivativeOrder) const
 {
 	double CoordinateCurve0 = Curves[0]->GetBoundary().Min + InPoint2D.X * (Curves[0]->GetBoundary().Max - Curves[0]->GetBoundary().Min);
 	double CoordinateCurve1 = Curves[1]->GetBoundary().Min + InPoint2D.X * (Curves[1]->GetBoundary().Max - Curves[1]->GetBoundary().Min);
@@ -54,7 +54,7 @@ void FRuledSurface::EvaluatePoint(const FPoint2D& InPoint2D, FSurfacicPoint& Out
 		if (InDerivativeOrder > 1)
 		{
 			OutPoint3D.LaplacianU = CurvePoint[0].Laplacian + InPoint2D.Y * (CurvePoint[1].Laplacian - CurvePoint[0].Laplacian);
-			OutPoint3D.LaplacianV = FPoint(0.0, 0.0, 0.0);
+			OutPoint3D.LaplacianV = FVector(0.0, 0.0, 0.0);
 		}
 	}
 }
@@ -96,7 +96,7 @@ void FRuledSurface::EvaluatePointGrid(const FCoordinateGrid& Coordinates, FSurfa
 
 	for (int32 Undex = 0; Undex < Coordinates.IsoCount(EIso::IsoU); Undex++)
 	{
-		FPoint GradientV = CurveCoords[1][Undex].Point - CurveCoords[0][Undex].Point;
+		FVector GradientV = CurveCoords[1][Undex].Point - CurveCoords[0][Undex].Point;
 
 		for (int32 Vndex = 0, Index = Undex; Vndex < VCount; ++Vndex)
 		{
@@ -108,8 +108,8 @@ void FRuledSurface::EvaluatePointGrid(const FCoordinateGrid& Coordinates, FSurfa
 		{
 			for (int32 Vndex = 0, Index = Undex; Vndex < VCount; ++Vndex)
 			{
-				FPoint GradientU = CurveCoords[0][Undex].Gradient + Coordinates[EIso::IsoV][Vndex] * (CurveCoords[1][Undex].Gradient - CurveCoords[0][Undex].Gradient);
-				OutPoints.Normals[Index] = GradientU ^ GradientV;
+				FVector GradientU = CurveCoords[0][Undex].Gradient + Coordinates[EIso::IsoV][Vndex] * (CurveCoords[1][Undex].Gradient - CurveCoords[0][Undex].Gradient);
+				OutPoints.Normals[Index] = (FVector3f)(GradientU ^ GradientV);
 				Index += UCount;
 			}
 		}
diff --git a/Engine/Source/Runtime/Datasmith/CADKernel/Base/Private/Geo/Surfaces/RuledSurface.h b/Engine/Source/Runtime/Datasmith/CADKernel/Base/Private/Geo/Surfaces/RuledSurface.h
--- a/Engine/Source/Runtime/Datasmith/CADKernel/Base/Private/Geo/Surfaces/RuledSurface.h
+++ b/Engine/Source/Runtime/Datasmith/CADKernel/Base/Private/Geo/Surfaces/RuledSurface.h
@@ -53,7 +53,7 @@ public:
 
 	virtual TSharedPtr<FEntityGeom> ApplyMatrix(const FMatrixH& InMatrix) const override;
 
-	virtual void EvaluatePoint(const FPoint2D& InSurfacicCoordinate, FSurfacicPoint& OutPoint3D, int32 InDerivativeOrder = 0) const override;
+	virtual void EvaluatePoint(const FVector2d& InSurfacicCoordinate, FSurfacicPoint& OutPoint3D, int32 InDerivativeOrder = 0) const override;
 	virtual void EvaluatePointGrid(const FCoordinateGrid& Coordinates, FSurfacicSampling& OutPoints, bool bComputeNormals = false) const override;
 
 	virtual void LinesNotDerivables(const FSurfacicBoundary& Bounds, int32 InDerivativeOrder, FCoordinateGrid& OutCoordinates) const override;
diff --git a/Engine/Source/Runtime/Datasmith/CADKernel/Base/Private/Geo/Surfaces/SphericalSurface.cpp b/Engine/Source/Runtime/Datasmith/CADKernel/Base/Private/Geo/Surfaces/SphericalSurface.cpp
--- a/Engine/Source/Runtime/Datasmith/CADKernel/Base/Private/Geo/Surfaces/SphericalSurface.cpp
+++ b/Engine/Source/Runtime/Datasmith/CADKernel/Base/Private/Geo/Surfaces/SphericalSurface.cpp
@@ -8,7 +8,7 @@
 namespace UE::CADKernel
 {
 
-void FSphericalSurface::EvaluatePointGridInCylindricalSpace(const FCoordinateGrid& Coordinates, TArray<FPoint2D>& OutPoints) const
+void FSphericalSurface::EvaluatePointGridInCylindricalSpace(const FCoordinateGrid& Coordinates, TArray<FVector2d>& OutPoints) const
 {
 	int32 PointNum = Coordinates.Count();
 	OutPoints.Empty(PointNum);
@@ -82,15 +82,15 @@ void FSphericalSurface::EvaluatePointGrid(const FCoordinateGrid& Coordinates, FS
 		}
 	}
 
-	for (FPoint& Point : OutPoints.Points3D)
+	for (FVector& Point : OutPoints.Points3D)
 	{
 		Point = Matrix.Multiply(Point);
 	}
 
 	if (bComputeNormals)
 	{
-		FPoint Center = Matrix.Column(3);
-		for (FPoint& Point : OutPoints.Points3D)
+		FVector Center = Matrix.Column(3);
+		for (FVector& Point : OutPoints.Points3D)
 		{
 			OutPoints.Normals.Emplace(Point - Center);
 		}
diff --git a/Engine/Source/Runtime/Datasmith/CADKernel/Base/Private/Geo/Surfaces/SphericalSurface.h b/Engine/Source/Runtime/Datasmith/CADKernel/Base/Private/Geo/Surfaces/SphericalSurface.h
--- a/Engine/Source/Runtime/Datasmith/CADKernel/Base/Private/Geo/Surfaces/SphericalSurface.h
+++ b/Engine/Source/Runtime/Datasmith/CADKernel/Base/Private/Geo/Surfaces/SphericalSurface.h
@@ -52,10 +52,10 @@ protected:
 	{
 		double Tolerance2D = Tolerance3D / Radius;
 
-		FPoint Origin = Matrix.Multiply(FPoint::ZeroVector);
+		FVector Origin = Matrix.Multiply(FVector::ZeroVector);
 
-		FPoint Point2DU{ 1 , 0, 0 };
-		FPoint Point2DV{ 0, 1, 0 };
+		FVector Point2DU{ 1 , 0, 0 };
+		FVector Point2DV{ 0, 1, 0 };
 
 		double ToleranceU = Tolerance2D / ComputeScaleAlongAxis(Point2DU, Matrix, Origin);
 		double ToleranceV = Tolerance2D / ComputeScaleAlongAxis(Point2DV, Matrix, Origin);
@@ -83,7 +83,7 @@ public:
 
 	virtual TSharedPtr<FEntityGeom> ApplyMatrix(const FMatrixH& InMatrix) const override;
 
-	virtual void EvaluatePoint(const FPoint2D& InSurfacicCoordinate, FSurfacicPoint& OutPoint3D, int32 InDerivativeOrder = 0) const override
+	virtual void EvaluatePoint(const FVector2d& InSurfacicCoordinate, FSurfacicPoint& OutPoint3D, int32 InDerivativeOrder = 0) const override
 	{
 		double CosU = cos(InSurfacicCoordinate.X);
 		double CosV = cos(InSurfacicCoordinate.Y);
@@ -97,8 +97,8 @@ public:
 
 		if (InDerivativeOrder > 0)
 		{
-			OutPoint3D.GradientU = FPoint(-Radius * CosV * SinU, Radius * CosV * CosU, 0.0);
-			OutPoint3D.GradientV = FPoint(-Radius * SinV * CosU, -Radius * SinV * SinU, Radius * CosV);
+			OutPoint3D.GradientU = FVector(-Radius * CosV * SinU, Radius * CosV * CosU, 0.0);
+			OutPoint3D.GradientV = FVector(-Radius * SinV * CosU, -Radius * SinV * SinU, Radius * CosV);
 
 			OutPoint3D.GradientU = Matrix.MultiplyVector(OutPoint3D.GradientU);
 			OutPoint3D.GradientV = Matrix.MultiplyVector(OutPoint3D.GradientV);
@@ -106,9 +106,9 @@ public:
 
 		if (InDerivativeOrder > 1)
 		{
-			OutPoint3D.LaplacianU = FPoint(-Radius * CosV * CosU, -Radius * CosV * SinU, 0.0);
-			OutPoint3D.LaplacianV = FPoint(OutPoint3D.LaplacianU.X, OutPoint3D.LaplacianU.Y, -Radius * SinV);
-			OutPoint3D.LaplacianUV = FPoint(Radius * SinV * SinU, -Radius * SinV * CosU, 0.);
+			OutPoint3D.LaplacianU = FVector(-Radius * CosV * CosU, -Radius * CosV * SinU, 0.0);
+			OutPoint3D.LaplacianV = FVector(OutPoint3D.LaplacianU.X, OutPoint3D.LaplacianU.Y, -Radius * SinV);
+			OutPoint3D.LaplacianUV = FVector(Radius * SinV * SinU, -Radius * SinV * CosU, 0.);
 
 			OutPoint3D.LaplacianU = Matrix.MultiplyVector(OutPoint3D.LaplacianU);
 			OutPoint3D.LaplacianV = Matrix.MultiplyVector(OutPoint3D.LaplacianV);
@@ -118,7 +118,7 @@ public:
 
 	virtual void EvaluatePointGrid(const FCoordinateGrid& Coordinates, FSurfacicSampling& OutPoints, bool bComputeNormals = false) const override;
 
-	virtual FPoint2D EvaluatePointInCylindricalSpace(const FPoint2D& InSurfacicCoordinate) const override
+	virtual FVector2d EvaluatePointInCylindricalSpace(const FVector2d& InSurfacicCoordinate) const override
 	{
 		double CosU = cos(InSurfacicCoordinate.X);
 		double CosV = cos(InSurfacicCoordinate.Y);
@@ -126,10 +126,10 @@ public:
 		double SinU = sin(InSurfacicCoordinate.X);
 		double SwapOrientation = (InSurfacicCoordinate.Y < DOUBLE_PI&& InSurfacicCoordinate.Y >= 0) ? 1.0 : -1.0;
 
-		return FPoint(Radius * CosV * CosU * SwapOrientation, Radius * CosV * SinU);
+		return FVector2d(Radius * CosV * CosU * SwapOrientation, Radius * CosV * SinU);
 	}
 
-	virtual void EvaluatePointGridInCylindricalSpace(const FCoordinateGrid& Coordinates, TArray<FPoint2D>&) const override;
+	virtual void EvaluatePointGridInCylindricalSpace(const FCoordinateGrid& Coordinates, TArray<FVector2d>&) const override;
 
 	virtual void Presample(const FSurfacicBoundary& InBoundaries, FCoordinateGrid& OutCoordinates) override
 	{
diff --git a/Engine/Source/Runtime/Datasmith/CADKernel/Base/Private/Geo/Surfaces/Surface.cpp b/Engine/Source/Runtime/Datasmith/CADKernel/Base/Private/Geo/Surfaces/Surface.cpp
--- a/Engine/Source/Runtime/Datasmith/CADKernel/Base/Private/Geo/Surfaces/Surface.cpp
+++ b/Engine/Source/Runtime/Datasmith/CADKernel/Base/Private/Geo/Surfaces/Surface.cpp
@@ -19,7 +19,7 @@ FInfoEntity& FSurface::GetInfo(FInfoEntity& Info) const
 	return FEntityGeom::GetInfo(Info)
 		.Add(TEXT("Surface type"), SurfacesTypesNames[(uint8)GetSurfaceType()])
 		.Add(TEXT("Boundary"), Boundary)
-		.Add(TEXT("MinToleranceIso"), (FPoint2D)*MinToleranceIso);
+		.Add(TEXT("MinToleranceIso"), (FVector2d)*MinToleranceIso);
 }
 #endif
 
@@ -54,7 +54,7 @@ void FSurface::Presample(const FSurfacicBoundary& InBoundaries, FCoordinateGrid&
 	Sampler.Sample();
 }
 
-void FSurface::EvaluatePoints(const TArray<FPoint2D>& InSurfacicCoordinates, TArray<FSurfacicPoint>& OutPoint3D, int32 InDerivativeOrder) const
+void FSurface::EvaluatePoints(const TArray<FVector2d>& InSurfacicCoordinates, TArray<FSurfacicPoint>& OutPoint3D, int32 InDerivativeOrder) const
 {
 	OutPoint3D.SetNum(InSurfacicCoordinates.Num());
 	for (int32 Index = 0; Index < InSurfacicCoordinates.Num(); ++Index)
@@ -148,7 +148,7 @@ void FSurface::EvaluatePoints(const TArray<FCurvePoint2D>& InPoints2D, FSurfacic
 	}
 }
 
-void FSurface::EvaluateNormals(const TArray<FPoint2D>& InPoints2D, TArray<FVector3f>& Normals) const
+void FSurface::EvaluateNormals(const TArray<FVector2d>& InPoints2D, TArray<FVector3f>& Normals) const
 {
 	int32 DerivativeOrder = 1;
 	TArray<FSurfacicPoint> Points3D;
@@ -188,7 +188,7 @@ void FSurface::EvaluatePointGrid(const FCoordinateGrid& Coordinates, FSurfacicSa
 	{
 		for (FSurfacicPoint Point : Point3D)
 		{
-			const FPoint Normal = Point.GradientU ^ Point.GradientV;
+			const FVector Normal = Point.GradientU ^ Point.GradientV;
 			OutPoints.Normals.Emplace(Normal);
 		}
 		OutPoints.NormalizeNormals();
diff --git a/Engine/Source/Runtime/Datasmith/CADKernel/Base/Private/Geo/Surfaces/SurfaceMaker.cpp b/Engine/Source/Runtime/Datasmith/CADKernel/Base/Private/Geo/Surfaces/SurfaceMaker.cpp
--- a/Engine/Source/Runtime/Datasmith/CADKernel/Base/Private/Geo/Surfaces/SurfaceMaker.cpp
+++ b/Engine/Source/Runtime/Datasmith/CADKernel/Base/Private/Geo/Surfaces/SurfaceMaker.cpp
@@ -19,7 +19,7 @@
 
 namespace UE::CADKernel
 {
-TSharedPtr<FSurface> FSurface::MakeBezierSurface(const double InToleranceGeometric, int32 InUDegre, int32 InVDegre, const TArray<FPoint>& InPoles)
+TSharedPtr<FSurface> FSurface::MakeBezierSurface(const double InToleranceGeometric, int32 InUDegre, int32 InVDegre, const TArray<FVector>& InPoles)
 {
 	return UE::CADKernel::FEntity::MakeShared<UE::CADKernel::FBezierSurface>(InToleranceGeometric, InUDegre, InVDegre, InPoles);
 }
diff --git a/Engine/Source/Runtime/Datasmith/CADKernel/Base/Private/Geo/Surfaces/SurfaceUtilities.cpp b/Engine/Source/Runtime/Datasmith/CADKernel/Base/Private/Geo/Surfaces/SurfaceUtilities.cpp
--- a/Engine/Source/Runtime/Datasmith/CADKernel/Base/Private/Geo/Surfaces/SurfaceUtilities.cpp
+++ b/Engine/Source/Runtime/Datasmith/CADKernel/Base/Private/Geo/Surfaces/SurfaceUtilities.cpp
@@ -20,7 +20,7 @@ namespace UE::CADKernel::SurfaceUtilities
 		return Poles;
 	}
 
-	bool AreControlPointsPlanar(const TArray<FPoint>& Poles, int32 UPoleCount)
+	bool AreControlPointsPlanar(const TArray<FVector>& Poles, int32 UPoleCount)
 	{
 		using namespace UE::CADKernel;
 
diff --git a/Engine/Source/Runtime/Datasmith/CADKernel/Base/Private/Geo/Surfaces/TabulatedCylinderSurface.cpp b/Engine/Source/Runtime/Datasmith/CADKernel/Base/Private/Geo/Surfaces/TabulatedCylinderSurface.cpp
--- a/Engine/Source/Runtime/Datasmith/CADKernel/Base/Private/Geo/Surfaces/TabulatedCylinderSurface.cpp
+++ b/Engine/Source/Runtime/Datasmith/CADKernel/Base/Private/Geo/Surfaces/TabulatedCylinderSurface.cpp
@@ -25,12 +25,11 @@ void FTabulatedCylinderSurface::ResetMarkersRecursively() const
 	GuideCurve->ResetMarkersRecursively();
 }
 
-void FTabulatedCylinderSurface::EvaluatePoint(const FPoint2D& InSurfacicCoordinate, FSurfacicPoint& OutPoint3D, int32 InDerivativeOrder) const
+void FTabulatedCylinderSurface::EvaluatePoint(const FVector2d& InSurfacicCoordinate, FSurfacicPoint& OutPoint3D, int32 InDerivativeOrder) const
 {
 	const FLinearBoundary& CurveBounds = GuideCurve->GetBoundary();
 
 	FCurvePoint PointU;
-	FPoint LaplacianU;
 	GuideCurve->EvaluatePoint(CurveBounds.GetAt(InSurfacicCoordinate.X), PointU, InDerivativeOrder);
 
 	OutPoint3D.DerivativeOrder = InDerivativeOrder;
@@ -44,8 +43,8 @@ void FTabulatedCylinderSurface::EvaluatePoint(const FPoint2D& InSurfacicCoordina
 		if (InDerivativeOrder > 1)
 		{
 			OutPoint3D.LaplacianU = PointU.Laplacian * FMath::Square(CurveBounds.Length());
-			OutPoint3D.LaplacianV = FPoint::ZeroVector;
-			OutPoint3D.LaplacianUV = FPoint::ZeroVector;
+			OutPoint3D.LaplacianV = FVector::ZeroVector;
+			OutPoint3D.LaplacianUV = FVector::ZeroVector;
 		}
 	}
 }
@@ -80,7 +79,7 @@ void FTabulatedCylinderSurface::EvaluatePointGrid(const FCoordinateGrid& Coordin
 		for (const FCurvePoint& Point : GuidePoints)
 		{
 			int32 Index = Undex++;
-			FPoint Normal = Point.Gradient ^ DirectorVector;
+			FVector3f Normal = (FVector3f)(Point.Gradient ^ DirectorVector);
 			for (int32 Vndex = 0; Vndex < NumV; ++Vndex)
 			{
 				OutPoints.Normals[Index] = Normal;
@@ -125,12 +124,12 @@ TSharedPtr<FEntityGeom> FTabulatedCylinderSurface::ApplyMatrix(const FMatrixH& I
 	}
 
 	double GuideUMin = GuideCurve->GetUMin();
-	FPoint Point = GuideCurve->EvaluatePoint(GuideUMin);
+	FVector Point = GuideCurve->EvaluatePoint(GuideUMin);
 	Point += DirectorVector;
 
-	FPoint TransformedPoint = InMatrix.Multiply(Point);
+	FVector TransformedPoint = InMatrix.Multiply(Point);
 	Point = StaticCastSharedPtr<FCurve>(TransformedGuideCurve)->EvaluatePoint(GuideUMin);
-	FPoint NewDirector = TransformedPoint - Point;
+	FVector NewDirector = TransformedPoint - Point;
 
 	return FEntity::MakeShared<FTabulatedCylinderSurface>(Tolerance3D, StaticCastSharedPtr<FCurve>(TransformedGuideCurve), NewDirector, Boundary[EIso::IsoV].Min, Boundary[EIso::IsoV].Max);
 }
diff --git a/Engine/Source/Runtime/Datasmith/CADKernel/Base/Private/Geo/Surfaces/TabulatedCylinderSurface.h b/Engine/Source/Runtime/Datasmith/CADKernel/Base/Private/Geo/Surfaces/TabulatedCylinderSurface.h
--- a/Engine/Source/Runtime/Datasmith/CADKernel/Base/Private/Geo/Surfaces/TabulatedCylinderSurface.h
+++ b/Engine/Source/Runtime/Datasmith/CADKernel/Base/Private/Geo/Surfaces/TabulatedCylinderSurface.h
@@ -14,7 +14,7 @@ class CADKERNEL_API FTabulatedCylinderSurface : public FSurface
 
 protected:
 	TSharedPtr<FCurve> GuideCurve;
-	FPoint DirectorVector;
+	FVector DirectorVector;
 
 	/**
 	 * A tabulated cylinder is a surface formed by moving a curve segment called GuideCurve parallel to
@@ -25,7 +25,7 @@ protected:
 	 * Bounds[EIso::IsoV].Min = InVMin;
 	 * Bounds[EIso::IsoV].Max = InVMax;
 	 */
-	FTabulatedCylinderSurface(const double InToleranceGeometric, TSharedPtr<FCurve> InGuideCurve, const FPoint& InDirectorVector, double InVMin = 0., double InVMax = 1.)
+	FTabulatedCylinderSurface(const double InToleranceGeometric, TSharedPtr<FCurve> InGuideCurve, const FVector& InDirectorVector, double InVMin = 0., double InVMax = 1.)
 		: FSurface(InToleranceGeometric, 0.0, 1.0, InVMin, InVMax)
 		, GuideCurve(InGuideCurve)
 		, DirectorVector(InDirectorVector)
@@ -57,12 +57,12 @@ public:
 		return GuideCurve;
 	}
 
-	FPoint GetDirectorVector() const
+	FVector GetDirectorVector() const
 	{
 		return DirectorVector;
 	}
 
-	virtual void EvaluatePoint(const FPoint2D& InSurfacicCoordinate, FSurfacicPoint& OutPoint3D, int32 InDerivativeOrder = 0) const override;
+	virtual void EvaluatePoint(const FVector2d& InSurfacicCoordinate, FSurfacicPoint& OutPoint3D, int32 InDerivativeOrder = 0) const override;
 	virtual void EvaluatePointGrid(const FCoordinateGrid& Coordinates, FSurfacicSampling& OutPoints, bool bComputeNormals = false) const override;
 
 	virtual void LinesNotDerivables(const FSurfacicBoundary& Bounds, int32 InDerivativeOrder, FCoordinateGrid& OutNotDerivables) const override;
diff --git a/Engine/Source/Runtime/Datasmith/CADKernel/Base/Private/Geo/Surfaces/TorusSurface.cpp b/Engine/Source/Runtime/Datasmith/CADKernel/Base/Private/Geo/Surfaces/TorusSurface.cpp
--- a/Engine/Source/Runtime/Datasmith/CADKernel/Base/Private/Geo/Surfaces/TorusSurface.cpp
+++ b/Engine/Source/Runtime/Datasmith/CADKernel/Base/Private/Geo/Surfaces/TorusSurface.cpp
@@ -28,7 +28,7 @@ FInfoEntity& FTorusSurface::GetInfo(FInfoEntity& Info) const
 }
 #endif
 
-void FTorusSurface::EvaluatePointGridInCylindricalSpace(const FCoordinateGrid& Coordinates, TArray<FPoint2D>& OutPoints) const
+void FTorusSurface::EvaluatePointGridInCylindricalSpace(const FCoordinateGrid& Coordinates, TArray<FVector2d>& OutPoints) const
 {
 	int32 PointNum = Coordinates.Count();
 
@@ -108,21 +108,21 @@ void FTorusSurface::EvaluatePointGrid(const FCoordinateGrid& Coordinates, FSurfa
 			double DeltaHeight = MinorRadius * CosV;
 			for (int32 Undex = 0; Undex < UCount; Undex++)
 			{
-				FPoint GradientU = FPoint(Rho * (-SinU[Undex]), Rho * CosU[Undex], 0.0);
-				FPoint GradientV = FPoint(-Height * CosU[Undex], -Height * SinU[Undex], DeltaHeight);
+				FVector GradientU = FVector(Rho * (-SinU[Undex]), Rho * CosU[Undex], 0.0);
+				FVector GradientV = FVector(-Height * CosU[Undex], -Height * SinU[Undex], DeltaHeight);
 				OutPoints.Normals.Emplace(GradientU ^ GradientV);
 			}
 		}
 	}
 
-	for (FPoint& Point : OutPoints.Points3D)
+	for (FVector& Point : OutPoints.Points3D)
 	{
 		Point = Matrix.Multiply(Point);
 	}
 
 	if (bComputeNormals)
 	{
-		FVector3f Center = Matrix.Column(3);
+		FVector3f Center = (FVector3f)Matrix.Column(3);
 		for (FVector3f& Normal : OutPoints.Normals)
 		{
 			 Normal = Matrix.MultiplyVector(Normal);
diff --git a/Engine/Source/Runtime/Datasmith/CADKernel/Base/Private/Geo/Surfaces/TorusSurface.h b/Engine/Source/Runtime/Datasmith/CADKernel/Base/Private/Geo/Surfaces/TorusSurface.h
--- a/Engine/Source/Runtime/Datasmith/CADKernel/Base/Private/Geo/Surfaces/TorusSurface.h
+++ b/Engine/Source/Runtime/Datasmith/CADKernel/Base/Private/Geo/Surfaces/TorusSurface.h
@@ -57,10 +57,10 @@ protected:
 		double Tolerance2DU = Tolerance3D / MajorRadius;
 		double Tolerance2DV = Tolerance3D / MinorRadius;
 
-		FPoint Origin = Matrix.Multiply(FPoint::ZeroVector);
+		FVector Origin = Matrix.Multiply(FVector::ZeroVector);
 
-		FPoint Point2DU{ 1 , 0, 0 };
-		FPoint Point2DV{ 0, 1, 0 };
+		FVector Point2DU{ 1 , 0, 0 };
+		FVector Point2DV{ 0, 1, 0 };
 
 		Tolerance2DU /= ComputeScaleAlongAxis(Point2DU, Matrix, Origin);
 		Tolerance2DV /= ComputeScaleAlongAxis(Point2DV, Matrix, Origin);
@@ -89,7 +89,7 @@ public:
 
 	virtual TSharedPtr<FEntityGeom> ApplyMatrix(const FMatrixH& InMatrix) const override;
 
-	virtual void EvaluatePoint(const FPoint2D& InSurfacicCoordinate, FSurfacicPoint& OutPoint3D, int32 InDerivativeOrder = 0) const override
+	virtual void EvaluatePoint(const FVector2d& InSurfacicCoordinate, FSurfacicPoint& OutPoint3D, int32 InDerivativeOrder = 0) const override
 	{
 		const double CosU = cos(InSurfacicCoordinate.X);
 		const double CosV = cos(InSurfacicCoordinate.Y);
@@ -105,17 +105,17 @@ public:
 
 		if (InDerivativeOrder > 0)
 		{
-			OutPoint3D.GradientU = FPoint((MajorRadius + MinorRadius * CosV) * -SinU, (MajorRadius + MinorRadius * CosV) * CosU, 0.0);
-			OutPoint3D.GradientV = FPoint((MinorRadius * -SinV) * CosU, (MinorRadius * -SinV) * SinU, MinorRadius * CosV);
+			OutPoint3D.GradientU = FVector((MajorRadius + MinorRadius * CosV) * -SinU, (MajorRadius + MinorRadius * CosV) * CosU, 0.0);
+			OutPoint3D.GradientV = FVector((MinorRadius * -SinV) * CosU, (MinorRadius * -SinV) * SinU, MinorRadius * CosV);
 
 			OutPoint3D.GradientU = Matrix.MultiplyVector(OutPoint3D.GradientU);
 			OutPoint3D.GradientV = Matrix.MultiplyVector(OutPoint3D.GradientV);
 
 			if (InDerivativeOrder > 1)
 			{
-				OutPoint3D.LaplacianU = FPoint((MajorRadius + MinorRadius * CosV) * -CosU, (MajorRadius + MinorRadius * CosV) * -SinU, 0.0);
-				OutPoint3D.LaplacianV = FPoint((MinorRadius * -CosV) * CosU, (MinorRadius * -CosV) * SinU, MinorRadius * -SinV);
-				OutPoint3D.LaplacianUV = FPoint((MinorRadius * -SinV) * -SinU, (MinorRadius * -SinV) * CosU, 0.0);
+				OutPoint3D.LaplacianU = FVector((MajorRadius + MinorRadius * CosV) * -CosU, (MajorRadius + MinorRadius * CosV) * -SinU, 0.0);
+				OutPoint3D.LaplacianV = FVector((MinorRadius * -CosV) * CosU, (MinorRadius * -CosV) * SinU, MinorRadius * -SinV);
+				OutPoint3D.LaplacianUV = FVector((MinorRadius * -SinV) * -SinU, (MinorRadius * -SinV) * CosU, 0.0);
 
 				OutPoint3D.LaplacianU = Matrix.MultiplyVector(OutPoint3D.LaplacianU);
 				OutPoint3D.LaplacianV = Matrix.MultiplyVector(OutPoint3D.LaplacianV);
@@ -126,9 +126,9 @@ public:
 
 	virtual void EvaluatePointGrid(const FCoordinateGrid& Coordinates, FSurfacicSampling& OutPoints, bool bComputeNormals = false) const override;
 
-	virtual void EvaluatePointGridInCylindricalSpace(const FCoordinateGrid& Coordinates, TArray<FPoint2D>&) const override;
+	virtual void EvaluatePointGridInCylindricalSpace(const FCoordinateGrid& Coordinates, TArray<FVector2d>&) const override;
 
-	virtual FPoint2D EvaluatePointInCylindricalSpace(const FPoint2D& InSurfacicCoordinate) const override
+	virtual FVector2d EvaluatePointInCylindricalSpace(const FVector2d& InSurfacicCoordinate) const override
 	{
 		const double CosU = cos(InSurfacicCoordinate.X);
 		const double CosV = cos(InSurfacicCoordinate.Y);
@@ -139,7 +139,7 @@ public:
 
 		double SwapOrientation = (InSurfacicCoordinate.Y < DOUBLE_PI&& InSurfacicCoordinate.Y >= 0) ? 1.0 : -1.0;
 
-		return FPoint2D(Rho * CosU * SwapOrientation, Rho * SinU);
+		return FVector2d(Rho * CosU * SwapOrientation, Rho * SinU);
 	}
 
 	virtual void Presample(const FSurfacicBoundary& InBoundaries, FCoordinateGrid& OutCoordinates) override
diff --git a/Engine/Source/Runtime/Datasmith/CADKernel/Base/Private/Math/BSpline.cpp b/Engine/Source/Runtime/Datasmith/CADKernel/Base/Private/Math/BSpline.cpp
--- a/Engine/Source/Runtime/Datasmith/CADKernel/Base/Private/Math/BSpline.cpp
+++ b/Engine/Source/Runtime/Datasmith/CADKernel/Base/Private/Math/BSpline.cpp
@@ -22,9 +22,9 @@ namespace BSpline
 
 // DeBoor Valeurs B-Spline
 void DeBoorValeursBSpline(const int32 degre, const double* const knot, const int32 segment, const int32 ndim, const int32 derivee, const double u, double* const pole_aux, double* const grad_aux, double* const lap_aux, double* const valeur, double* const Gradient, double* const Laplacian);
-void InsertKnot(TArray<double>& vn, TArray<FPoint>& poles, TArray<double>& weights, double u, double* newU = nullptr);
-void Blossom(int32 degre, const TArray<FPoint>& poles, const TArray<double>& vecteurNodal, const TArray<double>& poids, int32 seg, TArray<double>& params, FPoint& pnt, double& weight);
-void DuplicateNurbsCurveWithHigherDegree(int32 degre, const TArray<FPoint>& poles, const TArray<double>& nodalVector, const TArray<double>& weights, TArray<FPoint>& newPoles, TArray<double>& newNodalVector, TArray<double>& newWeights);
+void InsertKnot(TArray<double>& vn, TArray<FVector>& poles, TArray<double>& weights, double u, double* newU = nullptr);
+void Blossom(int32 degre, const TArray<FVector>& poles, const TArray<double>& vecteurNodal, const TArray<double>& poids, int32 seg, TArray<double>& params, FVector& pnt, double& weight);
+void DuplicateNurbsCurveWithHigherDegree(int32 degre, const TArray<FVector>& poles, const TArray<double>& nodalVector, const TArray<double>& weights, TArray<FVector>& newPoles, TArray<double>& newNodalVector, TArray<double>& newWeights);
 
 #define DEGRE_MAX_POLY 9
 namespace
@@ -685,7 +685,7 @@ TSharedPtr<FNURBSCurve> DuplicateNurbsCurveWithHigherDegree(int32 Degre, const F
 	return TSharedPtr<FNURBSCurve>();
 
 #ifdef UETOOL_5046_WIP
-	TArray<FPoint> NewPoles;
+	TArray<FVector> NewPoles;
 	TArray<double> NewNodalVector;
 	TArray<double> NewWeights;
 
@@ -699,7 +699,7 @@ TSharedPtr<FNURBSCurve> DuplicateNurbsCurveWithHigherDegree(int32 Degre, const F
 }
 
 
-void EvaluatePoint(const FNURBSSurface& Nurbs, const FPoint2D& InPoint2D, FSurfacicPoint& OutPoint3D, int32 InDerivativeOrder)
+void EvaluatePoint(const FNURBSSurface& Nurbs, const FVector2d& InPoint2D, FSurfacicPoint& OutPoint3D, int32 InDerivativeOrder)
 {
 	double Point4[4];
 	double GradientU[4];
@@ -813,7 +813,7 @@ void EvaluatePointGrid(const FNURBSSurface& Nurbs, const FCoordinateGrid& Coords
 	}
 	else
 	{
-		memcpy(OutPoints.Points3D.GetData(), HPoints.GetData(), PointCount * sizeof(FPoint));
+		memcpy(OutPoints.Points3D.GetData(), HPoints.GetData(), PointCount * sizeof(FVector));
 	}
 
 	if (bComputeNormals)
@@ -822,8 +822,8 @@ void EvaluatePointGrid(const FNURBSSurface& Nurbs, const FCoordinateGrid& Coords
 
 		if (Nurbs.IsRational())
 		{
-			FPoint DeltaU;
-			FPoint DeltaV;
+			FVector DeltaU;
+			FVector DeltaV;
 			for (int32 Index = 0; Index < PointCount; Index++)
 			{
 				double* HPoint = HPoints.GetData() + Index * 4;
@@ -839,17 +839,19 @@ void EvaluatePointGrid(const FNURBSSurface& Nurbs, const FCoordinateGrid& Coords
 				DeltaV.Y = (VGradient[1] * HPoint[3] - HPoint[1] * VGradient[3]) / SquarePointH;
 				DeltaV.Z = (VGradient[2] * HPoint[3] - HPoint[2] * VGradient[3]) / SquarePointH;
 
-				OutPoints.Normals[Index] = DeltaU ^ DeltaV;
+				OutPoints.Normals[Index] = (FVector3f)(DeltaU ^ DeltaV);
 			}
 		}
 		else
 		{
-			for (int32 Index = 0; Index < PointCount; Index++)
+			double* UGradientsPtr = UGradients.GetData();
+			double* VGradientsPtr = VGradients.GetData();
+			for (int32 Index = 0; Index < PointCount; Index++, UGradientsPtr += 3, VGradientsPtr += 3)
 			{
-				FPoint UGradient = UGradients.GetData() + Index * 3;
-				FPoint VGradient = VGradients.GetData() + Index * 3;
+				const FVector UGradient(UGradientsPtr[0], UGradientsPtr[1], UGradientsPtr[2]);
+				const FVector VGradient(VGradientsPtr[0], VGradientsPtr[1], VGradientsPtr[2]);
 
-				OutPoints.Normals[Index] = UGradient ^ VGradient;
+				OutPoints.Normals[Index] = (FVector3f)(UGradient ^ VGradient);
 			}
 		}
 
@@ -1375,8 +1377,8 @@ void SampleNodalVector(double UMin, double UMax, TArray<double>& nodalVector, in
 	TabU->Add(UMax);
 }
 
-void DuplicateNurbsCurveWithHigherDegree(int32 degre, const TArray<FPoint>& poles, const TArray<double>& nodalVector, const TArray<double>& weights,
-	TArray<FPoint>& newPoles, TArray<double>& newNodalVector, TArray<double>& newWeights)
+void DuplicateNurbsCurveWithHigherDegree(int32 degre, const TArray<FVector>& poles, const TArray<double>& nodalVector, const TArray<double>& weights,
+	TArray<FVector>& newPoles, TArray<double>& newNodalVector, TArray<double>& newWeights)
 {
 	//Count number of polynomial segments and
 	//create new knot vector
@@ -1439,7 +1441,7 @@ void DuplicateNurbsCurveWithHigherDegree(int32 degre, const TArray<FPoint>& pole
 				}
 			}
 
-			TArray<FPoint> poles_;
+			TArray<FVector> poles_;
 			poles_.SetNum(newN);
 			TArray<double> weightsTmp;
 			weightsTmp.SetNum(newN);
@@ -1448,7 +1450,7 @@ void DuplicateNurbsCurveWithHigherDegree(int32 degre, const TArray<FPoint>& pole
 				Blossom(degre, poles, nodalVector, weights, seg, params[k], poles_[k], weightsTmp[k]);
 			}
 
-			FPoint pole(0, 0, 0);
+			FVector pole(0, 0, 0);
 			double w = 0;
 			for (int32 k = 0; k < newN; k++)
 			{
@@ -1502,12 +1504,12 @@ void insertKnotInKnots(TArray<double>& vn, double u, double* newU = nullptr)
 	Swap(vn, newKnots);
 }
 
-void InsertKnot(TArray<double>& vn, TArray<FPoint>& poles, TArray<double>& weights, double u, double* newU)
+void InsertKnot(TArray<double>& vn, TArray<FVector>& poles, TArray<double>& weights, double u, double* newU)
 {
 	int32 i = 0;
 
 	int32 n = (int32)vn.Num() - (int32)poles.Num() + 1; //Degree
-	TArray<FPoint> newPoles;
+	TArray<FVector> newPoles;
 	TArray<double> newWeights;
 
 	newPoles.Reserve(poles.Num() + vn.Num());
@@ -1525,7 +1527,7 @@ void InsertKnot(TArray<double>& vn, TArray<FPoint>& poles, TArray<double>& weigh
 			double c1 = ((vn[i + n] - u) / (vn[i + n] - vn[i]));
 			double c2 = ((u - vn[i]) / (vn[i + n] - vn[i]));
 			double w = w1 * c1 + w2 * c2;
-			FPoint p((poles[i][0] * w1 * c1 + poles[i + 1][0] * w2 * c2) / w,
+			FVector p((poles[i][0] * w1 * c1 + poles[i + 1][0] * w2 * c2) / w,
 				(poles[i][1] * w1 * c1 + poles[i + 1][1] * w2 * c2) / w,
 				(poles[i][2] * w1 * c1 + poles[i + 1][2] * w2 * c2) / w);
 			newPoles.Add(p);
@@ -1545,7 +1547,7 @@ void InsertKnot(TArray<double>& vn, TArray<FPoint>& poles, TArray<double>& weigh
 	insertKnotInKnots(vn, u, newU);
 }
 
-void insertKnotVInPatch(double v, int32 nbTime, TArray<double>& knotsV, TArray<TArray<FPoint>>& poles, TArray<TArray<double>>& weights)
+void insertKnotVInPatch(double v, int32 nbTime, TArray<double>& knotsV, TArray<TArray<FVector>>& poles, TArray<TArray<double>>& weights)
 {
 	TArray<double> knotsTmp;
 	int32 nbCurveV = (int32)poles.Num();
@@ -1560,18 +1562,18 @@ void insertKnotVInPatch(double v, int32 nbTime, TArray<double>& knotsV, TArray<T
 	knotsV = knotsTmp;
 }
 
-void insertKnotUInPatch(double u, int32 nbTime, TArray<double>& knotsU, TArray<TArray<FPoint>>& poles, TArray<TArray<double>>& weights)
+void insertKnotUInPatch(double u, int32 nbTime, TArray<double>& knotsU, TArray<TArray<FVector>>& poles, TArray<TArray<double>>& weights)
 {
 	int32 nbCurveU = (int32)poles[0].Num();
 	int32 nbPoleCurveU = (int32)poles.Num();
 	int32 i, j;
-	TArray<FPoint> polesCurveU;
+	TArray<FVector> polesCurveU;
 	TArray<double> weightsCurveU;
 	TArray<double> knotsTmp;
 
 	for (j = 0; j < nbTime; j++)
 	{
-		TArray<FPoint>& SubPoles = poles.Emplace_GetRef();
+		TArray<FVector>& SubPoles = poles.Emplace_GetRef();
 		SubPoles.SetNum(nbCurveU);
 		TArray<double>& SubWeights = weights.Emplace_GetRef();
 		SubWeights.Reserve(nbCurveU);
@@ -1602,7 +1604,7 @@ void insertKnotUInPatch(double u, int32 nbTime, TArray<double>& knotsU, TArray<T
 	knotsU = knotsTmp;
 }
 
-void rebound(TArray<double>& knots, TArray<FPoint>& poles, TArray<double>& weights, double u1, double u2)
+void rebound(TArray<double>& knots, TArray<FVector>& poles, TArray<double>& weights, double u1, double u2)
 {
 	// TD 21/09/2016 : corrige un crash à l'import de certaines pièces
 	if (u1 > u2)
@@ -1636,7 +1638,7 @@ void rebound(TArray<double>& knots, TArray<FPoint>& poles, TArray<double>& weigh
 	//Now trim the NURBS around these abscissas
 	TArray<double> tKnots;
 	tKnots.Append(knots.GetData() + u1Ind, u2Ind + 1 - u1Ind);
-	TArray<FPoint> tPoles;
+	TArray<FVector> tPoles;
 	tPoles.Append(poles.GetData() + u1Ind, u2Ind + 1 - u1Ind);
 	TArray<double> tWeights;
 	tWeights.Append(weights.GetData() + u1Ind, u2Ind + 1 - u1Ind);
@@ -1692,7 +1694,7 @@ TArray<double> homogeniseCurveKnots(TArray<TArray<double>>& cbKnots, double uMin
 	return knots;
 }
 
-void Blossom(int32 degre, const TArray<FPoint>& poles, const TArray<double>& nodalVector, const TArray<double>& weights, int32 seg, TArray<double>& params, FPoint& pnt, double& weight)
+void Blossom(int32 degre, const TArray<FVector>& poles, const TArray<double>& nodalVector, const TArray<double>& weights, int32 seg, TArray<double>& params, FVector& pnt, double& weight)
 {
 	int32 n = degre;
 	int32 i = 0;
@@ -1710,7 +1712,7 @@ void Blossom(int32 degre, const TArray<FPoint>& poles, const TArray<double>& nod
 	//	d(5,0) d(5,1) d(5,2) d(5,3)
 	//So with firstPole=2, recursion[4-firstPole][1]=d(4,1);
 
-	TArray<TArray<FPoint>> recursion;
+	TArray<TArray<FVector>> recursion;
 	recursion.SetNum(n + 1);	//We have n+1 lines
 	for (i = 0; i < n + 1; i++)
 	{
@@ -1773,7 +1775,7 @@ void Blossom(int32 degre, const TArray<FPoint>& poles, const TArray<double>& nod
 			//line goes from line to n
 			//Evaluate intermediate pole.
 			double dt = (endParam[l - 0][c - 1] - startParam[l - 1][c - 1]);
-			FPoint pole;
+			FVector pole;
 			double div = (params[c - 1] - startParam[l - 1][c - 1]) / dt;
 			pole[0] = recursion[l - 1][c - 1][0] * 1.0 + (-recursion[l - 1][c - 1][0] + recursion[l][c - 1][0]) * div;
 			pole[1] = recursion[l - 1][c - 1][1] * 1.0 + (-recursion[l - 1][c - 1][1] + recursion[l][c - 1][1]) * div;
diff --git a/Engine/Source/Runtime/Datasmith/CADKernel/Base/Private/Math/Geometry.cpp b/Engine/Source/Runtime/Datasmith/CADKernel/Base/Private/Math/Geometry.cpp
--- a/Engine/Source/Runtime/Datasmith/CADKernel/Base/Private/Math/Geometry.cpp
+++ b/Engine/Source/Runtime/Datasmith/CADKernel/Base/Private/Math/Geometry.cpp
@@ -23,9 +23,9 @@ struct FIntersectionContext
 {
 	const FSegment2D& SegmentAB;
 	const FSegment2D& SegmentCD;
-	const FPoint2D AB;
-	const FPoint2D CD;
-	const FPoint2D CA;
+	const FVector2d AB;
+	const FVector2d CD;
+	const FVector2d CA;
 
 	double NormAB = 0;
 	double NormCD = 0;
@@ -60,7 +60,7 @@ constexpr double MaxValue(bool OnlyInside)
 	return OnlyInside ? 1. - DOUBLE_KINDA_SMALL_NUMBER : 1. + DOUBLE_KINDA_SMALL_NUMBER;
 }
 
-bool ConfirmIntersectionWhenNearlyCoincident(const FPoint2D& AB, const FPoint2D& AC, const FPoint2D& AD, const double NormAB)
+bool ConfirmIntersectionWhenNearlyCoincident(const FVector2d& AB, const FVector2d& AC, const FVector2d& AD, const double NormAB)
 {
 	double HeightC = (AB ^ AC) / NormAB;
 	double HeightD = (AB ^ AD) / NormAB;
@@ -76,12 +76,12 @@ bool ConfirmIntersectionWhenNearlyCoincident(const FIntersectionContext& Context
 {
 	if (Context.NormAB > Context.NormCD)
 	{
-		const FPoint2D DA = Context.SegmentAB[0] - Context.SegmentCD[1];
+		const FVector2d DA = Context.SegmentAB[0] - Context.SegmentCD[1];
 		return ConfirmIntersectionWhenNearlyCoincident(Context.AB, Context.CA, DA, Context.NormAB);
 	}
 	else
 	{
-		const FPoint2D CB = Context.SegmentAB[1] - Context.SegmentCD[0];
+		const FVector2d CB = Context.SegmentAB[1] - Context.SegmentCD[0];
 		return ConfirmIntersectionWhenNearlyCoincident(Context.CD, Context.CA, CB, Context.NormCD);
 	}
 }
@@ -89,24 +89,24 @@ bool ConfirmIntersectionWhenNearlyCoincident(const FIntersectionContext& Context
 }
 
 
-double ComputeCurvature(const FPoint& Gradient, const FPoint& Laplacian)
+double ComputeCurvature(const FVector& Gradient, const FVector& Laplacian)
 {
-	const FPoint GradientCopy = Gradient.GetSafeNormal();
-	const FPoint LaplacianCopy = Laplacian.GetSafeNormal();
-	const FPoint Normal = GradientCopy ^ LaplacianCopy;
+	const FVector GradientCopy = Gradient.GetSafeNormal();
+	const FVector LaplacianCopy = Laplacian.GetSafeNormal();
+	const FVector Normal = GradientCopy ^ LaplacianCopy;
 	return (Normal.Length() * Laplacian.Length()) / (Gradient.Length() * Gradient.Length());
 }
 
-double ComputeCurvature(const FPoint& Normal, const FPoint& Gradient, const FPoint& Laplacian)
+double ComputeCurvature(const FVector& Normal, const FVector& Gradient, const FVector& Laplacian)
 {
-	const FPoint GradientCopy = Gradient.GetSafeNormal();
-	const FPoint LaplacianCopy = Laplacian.GetSafeNormal();
-	const FPoint NormalCopy = Normal.GetSafeNormal();
-	const FPoint Coef = (LaplacianCopy ^ GradientCopy) ^ NormalCopy;
+	const FVector GradientCopy = Gradient.GetSafeNormal();
+	const FVector LaplacianCopy = Laplacian.GetSafeNormal();
+	const FVector NormalCopy = Normal.GetSafeNormal();
+	const FVector Coef = (LaplacianCopy ^ GradientCopy) ^ NormalCopy;
 	return (Coef.Length() * Laplacian.Length()) / Gradient.SquaredLength();
 }
 
-void FindLoopIntersectionsWithIso(const EIso Iso, const double IsoParameter, const TArray<TArray<FPoint2D>>& Loops, TArray<double>& OutIntersections)
+void FindLoopIntersectionsWithIso(const EIso Iso, const double IsoParameter, const TArray<TArray<FVector2d>>& Loops, TArray<double>& OutIntersections)
 {
 	TArray<double> LocalIntersections;
 	LocalIntersections.Reserve(8);
@@ -114,7 +114,7 @@ void FindLoopIntersectionsWithIso(const EIso Iso, const double IsoParameter, con
 	const int32 UIndex = Iso == EIso::IsoU ? 0 : 1;
 	const int32 VIndex = Iso == EIso::IsoU ? 1 : 0;
 
-	TFunction<void(const FPoint2D&, const FPoint2D&)> ComputeIntersection = [&](const FPoint2D& Point1, const FPoint2D& Point2)
+	TFunction<void(const FVector2d&, const FVector2d&)> ComputeIntersection = [&](const FVector2d& Point1, const FVector2d& Point2)
 	{
 		if (IsoParameter > Point1[UIndex] && IsoParameter <= Point2[UIndex])
 		{
@@ -123,10 +123,10 @@ void FindLoopIntersectionsWithIso(const EIso Iso, const double IsoParameter, con
 		}
 	};
 
-	for (const TArray<FPoint2D>& Loop : Loops)
+	for (const TArray<FVector2d>& Loop : Loops)
 	{
-		const FPoint2D* Point1 = &Loop.Last();
-		for (const FPoint2D& Point2 : Loop)
+		const FVector2d* Point1 = &Loop.Last();
+		for (const FVector2d& Point2 : Loop)
 		{
 			if (!FMath::IsNearlyEqual((*Point1)[UIndex], Point2[UIndex]))
 			{
@@ -201,9 +201,9 @@ bool DoIntersect(const FSegment2D& SegmentAB, const FSegment2D& SegmentCD, TFunc
 	{
 		// double check with normalized vectors
 		{
-			FPoint2D NormalizedAB = Context.AB;
-			FPoint2D NormalizedCD = Context.CD;
-			FPoint2D NormalizedCA = Context.CA;
+			FVector2d NormalizedAB = Context.AB;
+			FVector2d NormalizedCD = Context.CD;
+			FVector2d NormalizedCA = Context.CA;
 
 			NormalizedAB.Normalize(Context.NormAB);
 			NormalizedCD.Normalize(Context.NormCD);
diff --git a/Engine/Source/Runtime/Datasmith/CADKernel/Base/Private/Math/MatrixH.cpp b/Engine/Source/Runtime/Datasmith/CADKernel/Base/Private/Math/MatrixH.cpp
--- a/Engine/Source/Runtime/Datasmith/CADKernel/Base/Private/Math/MatrixH.cpp
+++ b/Engine/Source/Runtime/Datasmith/CADKernel/Base/Private/Math/MatrixH.cpp
@@ -10,7 +10,7 @@ namespace UE::CADKernel
 
 const FMatrixH FMatrixH::Identity;
 
-void FMatrixH::BuildChangeOfCoordinateSystemMatrix(const FPoint& Xaxis, const FPoint& Yaxis, const FPoint& Zaxis, const FPoint& Origin)
+void FMatrixH::BuildChangeOfCoordinateSystemMatrix(const FVector& Xaxis, const FVector& Yaxis, const FVector& Zaxis, const FVector& Origin)
 {
 	Get(0, 0) = Xaxis[0];
 	Get(1, 0) = Xaxis[1];
@@ -33,21 +33,21 @@ void FMatrixH::BuildChangeOfCoordinateSystemMatrix(const FPoint& Xaxis, const FP
 	Get(3, 3) = 1.0;
 }
 
-void FMatrixH::FromAxisOrigin(const FPoint& Axis, const FPoint& Origin)
+void FMatrixH::FromAxisOrigin(const FVector& Axis, const FVector& Origin)
 {
-	FPoint Zaxis = FPoint(0, 1, 0);
-	FPoint Xaxis = Zaxis ^ Axis;
+	FVector Zaxis = FVector(0, 1, 0);
+	FVector Xaxis = Zaxis ^ Axis;
 
 	//on cherche le vecteur ox
 	Xaxis.Normalize();
 	if (FMath::Abs(Xaxis.Length()) < DOUBLE_SMALL_NUMBER)
 	{
-		Zaxis = FPoint(1, 0, 0);
+		Zaxis = FVector(1, 0, 0);
 		Xaxis = Axis ^ Zaxis;
 		Xaxis.Normalize();
 		if (FMath::Abs(Xaxis.Length()) < DOUBLE_SMALL_NUMBER)
 		{
-			Zaxis = FPoint(0, 0, 1);
+			Zaxis = FVector(0, 0, 1);
 			Xaxis = Axis ^ Zaxis;
 
 			Xaxis.Normalize();
@@ -55,7 +55,7 @@ void FMatrixH::FromAxisOrigin(const FPoint& Axis, const FPoint& Origin)
 		}
 	}
 
-	FPoint YAxis = Axis ^ Xaxis;
+	FVector YAxis = Axis ^ Xaxis;
 	Zaxis = Axis;
 
 	YAxis.Normalize();
@@ -64,7 +64,7 @@ void FMatrixH::FromAxisOrigin(const FPoint& Axis, const FPoint& Origin)
 	BuildChangeOfCoordinateSystemMatrix(Xaxis, YAxis, Zaxis, Origin);
 }
 
-FMatrixH FMatrixH::MakeRotationMatrix(double Angle, FPoint Axe)
+FMatrixH FMatrixH::MakeRotationMatrix(double Angle, FVector Axe)
 {
 	FMatrixH Matrix;
 	Matrix.SetIdentity();
@@ -86,7 +86,7 @@ FMatrixH FMatrixH::MakeRotationMatrix(double Angle, FPoint Axe)
 	return Matrix;
 }
 
-FMatrixH FMatrixH::MakeTranslationMatrix(const FPoint& Point)
+FMatrixH FMatrixH::MakeTranslationMatrix(const FVector& Point)
 {
 	FMatrixH Matrix;
 	Matrix.SetIdentity();
@@ -106,7 +106,7 @@ FMatrixH FMatrixH::MakeScaleMatrix(double XScale, double YScale, double ZScale)
 	return Matrix;
 }
 
-FMatrixH FMatrixH::MakeScaleMatrix(FPoint& Scale)
+FMatrixH FMatrixH::MakeScaleMatrix(FVector& Scale)
 {
 	return MakeScaleMatrix(Scale.X, Scale.Y, Scale.Z);
 }
diff --git a/Engine/Source/Runtime/Datasmith/CADKernel/Base/Private/Math/Point.cpp b/Engine/Source/Runtime/Datasmith/CADKernel/Base/Private/Math/Point.cpp
--- a/Engine/Source/Runtime/Datasmith/CADKernel/Base/Private/Math/Point.cpp
+++ b/Engine/Source/Runtime/Datasmith/CADKernel/Base/Private/Math/Point.cpp
@@ -10,11 +10,11 @@ namespace UE::CADKernel
 const FVector FVectorUtil::FarawayPoint3D(HUGE_VALUE, HUGE_VALUE, HUGE_VALUE);
 const FVector2d FVectorUtil::FarawayPoint2D(HUGE_VALUE, HUGE_VALUE);
 
-//double FPoint::SignedAngle(const FPoint & Other, const FPoint & Normal) const
+//double FVector::SignedAngle(const FVector & Other, const FVector & Normal) const
 //{
-//	FPoint Vector1 = *this; 
-//	FPoint Vector2 = Other; 
-//	FPoint Vector3 = Normal; 
+//	FVector Vector1 = *this; 
+//	FVector Vector2 = Other; 
+//	FVector Vector3 = Normal; 
 //
 //	Vector1.Normalize();
 //	Vector2.Normalize();
diff --git a/Engine/Source/Runtime/Datasmith/CADKernel/Base/Private/Mesh/Criteria/CriteriaGrid.cpp b/Engine/Source/Runtime/Datasmith/CADKernel/Base/Private/Mesh/Criteria/CriteriaGrid.cpp
--- a/Engine/Source/Runtime/Datasmith/CADKernel/Base/Private/Mesh/Criteria/CriteriaGrid.cpp
+++ b/Engine/Source/Runtime/Datasmith/CADKernel/Base/Private/Mesh/Criteria/CriteriaGrid.cpp
@@ -52,46 +52,46 @@ FCriteriaGrid::FCriteriaGrid(FTopologicalFace& InFace)
 
 void FCriteriaGrid::ComputeFaceMinMaxThicknessAlongIso()
 {
-	using FGetPointFunc = TFunction<const FPoint& (int32, int32)>;
+	using FGetPointFunc = TFunction<const FVector& (int32, int32)>;
 
 	TFunction<void(EIso, FGetPointFunc, FGetPointFunc, int32, int32)> ComputeFaceThicknessAlong = [this](EIso Iso, FGetPointFunc GetPointAlong, FGetPointFunc GetIntermediatePointAlong, int32 IsoUCount, int32 IsoVCount)
 	{
 		for (int32 VIndex = 0; VIndex < IsoVCount; ++VIndex)
 		{
-			const FPoint& StartPoint = GetPointAlong(0, VIndex);
-			const FPoint* PreviousPoint = &GetIntermediatePointAlong(0, VIndex);
+			const FVector& StartPoint = GetPointAlong(0, VIndex);
+			const FVector* PreviousPoint = &GetIntermediatePointAlong(0, VIndex);
 			double Length = FVector::Distance(StartPoint, *PreviousPoint);
 
 			for (int32 UIndex = 1; UIndex < IsoUCount - 1; ++UIndex)
 			{
-				const FPoint& Point = GetPointAlong(UIndex, VIndex);
+				const FVector& Point = GetPointAlong(UIndex, VIndex);
 
 				PreviousPoint = &GetIntermediatePointAlong(UIndex, VIndex);
 				Length += FVector::Distance(Point, *PreviousPoint);
 			}
-			const FPoint& LastPoint = GetPointAlong(IsoUCount - 1, VIndex);
+			const FVector& LastPoint = GetPointAlong(IsoUCount - 1, VIndex);
 			Length += FVector::Distance(LastPoint, *PreviousPoint);
 
 			this->FaceMinMax[Iso].ExtendTo(Length);
 		}
 	};
 
-	FGetPointFunc GetPointAlongU = [this](int32 UIndex, int32 VIndex) -> const FPoint&
+	FGetPointFunc GetPointAlongU = [this](int32 UIndex, int32 VIndex) -> const FVector&
 	{
 		return GetPoint(UIndex, VIndex);
 	};
 
-	FGetPointFunc GetPointAlongV = [this](int32 VIndex, int32 UIndex) -> const FPoint&
+	FGetPointFunc GetPointAlongV = [this](int32 VIndex, int32 UIndex) -> const FVector&
 	{
 		return GetPoint(UIndex, VIndex);
 	};
 
-	FGetPointFunc GetIntermediatePointAlongU = [this](int32 UIndex, int32 VIndex) -> const FPoint&
+	FGetPointFunc GetIntermediatePointAlongU = [this](int32 UIndex, int32 VIndex) -> const FVector&
 	{
 		return GetIntermediateU(UIndex, VIndex);
 	};
 
-	FGetPointFunc GetIntermediatePointAlongV = [this](int32 VIndex, int32 UIndex) -> const FPoint&
+	FGetPointFunc GetIntermediatePointAlongV = [this](int32 VIndex, int32 UIndex) -> const FVector&
 	{
 		return GetIntermediateV(UIndex, VIndex);
 	};
diff --git a/Engine/Source/Runtime/Datasmith/CADKernel/Base/Private/Mesh/Criteria/CriteriaGrid.h b/Engine/Source/Runtime/Datasmith/CADKernel/Base/Private/Mesh/Criteria/CriteriaGrid.h
--- a/Engine/Source/Runtime/Datasmith/CADKernel/Base/Private/Mesh/Criteria/CriteriaGrid.h
+++ b/Engine/Source/Runtime/Datasmith/CADKernel/Base/Private/Mesh/Criteria/CriteriaGrid.h
@@ -45,7 +45,7 @@ protected:
 		return TrueVIndex * CuttingCount[EIso::IsoU] + TrueUIndex;
 	}
 
-	const FPoint& GetPoint(int32 UIndex, int32 VIndex, bool bIsInternalU, bool bIsInternalV) const
+	const FVector& GetPoint(int32 UIndex, int32 VIndex, bool bIsInternalU, bool bIsInternalV) const
 	{
 		int32 Index = GetIndex(UIndex, VIndex, bIsInternalU, bIsInternalV);
 		ensureCADKernel(Points3D.IsValidIndex(Index));
@@ -64,22 +64,22 @@ public:
 	FCriteriaGrid(FTopologicalFace& InFace);
 
 
-	const FPoint& GetPoint(int32 iU, int32 iV) const
+	const FVector& GetPoint(int32 iU, int32 iV) const
 	{
 		return GetPoint(iU, iV, false, false);
 	}
 
-	const FPoint& GetIntermediateU(int32 iU, int32 iV) const
+	const FVector& GetIntermediateU(int32 iU, int32 iV) const
 	{
 		return GetPoint(iU, iV, true, false);
 	}
 
-	const FPoint& GetIntermediateV(int32 iU, int32 iV) const
+	const FVector& GetIntermediateV(int32 iU, int32 iV) const
 	{
 		return GetPoint(iU, iV, false, true);
 	}
 
-	const FPoint& GetIntermediateUV(int32 iU, int32 iV) const
+	const FVector& GetIntermediateUV(int32 iU, int32 iV) const
 	{
 		return GetPoint(iU, iV, true, true);
 	}
diff --git a/Engine/Source/Runtime/Datasmith/CADKernel/Base/Private/Mesh/Criteria/Criterion.cpp b/Engine/Source/Runtime/Datasmith/CADKernel/Base/Private/Mesh/Criteria/Criterion.cpp
--- a/Engine/Source/Runtime/Datasmith/CADKernel/Base/Private/Mesh/Criteria/Criterion.cpp
+++ b/Engine/Source/Runtime/Datasmith/CADKernel/Base/Private/Mesh/Criteria/Criterion.cpp
@@ -32,9 +32,9 @@ void FCriterion::ApplyOnEdgeParameters(FTopologicalEdge& Edge, const TArray<doub
 		const int32 MiddleIndex = PreviousIndex + Index;
 		const double DeltaU = Coordinates[Index] - Coordinates[PreviousIndex];
 
-		const FPoint Start = Points[2 * PreviousIndex].Point;
-		const FPoint End = Points[2 * Index].Point;
-		const FPoint Middle = Points[MiddleIndex].Point;
+		const FVector Start = Points[2 * PreviousIndex].Point;
+		const FVector End = Points[2 * Index].Point;
+		const FVector Middle = Points[MiddleIndex].Point;
 
 		double ChordLength;
 		const double SagFromStart = EvaluateSag(Start, End, Middle, ChordLength);
diff --git a/Engine/Source/Runtime/Datasmith/CADKernel/Base/Private/Mesh/Meshers/BowyerWatsonTriangulator.h b/Engine/Source/Runtime/Datasmith/CADKernel/Base/Private/Mesh/Meshers/BowyerWatsonTriangulator.h
--- a/Engine/Source/Runtime/Datasmith/CADKernel/Base/Private/Mesh/Meshers/BowyerWatsonTriangulator.h
+++ b/Engine/Source/Runtime/Datasmith/CADKernel/Base/Private/Mesh/Meshers/BowyerWatsonTriangulator.h
@@ -19,20 +19,20 @@ protected:
 	{
 		int32 VertexIndices[3];
 		double SquareRadius;
-		FPoint Center;
+		FVector Center;
 
-		FTriangle(const int32& Index0, const int32& Index1, const int32& Index2, const TArray<TPair<int32, FPoint2D>>& InVertices)
+		FTriangle(const int32& Index0, const int32& Index1, const int32& Index2, const TArray<TPair<int32, FVector2d>>& InVertices)
 		{
 			Set(Index0, Index1, Index2, InVertices);
 		}
 
-		void Set(const int32& Index0, const int32& Index1, const int32& Index2, const TArray<TPair<int32, FPoint2D>>& InVertices)
+		void Set(const int32& Index0, const int32& Index1, const int32& Index2, const TArray<TPair<int32, FVector2d>>& InVertices)
 		{
 			VertexIndices[0] = Index0;
 			VertexIndices[1] = Index1;
 			VertexIndices[2] = Index2;
 			FTriangle2D Triangle(InVertices[Index0].Value, InVertices[Index1].Value, InVertices[Index2].Value);
-			Center = Triangle.CircumCircleCenterWithSquareRadius(SquareRadius);
+			Center = FVector(Triangle.CircumCircleCenterWithSquareRadius(SquareRadius), 0.);
 		}
 	};
 
@@ -44,7 +44,7 @@ public:
 	 * @param OutEdgeVertices, the edges of the mesh. An edge is defined by the indices of its vertices
 	 * So the ith edge is defined by the vertices EdgeVertexIndices[2 * ith] and EdgeVertexIndices[2 * ith + 1]
 	 */
-	FBowyerWatsonTriangulator(TArray<TPair<int32, FPoint2D>>& InVertices, TArray<int32>& OutEdgeVertices)
+	FBowyerWatsonTriangulator(TArray<TPair<int32, FVector2d>>& InVertices, TArray<int32>& OutEdgeVertices)
 		: VerticesCount(InVertices.Num())
 		, Vertices(InVertices)
 		, EdgeVertexIndices(OutEdgeVertices)
@@ -59,13 +59,13 @@ public:
 		F3DDebugSession DelaunayDebugSession(bDisplay, TEXT("Delaunay Algo"));
 #endif // DEBUG_DELAUNAY
 
-		Vertices.Sort([](const TPair<int32, FPoint2D>& Vertex1, const TPair<int32, FPoint2D>& Vertex2) {return (Vertex1.Value.X + Vertex1.Value.Y) > (Vertex2.Value.X + Vertex2.Value.Y); });
+		Vertices.Sort([](const TPair<int32, FVector2d>& Vertex1, const TPair<int32, FVector2d>& Vertex2) {return (Vertex1.Value.X + Vertex1.Value.Y) > (Vertex2.Value.X + Vertex2.Value.Y); });
 
 #ifdef DEBUG_BOWYERWATSON
 		if (bDisplay)
 		{
 			F3DDebugSession _(TEXT("Vertices"));
-			for (const TPair<int32, FPoint2D>& Vertex : Vertices)
+			for (const TPair<int32, FVector2d>& Vertex : Vertices)
 			{
 				//F3DDebugSession _(TEXT("Vertex"));
 				DisplayPoint2DWithScale(Vertex.Value , EVisuProperty::YellowPoint, Vertex.Key);
@@ -103,7 +103,7 @@ public:
 			SelectedTriangleIndices.Reset(VerticesCount);
 			AdditionalTriangleIndices.Reset(VerticesCount);
 
-			const FPoint2D& NewVertex = Vertices[VertexIndex].Value;
+			const FVector2d& NewVertex = Vertices[VertexIndex].Value;
 
 #ifdef DEBUG_BOWYERWATSON_STEP
 			F3DDebugSession _(bDisplay, FString::Printf(TEXT("Step %d"), VIndex));
@@ -132,9 +132,9 @@ public:
 
 				for (int32 TriangleIndex = 0; TriangleIndex < TriangleSet.Num(); TriangleIndex++)
 				{
-					const FPoint2D Center = TriangleSet[TriangleIndex].Center;
+					const FVector2d Center(TriangleSet[TriangleIndex].Center.X, TriangleSet[TriangleIndex].Center.Y);
 
-					const double SquareDistanceToCenter = FPoint2D::DistSquared(Center, NewVertex);
+					const double SquareDistanceToCenter = FVector2d::DistSquared(Center, NewVertex);
 
 					const double SquareRadiusMax = TriangleSet[TriangleIndex].SquareRadius * IncreaseFactor;
 					const double SquareRadiusMin = SquareRadiusMax * ReducingFactor;
@@ -236,7 +236,7 @@ public:
 				TArray<TPair<int32, double>> BoundaryVertexToSlope;
 				for (int32 Index : BoundaryVertex)
 				{
-					FPoint2D Vertex = Vertices[Index].Value;
+					FVector2d Vertex = Vertices[Index].Value;
 					double Slope = ComputeSlope(NewVertex, Vertex);
 					BoundaryVertexToSlope.Emplace(Index, Slope);
 				}
@@ -277,7 +277,7 @@ public:
 							continue;
 						}
 
-						FPoint2D CandidateVertex = Vertices[CandidateVertexIndex].Value;
+						FVector2d CandidateVertex = Vertices[CandidateVertexIndex].Value;
 
 #ifdef DEBUG_BOWYERWATSON_STEP_2
 						if (VIndex == StopIndex)
@@ -310,11 +310,11 @@ public:
 
 						StartIndex = EndIndex == 0 ? BoundaryVertexToSlope.Num() - 1 : EndIndex - 1;
 
-						FPoint2D StartPoint = Vertices[BoundaryVertexToSlope[StartIndex].Key].Value;
+						FVector2d StartPoint = Vertices[BoundaryVertexToSlope[StartIndex].Key].Value;
 						double SlopeNewStartPlusPi = ComputeSlope(NewVertex, StartPoint) + Slope::RightSlope;
 						double SlopeStartCandidate = ComputePositiveSlope(StartPoint, CandidateVertex, SlopeNewStartPlusPi);
 
-						FPoint2D EndPoint = Vertices[BoundaryVertexToSlope[EndIndex].Key].Value;
+						FVector2d EndPoint = Vertices[BoundaryVertexToSlope[EndIndex].Key].Value;
 						double SlopeNewEndMinusPi = ComputeSlope(NewVertex, EndPoint) - Slope::RightSlope;
 						double SlopeEndCandidate = ComputePositiveSlope(EndPoint, CandidateVertex, SlopeNewEndMinusPi);
 
@@ -579,7 +579,7 @@ public:
 private:
 	int32 VerticesCount;
 
-	TArray<TPair<int32, FPoint2D>>& Vertices;
+	TArray<TPair<int32, FVector2d>>& Vertices;
 
 	/**
 	 * An edge is defined by the indices of its vertices
@@ -602,7 +602,7 @@ private:
 	void MakeBoundingMesh()
 	{
 		FAABB2D VerticesBBox;
-		for (const TPair<int32, FPoint2D>& Vertex : Vertices)
+		for (const TPair<int32, FVector2d>& Vertex : Vertices)
 		{
 			VerticesBBox += Vertex.Value;
 		}
@@ -718,11 +718,11 @@ private:
 		return true;
 	}
 
-	bool DoesTriangleContainingNewVertex(const FPoint2D& NewVertex, int32 CandiadateTriangle)
+	bool DoesTriangleContainingNewVertex(const FVector2d& NewVertex, int32 CandiadateTriangle)
 	{
-		const FPoint2D& Point0 = Vertices[TriangleSet[CandiadateTriangle].VertexIndices[0]].Value;
-		const FPoint2D& Point1 = Vertices[TriangleSet[CandiadateTriangle].VertexIndices[1]].Value;
-		const FPoint2D& Point2 = Vertices[TriangleSet[CandiadateTriangle].VertexIndices[2]].Value;
+		const FVector2d& Point0 = Vertices[TriangleSet[CandiadateTriangle].VertexIndices[0]].Value;
+		const FVector2d& Point1 = Vertices[TriangleSet[CandiadateTriangle].VertexIndices[1]].Value;
+		const FVector2d& Point2 = Vertices[TriangleSet[CandiadateTriangle].VertexIndices[2]].Value;
 		double Slope0 = ComputePositiveSlope(Point0, Point1, NewVertex);
 		double Slope1 = ComputePositiveSlope(Point1, Point2, NewVertex);
 		double Slope2 = ComputePositiveSlope(Point2, Point0, NewVertex);
@@ -733,7 +733,7 @@ private:
 		return false;
 	}
 
-	void AddTrianglesContainingNewVertexToSelectedTriangles(const FPoint2D& NewVertex, TArray<int32>& CandiadateTriangles)
+	void AddTrianglesContainingNewVertexToSelectedTriangles(const FVector2d& NewVertex, TArray<int32>& CandiadateTriangles)
 	{
 		for (int32& TriangleIndex : CandiadateTriangles)
 		{
@@ -900,17 +900,17 @@ private:
 	{
 		F3DDebugSession A(TEXT("Selected Triangles Boundary"));
 
-		FPoint2D Previous = Vertices[VertexToSlop[VertexToSlop.Num() - 1].Key].Value;
+		FVector2d Previous = Vertices[VertexToSlop[VertexToSlop.Num() - 1].Key].Value;
 		for (int Index = 0; Index < VertexToSlop.Num(); ++Index)
 		{
-			FPoint2D VertexPoint = Vertices[VertexToSlop[Index].Key].Value;
+			FVector2d VertexPoint = Vertices[VertexToSlop[Index].Key].Value;
 			DisplayPoint(VertexPoint * DisplayScale, EVisuProperty::GreenPoint);
 			DisplaySegment(VertexPoint * DisplayScale, Previous * DisplayScale, 0, EVisuProperty::GreenCurve);
 			Previous = VertexPoint;
 		}
 	}
 
-	void DisplayLocalVoronoiDiagram(const FPoint2D& CandidateVertex, const FPoint2D& NewVertex, const FPoint2D& StartPoint, const FPoint2D& EndPoint)
+	void DisplayLocalVoronoiDiagram(const FVector2d& CandidateVertex, const FVector2d& NewVertex, const FVector2d& StartPoint, const FVector2d& EndPoint)
 	{
 		F3DDebugSession B(TEXT("Voronoi Diagram"));
 		{
@@ -925,11 +925,11 @@ private:
 			F3DDebugSession A(TEXT("Start"));
 			DisplayPoint(StartPoint * DisplayScale, EVisuProperty::YellowPoint);
 
-			FPoint2D Normal = NewVertex - StartPoint;
+			FVector2d Normal = NewVertex - StartPoint;
 			constexpr double HalfPi = -DOUBLE_PI / 2.;
 			Normal.Normalize();
-			FPoint2D Tangent = Normal.Rotate(HalfPi) * 1000.;
-			FPoint2D TangentPoint = StartPoint + Tangent;
+			FVector2d Tangent = Normal.Rotate(HalfPi) * 1000.;
+			FVector2d TangentPoint = StartPoint + Tangent;
 
 			DisplaySegment(StartPoint * DisplayScale, TangentPoint * DisplayScale, 1, EVisuProperty::YellowCurve);
 			DisplaySegment(NewVertex * DisplayScale, StartPoint * DisplayScale, 1, EVisuProperty::YellowCurve);
@@ -944,11 +944,11 @@ private:
 			F3DDebugSession A(TEXT("End"));
 			DisplayPoint(EndPoint * DisplayScale, EVisuProperty::OrangePoint);
 
-			FPoint2D Normal = NewVertex - EndPoint;
+			FVector2d Normal = NewVertex - EndPoint;
 			constexpr double HalfPi = DOUBLE_PI / 2.;
 			Normal.Normalize();
-			FPoint2D Tangent = Normal.Rotate(HalfPi) * 1000.;
-			FPoint2D TangentPoint = EndPoint + Tangent;
+			FVector2d Tangent = Normal.Rotate(HalfPi) * 1000.;
+			FVector2d TangentPoint = EndPoint + Tangent;
 
 			DisplaySegment(EndPoint * DisplayScale, TangentPoint * DisplayScale, 1, EVisuProperty::OrangeCurve);
 			DisplaySegment(NewVertex * DisplayScale, EndPoint * DisplayScale, 1, EVisuProperty::OrangeCurve);
diff --git a/Engine/Source/Runtime/Datasmith/CADKernel/Base/Private/Mesh/Meshers/CycleTriangulator.cpp b/Engine/Source/Runtime/Datasmith/CADKernel/Base/Private/Mesh/Meshers/CycleTriangulator.cpp
--- a/Engine/Source/Runtime/Datasmith/CADKernel/Base/Private/Mesh/Meshers/CycleTriangulator.cpp
+++ b/Engine/Source/Runtime/Datasmith/CADKernel/Base/Private/Mesh/Meshers/CycleTriangulator.cpp
@@ -130,8 +130,8 @@ bool FCycleTriangulator::FindTheCycleToMesh(FIsoSegment* Segment, bool bOrientat
 	FIsoNode* EndNode = bOrientation ? &Segment->GetSecondNode() : &Segment->GetFirstNode();
 	FIsoNode* Node = EndNode;
 
-	const FPoint2D* PreviousPoint2D = &PreviousNode->Get2DPoint(EGridSpace::UniformScaled, Grid);
-	const FPoint2D* NodePoint2D = &Node->Get2DPoint(EGridSpace::UniformScaled, Grid);
+	const FVector2d* PreviousPoint2D = &PreviousNode->Get2DPoint(EGridSpace::UniformScaled, Grid);
+	const FVector2d* NodePoint2D = &Node->Get2DPoint(EGridSpace::UniformScaled, Grid);
 
 	double PreviousSlop = ComputeSlope(*NodePoint2D, *PreviousPoint2D);
 
@@ -145,7 +145,7 @@ bool FCycleTriangulator::FindTheCycleToMesh(FIsoSegment* Segment, bool bOrientat
 	for (FIsoSegment* NextSegment = FindNextSegment(Segment, Node, GetSlopeAtEndNode);; NextSegment = FindNextSegment(NextSegment, Node, GetSlopeAtEndNode))
 	{
 		FIsoNode* NextNode = &NextSegment->GetOtherNode(Node);
-		const FPoint2D& NextPoint2D = NextNode->Get2DPoint(EGridSpace::UniformScaled, Grid);
+		const FVector2d& NextPoint2D = NextNode->Get2DPoint(EGridSpace::UniformScaled, Grid);
 
 		double NextSlop = ComputeSlope(*NodePoint2D, NextPoint2D);
 		double NodeSlop = TransformIntoPositiveSlope(PreviousSlop - NextSlop);
@@ -197,7 +197,7 @@ bool FCycleTriangulator::FindTheCycleToMesh(FIsoSegment* Segment, bool bOrientat
 	return true;
 }
 
-bool IsCycleInsideCandidate(const FPoint2D& Candidate, const FPoint2D& Start, const FPoint2D& End, const FPoint2D& Next, TFunction<double(double)> ConvertSlope)
+bool IsCycleInsideCandidate(const FVector2d& Candidate, const FVector2d& Start, const FVector2d& End, const FVector2d& Next, TFunction<double(double)> ConvertSlope)
 {
 	const double StartEndSlop = ComputeSlope(Start, End);
 	const double StartCandidateSlope = ComputeSlope(Start, Candidate);
@@ -237,18 +237,18 @@ bool FCycleTriangulator::FindTheBestAcuteTriangle()
 
 	int32 Index = -1;
 	FIsoNode* CandidateNode = nullptr;
-	const FPoint2D* CandidatePoint = nullptr;
+	const FVector2d* CandidatePoint = nullptr;
 
 	int32 StartId = -1;
 	FIsoNode* NewSegmentStart = nullptr;
-	const FPoint2D* StartPoint = nullptr;
+	const FVector2d* StartPoint = nullptr;
 
 	int32 EndId = -1;
 	FIsoNode* NewSegmentEnd = nullptr;
-	const FPoint2D* EndPoint = nullptr;
+	const FVector2d* EndPoint = nullptr;
 
-	const FPoint2D* NextPoint;
-	const FPoint2D* PrevPoint;
+	const FVector2d* NextPoint;
+	const FVector2d* PrevPoint;
 
 	TFunction<void()> SetFinalCandidate = [&]()
 	{
@@ -365,8 +365,8 @@ bool FCycleTriangulator::FindCandidateNodes(int32 StartIndex)
 	FirstSideStartNode = SubCycleNodes[FirstSideStartIndex];
 	FirstSideEndNode = SubCycleNodes[FirstSideEndIndex];
 
-	const FPoint2D& StartPoint2D = FirstSideStartNode->Get2DPoint(EGridSpace::UniformScaled, Grid);
-	const FPoint2D& EndPoint2D = FirstSideEndNode->Get2DPoint(EGridSpace::UniformScaled, Grid);
+	const FVector2d& StartPoint2D = FirstSideStartNode->Get2DPoint(EGridSpace::UniformScaled, Grid);
+	const FVector2d& EndPoint2D = FirstSideEndNode->Get2DPoint(EGridSpace::UniformScaled, Grid);
 
 	// StartMaxSlope and EndMaxSlope are at most equal to 4, because if the slop with candidate node is biggest to 4, the new triangle will be inverted
 	double StartReferenceSlope = ComputePositiveSlope(StartPoint2D, EndPoint2D, 0);
@@ -381,7 +381,7 @@ bool FCycleTriangulator::FindCandidateNodes(int32 StartIndex)
 
 		FIsoNode* Node = SubCycleNodes[Index];
 
-		const FPoint2D& NodePoint2D = Node->Get2DPoint(EGridSpace::UniformScaled, Grid);
+		const FVector2d& NodePoint2D = Node->Get2DPoint(EGridSpace::UniformScaled, Grid);
 		double SlopeAtStartNode = GetSlopeAtStartNode(StartPoint2D, NodePoint2D, StartReferenceSlope);
 		double SlopeAtEndNode = GetSlopeAtEndNode(EndPoint2D, NodePoint2D, EndReferenceSlope);
 
@@ -482,9 +482,9 @@ bool FCycleTriangulator::FindTheBestCandidateNode()
 	bool bIsDegeneratedDueToLength = false;
 	int32 CandidateIntersectionCount = 0;
 
-	const FPoint2D& StartPoint2D = FirstSideStartNode->Get2DPoint(EGridSpace::UniformScaled, Grid);
-	const FPoint2D& EndPoint2D = FirstSideEndNode->Get2DPoint(EGridSpace::UniformScaled, Grid);
-	double NewSegmentSquareLength2500 = FPoint2D::DistSquared(StartPoint2D, EndPoint2D) * 2500; // need to check if the length of the sides are not disproportionate
+	const FVector2d& StartPoint2D = FirstSideStartNode->Get2DPoint(EGridSpace::UniformScaled, Grid);
+	const FVector2d& EndPoint2D = FirstSideEndNode->Get2DPoint(EGridSpace::UniformScaled, Grid);
+	double NewSegmentSquareLength2500 = FVector2d::DistSquared(StartPoint2D, EndPoint2D) * 2500; // need to check if the length of the sides are not disproportionate
 
 #ifdef DEBUG_FIND_THE_BEST_CANDIDATE_NODE
 	static int32 BestCandidateNode = 0;
@@ -540,7 +540,7 @@ bool FCycleTriangulator::FindTheBestCandidateNode()
 		FCandidateNode* CNode = NodeIntersection.Key;
 		FIsoNode* Node = CNode->Node;
 
-		const FPoint2D& NodePoint2D = Node->Get2DPoint(EGridSpace::UniformScaled, Grid);
+		const FVector2d& NodePoint2D = Node->Get2DPoint(EGridSpace::UniformScaled, Grid);
 
 		double PointCriteria = IsoTriangulatorImpl::IsoscelesCriteriaMax(StartPoint2D, EndPoint2D, NodePoint2D);
 
@@ -554,8 +554,8 @@ bool FCycleTriangulator::FindTheBestCandidateNode()
 		}
 #endif
 
-		double CEndSquareDistance = FPoint2D::DistSquared(StartPoint2D, NodePoint2D);
-		double CStartSquareDistance = FPoint2D::DistSquared(EndPoint2D, NodePoint2D);
+		double CEndSquareDistance = FVector2d::DistSquared(StartPoint2D, NodePoint2D);
+		double CStartSquareDistance = FVector2d::DistSquared(EndPoint2D, NodePoint2D);
 
 		// if the previous candidate and the current candidate is Degenerated Due To Length, the previous is preferred
 		if (bIsDegeneratedDueToLength && IsDegeneratedDueToLength(CEndSquareDistance, CStartSquareDistance))
@@ -669,10 +669,10 @@ bool FCycleTriangulator::FindTheBestCandidateNode()
 	};
 #endif
 
-	TFunction<void(FCandidateNode*, const double, const FPoint2D&, double, double) > CompareCandidateAtSegmentNode = [&CandidateNode, &Grid = Grid](FCandidateNode* NewCandidate, const double CloseLimit, const FPoint2D& SegmentPoint2D, double CandidateToSegmentSquareDistance, double SlopeDiffAtSegmentExtremity)
+	TFunction<void(FCandidateNode*, const double, const FVector2d&, double, double) > CompareCandidateAtSegmentNode = [&CandidateNode, &Grid = Grid](FCandidateNode* NewCandidate, const double CloseLimit, const FVector2d& SegmentPoint2D, double CandidateToSegmentSquareDistance, double SlopeDiffAtSegmentExtremity)
 	{
-		const FPoint2D& NodePoint2D = NewCandidate->Node->Get2DPoint(EGridSpace::UniformScaled, Grid);
-		const double CFirstSquareDistance = FPoint2D::DistSquared(SegmentPoint2D, NodePoint2D) * 2.;
+		const FVector2d& NodePoint2D = NewCandidate->Node->Get2DPoint(EGridSpace::UniformScaled, Grid);
+		const double CFirstSquareDistance = FVector2d::DistSquared(SegmentPoint2D, NodePoint2D) * 2.;
 
 		if (SlopeDiffAtSegmentExtremity < CloseLimit && CFirstSquareDistance < CandidateToSegmentSquareDistance)
 		{
@@ -761,8 +761,8 @@ void FCycleTriangulator::ValidateAddNodesAccordingSlopeWithSide(FAdditionalIso&
 {
 	constexpr double TooSmallSlope = Slope::FiveDegree;
 
-	const FPoint2D& StartPoint = Side.StartNode->Get2DPoint(EGridSpace::UniformScaled, Grid);
-	const FPoint2D& EndPoint = Side.EndNode->Get2DPoint(EGridSpace::UniformScaled, Grid);
+	const FVector2d& StartPoint = Side.StartNode->Get2DPoint(EGridSpace::UniformScaled, Grid);
+	const FVector2d& EndPoint = Side.EndNode->Get2DPoint(EGridSpace::UniformScaled, Grid);
 	const double SlopeStartEnd = ComputeSlope(StartPoint, EndPoint);
 
 	for (int32 Index = 0; Index < 2; ++Index)
@@ -772,7 +772,7 @@ void FCycleTriangulator::ValidateAddNodesAccordingSlopeWithSide(FAdditionalIso&
 		{
 			continue;
 		}
-		const FPoint2D& CandidatePoint = Candidate->Get2DPoint(EGridSpace::UniformScaled, Grid);
+		const FVector2d& CandidatePoint = Candidate->Get2DPoint(EGridSpace::UniformScaled, Grid);
 		const double SlopeStartCandidate = ComputeSlope(StartPoint, CandidatePoint);
 		const double SlopeEndCandidate = ComputeSlope(EndPoint, CandidatePoint);
 
@@ -978,12 +978,12 @@ bool FCycleTriangulator::ValidComplementaryNodeOrDeleteIt(FAdditionalIso& Side,
 		return CandidateIndex ? NextIndex(Index) : PreviousIndex(Index);
 	};
 
-	const FPoint2D& CandidatePoint = Candidate->Get2DPoint(EGridSpace::UniformScaled, Grid);
-	const FPoint2D& OppositePoint = OppositeNode->Get2DPoint(EGridSpace::UniformScaled, Grid);
+	const FVector2d& CandidatePoint = Candidate->Get2DPoint(EGridSpace::UniformScaled, Grid);
+	const FVector2d& OppositePoint = OppositeNode->Get2DPoint(EGridSpace::UniformScaled, Grid);
 
 	const int32 TestPointIndex = NextFunc(OppositeIndex);
 	const FIsoNode* TestNode = SubCycleNodes[TestPointIndex];
-	const FPoint2D& TestPoint = TestNode->Get2DPoint(EGridSpace::UniformScaled, Grid);
+	const FVector2d& TestPoint = TestNode->Get2DPoint(EGridSpace::UniformScaled, Grid);
 
 #ifdef DEBUG_VALID_COMPLEMENTARY_NODE
 	if (Grid.bDisplay)
@@ -1030,11 +1030,11 @@ bool FCycleTriangulator::IsInnerSideSegmentInsideCycle(FAdditionalIso& Side)
 {
 	TFunction<bool(const FIsoNode*, const FIsoNode*, const FIsoNode*, const FIsoNode*)> IsInside = [&](const FIsoNode* P, const FIsoNode* A, const FIsoNode* N, const FIsoNode* B) -> bool
 	{
-		const FPoint2D& PPoint = P->Get2DPoint(EGridSpace::UniformScaled, Grid);
-		const FPoint2D& NPoint = N->Get2DPoint(EGridSpace::UniformScaled, Grid);
+		const FVector2d& PPoint = P->Get2DPoint(EGridSpace::UniformScaled, Grid);
+		const FVector2d& NPoint = N->Get2DPoint(EGridSpace::UniformScaled, Grid);
 
-		const FPoint2D& APoint = A->Get2DPoint(EGridSpace::UniformScaled, Grid);
-		const FPoint2D& BPoint = B->Get2DPoint(EGridSpace::UniformScaled, Grid);
+		const FVector2d& APoint = A->Get2DPoint(EGridSpace::UniformScaled, Grid);
+		const FVector2d& BPoint = B->Get2DPoint(EGridSpace::UniformScaled, Grid);
 
 		double SlopeAP = ComputeSlope(APoint, PPoint);
 		double SlopeAN = ComputeSlope(APoint, NPoint);
@@ -1179,10 +1179,10 @@ void FCycleTriangulator::ComputeSideCandidateEquilateralCriteria(FAdditionalIso&
 	if (Side.CandidateNodeCount() == 1)
 	{
 		int32 Index = Side.Nodes[0] ? 0 : 1;
-		const FPoint2D& SideBStartPoint = Side.StartNode->Get2DPoint(EGridSpace::UniformScaled, Grid);;
-		const FPoint2D& SideBEndPoint = Side.EndNode->Get2DPoint(EGridSpace::UniformScaled, Grid);;
+		const FVector2d& SideBStartPoint = Side.StartNode->Get2DPoint(EGridSpace::UniformScaled, Grid);;
+		const FVector2d& SideBEndPoint = Side.EndNode->Get2DPoint(EGridSpace::UniformScaled, Grid);;
 		const FIsoNode* SideBCandidate = Side.Nodes[Index];
-		const FPoint2D& SideBCandidatePoint = SideBCandidate->Get2DPoint(EGridSpace::UniformScaled, Grid);
+		const FVector2d& SideBCandidatePoint = SideBCandidate->Get2DPoint(EGridSpace::UniformScaled, Grid);
 
 		const double SideBCriteria = IsoTriangulatorImpl::EquilateralSlopeCriteria(SideBStartPoint, SideBCandidatePoint, SideBEndPoint);
 		Side.EquilateralCriteria[Index] = SideBCriteria;
@@ -1194,10 +1194,10 @@ void FCycleTriangulator::ComputeSideCandidateEquilateralCriteria(FAdditionalIso&
 		const FIsoNode* Candidate0 = Side.Nodes[0];
 		const FIsoNode* Candidate1 = Side.Nodes[1];
 
-		const FPoint2D& StartPoint = StartNode->Get2DPoint(EGridSpace::UniformScaled, Grid);
-		const FPoint2D& EndPoint = EndNode->Get2DPoint(EGridSpace::UniformScaled, Grid);
-		const FPoint2D& Candidate0Point = Candidate0->Get2DPoint(EGridSpace::UniformScaled, Grid);
-		const FPoint2D& Candidate1Point = Candidate1->Get2DPoint(EGridSpace::UniformScaled, Grid);
+		const FVector2d& StartPoint = StartNode->Get2DPoint(EGridSpace::UniformScaled, Grid);
+		const FVector2d& EndPoint = EndNode->Get2DPoint(EGridSpace::UniformScaled, Grid);
+		const FVector2d& Candidate0Point = Candidate0->Get2DPoint(EGridSpace::UniformScaled, Grid);
+		const FVector2d& Candidate1Point = Candidate1->Get2DPoint(EGridSpace::UniformScaled, Grid);
 
 		Side.EquilateralCriteria[0] = IsoTriangulatorImpl::EquilateralSlopeCriteria(StartPoint, Candidate0Point, Candidate1Point);
 		Side.EquilateralCriteria[1] = IsoTriangulatorImpl::EquilateralSlopeCriteria(Candidate0Point, Candidate1Point, EndPoint);
@@ -1397,7 +1397,7 @@ void MeshTriangle(const FGrid& Grid, FIsoNode** Nodes, FFaceMesh& Mesh)
 
 void MeshQuadrilateral(const FGrid& Grid, FIsoNode** Nodes, FFaceMesh& Mesh)
 {
-	const FPoint2D* NodeCoordinates[4];
+	const FVector2d* NodeCoordinates[4];
 	for (int32 Index = 0; Index < 4; ++Index)
 	{
 		NodeCoordinates[Index] = &Nodes[Index]->Get2DPoint(EGridSpace::UniformScaled, Grid);
@@ -1448,7 +1448,7 @@ void MeshQuadrilateral(const FGrid& Grid, FIsoNode** Nodes, FFaceMesh& Mesh)
 void MeshPentagon(const FGrid& Grid, FIsoNode** Nodes, FFaceMesh& Mesh)
 {
 	double TriangleSlopes[10];
-	const FPoint2D* NodeCoordinates[5];
+	const FVector2d* NodeCoordinates[5];
 	for (int32 Index = 0; Index < 5; ++Index)
 	{
 		NodeCoordinates[Index] = &Nodes[Index]->Get2DPoint(EGridSpace::UniformScaled, Grid);
@@ -1667,8 +1667,8 @@ bool FCycleTriangulator::CanCycleBeMeshed()
 
 FIsoSegment* FCycleTriangulator::FindNextSegment(const FIsoSegment* StartSegment, const FIsoNode* StartNode, SlopeMethod GetSlope) const
 {
-	const FPoint2D& StartPoint = StartNode->Get2DPoint(Space, Grid);
-	const FPoint2D& EndPoint = (StartSegment->IsFirstNode(StartNode)) ? StartSegment->GetSecondNode().Get2DPoint(Space, Grid) : StartSegment->GetFirstNode().Get2DPoint(Space, Grid);
+	const FVector2d& StartPoint = StartNode->Get2DPoint(Space, Grid);
+	const FVector2d& EndPoint = (StartSegment->IsFirstNode(StartNode)) ? StartSegment->GetSecondNode().Get2DPoint(Space, Grid) : StartSegment->GetFirstNode().Get2DPoint(Space, Grid);
 
 #ifdef DEBUG_FIND_NEXTSEGMENT
 	F3DDebugSession G(bDisplayStar, TEXT("FindNextSegment"));
@@ -1689,7 +1689,7 @@ FIsoSegment* FCycleTriangulator::FindNextSegment(const FIsoSegment* StartSegment
 
 	for (FIsoSegment* Segment : StartNode->GetConnectedSegments())
 	{
-		const FPoint2D& OtherPoint = (Segment->IsFirstNode(StartNode)) ? Segment->GetSecondNode().Get2DPoint(Space, Grid) : Segment->GetFirstNode().Get2DPoint(Space, Grid);
+		const FVector2d& OtherPoint = (Segment->IsFirstNode(StartNode)) ? Segment->GetSecondNode().Get2DPoint(Space, Grid) : Segment->GetFirstNode().Get2DPoint(Space, Grid);
 
 		double Slope = GetSlope(StartPoint, OtherPoint, ReferenceSlope);
 		if (Slope < SMALL_NUMBER_SQUARE) Slope = 8;
diff --git a/Engine/Source/Runtime/Datasmith/CADKernel/Base/Private/Mesh/Meshers/CycleTriangulator.h b/Engine/Source/Runtime/Datasmith/CADKernel/Base/Private/Mesh/Meshers/CycleTriangulator.h
--- a/Engine/Source/Runtime/Datasmith/CADKernel/Base/Private/Mesh/Meshers/CycleTriangulator.h
+++ b/Engine/Source/Runtime/Datasmith/CADKernel/Base/Private/Mesh/Meshers/CycleTriangulator.h
@@ -20,7 +20,7 @@ class FIntersectionSegmentTool;
 class FIsoSegment;
 class FIsoTriangulator;
 
-typedef TFunction<double(const FPoint2D&, const FPoint2D&, double)> SlopeMethod;
+typedef TFunction<double(const FVector2d&, const FVector2d&, double)> SlopeMethod;
 typedef TFunction<int32(int32)> NextIndexMethod;
 
 struct FCandidateNode
