Replaced FPoint and FPoint2D by FVector and  FVector2d and remove unused FFPoint and FPointH classes - Final

diff --git a/Engine/Source/Runtime/Datasmith/CADKernel/Base/Private/Mesh/Meshers/IsoTriangulator.cpp b/Engine/Source/Runtime/Datasmith/CADKernel/Base/Private/Mesh/Meshers/IsoTriangulator.cpp
--- a/Engine/Source/Runtime/Datasmith/CADKernel/Base/Private/Mesh/Meshers/IsoTriangulator.cpp
+++ b/Engine/Source/Runtime/Datasmith/CADKernel/Base/Private/Mesh/Meshers/IsoTriangulator.cpp
@@ -51,7 +51,7 @@ struct FCandidateSegment
 		: StartNode(Node1)
 		, EndNode(Node2)
 	{
-		Length = FPoint2D::Distance(Node1.Get2DPoint(EGridSpace::UniformScaled, Grid), Node2.Get2DPoint(EGridSpace::UniformScaled, Grid));
+		Length = FVector2d::Distance(Node1.Get2DPoint(EGridSpace::UniformScaled, Grid), Node2.Get2DPoint(EGridSpace::UniformScaled, Grid));
 	}
 };
 
@@ -216,7 +216,7 @@ void FIsoTriangulator::BuildNodes()
 	FTimePoint StartTime = FChrono::Now();
 
 	LoopNodeCount = 0;
-	for (const TArray<FPoint2D>& LoopPoints : Grid.GetLoops2D(EGridSpace::Default2D))
+	for (const TArray<FVector2d>& LoopPoints : Grid.GetLoops2D(EGridSpace::Default2D))
 	{
 		LoopNodeCount += (int32)LoopPoints.Num();
 	}
@@ -226,7 +226,7 @@ void FIsoTriangulator::BuildNodes()
 	// Loop nodes
 	int32 NodeIndex = 0;
 	int32 LoopIndex = 0;
-	for (const TArray<FPoint2D>& LoopPoints : Grid.GetLoops2D(EGridSpace::Default2D))
+	for (const TArray<FVector2d>& LoopPoints : Grid.GetLoops2D(EGridSpace::Default2D))
 	{
 		LoopStartIndex.Add(LoopNodeCount);
 		const TArray<int32>& LoopIds = Grid.GetNodeIdsOfFaceLoops()[LoopIndex];
@@ -269,7 +269,7 @@ void FIsoTriangulator::FillMeshNodes()
 	int32 TriangleNum = 50 + (int32)((2 * InnerNodeCount + LoopNodeCount) * 1.1);
 	Mesh.Init(TriangleNum, InnerNodeCount + LoopNodeCount);
 
-	TArray<FPoint>& InnerNodeCoordinates = Mesh.GetNodeCoordinates();
+	TArray<FVector>& InnerNodeCoordinates = Mesh.GetNodeCoordinates();
 	InnerNodeCoordinates.Reserve(InnerNodeCount);
 	for (int32 Index = 0; Index < (int32)Grid.GetInner3DPoints().Num(); ++Index)
 	{
@@ -309,13 +309,13 @@ void FIsoTriangulator::FillMeshNodes()
 
 	for (FLoopNode& Node : LoopNodes)
 	{
-		const FPoint2D& UVCoordinate = Node.Get2DPoint(EGridSpace::Scaled, Grid);
+		const FVector2d& UVCoordinate = Node.Get2DPoint(EGridSpace::Scaled, Grid);
 		Mesh.UVMap.Emplace(UVCoordinate.X, UVCoordinate.Y);
 	}
 
 	for (FIsoInnerNode& Node : InnerNodes)
 	{
-		const FPoint2D& UVCoordinate = Node.Get2DPoint(EGridSpace::Scaled, Grid);
+		const FVector2d& UVCoordinate = Node.Get2DPoint(EGridSpace::Scaled, Grid);
 		Mesh.UVMap.Emplace(UVCoordinate.X, UVCoordinate.Y);
 	}
 }
@@ -386,7 +386,7 @@ void FIsoTriangulator::GetThinZoneMesh(const TMap<int32, FLoopNode*>& IndexToNod
 	TArray<TPair<int32, FPairOfIndex>> CrossZoneElements;
 	TArray<FCandidateSegment> MeshOfThinZones;
 
-	FAddMeshNodeFunc AddElement = [&CrossZoneElements](const int32 NodeIndice, const FPoint2D& MeshNode2D, double MeshingTolerance3D, const FEdgeSegment& EdgeSegment, const FPairOfIndex& OppositeNodeIndices)
+	FAddMeshNodeFunc AddElement = [&CrossZoneElements](const int32 NodeIndice, const FVector2d& MeshNode2D, double MeshingTolerance3D, const FEdgeSegment& EdgeSegment, const FPairOfIndex& OppositeNodeIndices)
 	{
 		if (CrossZoneElements.Num() && CrossZoneElements.Last().Key == NodeIndice)
 		{
@@ -430,8 +430,8 @@ void FIsoTriangulator::GetThinZoneMesh(const TMap<int32, FLoopNode*>& IndexToNod
 			return;
 		}
 
-		const FPoint2D& CoordinateA = NodeA->Get2DPoint(EGridSpace::UniformScaled, Grid);
-		const FPoint2D& CoordinateB = NodeB->Get2DPoint(EGridSpace::UniformScaled, Grid);
+		const FVector2d& CoordinateA = NodeA->Get2DPoint(EGridSpace::UniformScaled, Grid);
+		const FVector2d& CoordinateB = NodeB->Get2DPoint(EGridSpace::UniformScaled, Grid);
 
 		// Is Outside and not too flat at Node1
 		const double FlatAngle = 0.1;
@@ -546,13 +546,13 @@ void FIsoTriangulator::BuildInnerSegments()
 	//        Node1 *------------------* Node2 	        Node1 *----------\./-------* Node2 
 	//                                                                       
 	//
-	TFunction<bool(const FPoint2D&, const FPoint2D&, const ESegmentType, const double)> AlmostHitsLoop = [&](const FPoint2D& Node1, const FPoint2D& Node2, const ESegmentType InType, const double Tolerance) -> bool
+	TFunction<bool(const FVector2d&, const FVector2d&, const ESegmentType, const double)> AlmostHitsLoop = [&](const FVector2d& Node1, const FVector2d& Node2, const ESegmentType InType, const double Tolerance) -> bool
 	{
 		if (InType == ESegmentType::IsoV)
 		{
-			for (const TArray<FPoint2D>& Loop : Grid.GetLoops2D(EGridSpace::UniformScaled))
+			for (const TArray<FVector2d>& Loop : Grid.GetLoops2D(EGridSpace::UniformScaled))
 			{
-				for (const FPoint2D& LoopPoint : Loop)
+				for (const FVector2d& LoopPoint : Loop)
 				{
 					if (FMath::IsNearlyEqual(LoopPoint.Y, Node1.Y, Tolerance))
 					{
@@ -574,9 +574,9 @@ void FIsoTriangulator::BuildInnerSegments()
 		}
 		else
 		{
-			for (const TArray<FPoint2D>& Loop : Grid.GetLoops2D(EGridSpace::UniformScaled))
+			for (const TArray<FVector2d>& Loop : Grid.GetLoops2D(EGridSpace::UniformScaled))
 			{
-				for (const FPoint2D& LoopPoint : Loop)
+				for (const FVector2d& LoopPoint : Loop)
 				{
 					if (FMath::IsNearlyEqual(LoopPoint.X, Node1.X, Tolerance))
 					{
@@ -601,8 +601,8 @@ void FIsoTriangulator::BuildInnerSegments()
 
 	TFunction<void(const int32, const int32, const ESegmentType)> AddToInnerToOuterSegmentsIntersectionTool = [&](const int32 IndexNode1, const int32 IndexNode2, const ESegmentType InType)
 	{
-		const FPoint2D& Point1 = Grid.GetInner2DPoint(EGridSpace::UniformScaled, IndexNode1);
-		const FPoint2D& Point2 = Grid.GetInner2DPoint(EGridSpace::UniformScaled, IndexNode2);
+		const FVector2d& Point1 = Grid.GetInner2DPoint(EGridSpace::UniformScaled, IndexNode1);
+		const FVector2d& Point2 = Grid.GetInner2DPoint(EGridSpace::UniformScaled, IndexNode2);
 
 		InnerToOuterIsoSegmentsIntersectionTool.AddIsoSegment(Point1, Point2, InType);
 	};
@@ -1028,7 +1028,7 @@ void FIsoTriangulator::FindCellContainingBoundaryNodes(TArray<FCell>& Cells)
 		{
 			if (!LoopPoint.IsDelete())
 			{
-				const FPoint2D& Coordinate = LoopPoint.Get2DPoint(EGridSpace::UniformScaled, Grid);
+				const FVector2d& Coordinate = LoopPoint.Get2DPoint(EGridSpace::UniformScaled, Grid);
 				ArrayUtils::FindCoordinateIndex(IsoUCoordinates, Coordinate.X, IndexU);
 				ArrayUtils::FindCoordinateIndex(IsoVCoordinates, Coordinate.Y, IndexV);
 
@@ -1207,10 +1207,10 @@ void FIsoTriangulator::FindInnerGridCellSurroundingSmallLoop()
 	const TArray<double>& UCoordinates = Grid.GetCuttingCoordinatesAlongIso(EIso::IsoU);
 	const TArray<double>& VCoordinates = Grid.GetCuttingCoordinatesAlongIso(EIso::IsoV);
 
-	const TArray<TArray<FPoint2D>>& Loops = Grid.GetLoops2D(EGridSpace::Default2D);
+	const TArray<TArray<FVector2d>>& Loops = Grid.GetLoops2D(EGridSpace::Default2D);
 	for (int32 LoopIndex = 1; LoopIndex < Loops.Num(); ++LoopIndex)
 	{
-		FPoint2D FirstPoint = Loops[LoopIndex][0];
+		FVector2d FirstPoint = Loops[LoopIndex][0];
 
 		int32 IndexU = 0;
 		for (; IndexU < NumU - 1; ++IndexU)
@@ -1236,7 +1236,7 @@ void FIsoTriangulator::FindInnerGridCellSurroundingSmallLoop()
 		double VMax = VCoordinates[IndexV + 1] + DOUBLE_SMALL_NUMBER;
 
 		bool bBoudardyIsSurrounded = true;
-		for (const FPoint2D& LoopPoint : Loops[LoopIndex])
+		for (const FVector2d& LoopPoint : Loops[LoopIndex])
 		{
 			if (LoopPoint.X < UMin || LoopPoint.X > UMax || LoopPoint.Y < VMin || LoopPoint.Y > VMax)
 			{
@@ -1555,8 +1555,8 @@ bool FIsoTriangulator::FindCycle(FIsoSegment* StartSegment, bool LeftSide, TArra
 
 FIsoSegment* FIsoTriangulator::FindNextSegment(EGridSpace Space, const FIsoSegment* StartSegment, const FIsoNode* StartNode, SlopeMethod GetSlope) const
 {
-	const FPoint2D& StartPoint = StartNode->Get2DPoint(Space, Grid);
-	const FPoint2D& EndPoint = (StartNode == &StartSegment->GetFirstNode()) ? StartSegment->GetSecondNode().Get2DPoint(Space, Grid) : StartSegment->GetFirstNode().Get2DPoint(Space, Grid);
+	const FVector2d& StartPoint = StartNode->Get2DPoint(Space, Grid);
+	const FVector2d& EndPoint = (StartNode == &StartSegment->GetFirstNode()) ? StartSegment->GetSecondNode().Get2DPoint(Space, Grid) : StartSegment->GetFirstNode().Get2DPoint(Space, Grid);
 
 	double ReferenceSlope = ComputePositiveSlope(StartPoint, EndPoint, 0);
 
@@ -1565,7 +1565,7 @@ FIsoSegment* FIsoTriangulator::FindNextSegment(EGridSpace Space, const FIsoSegme
 
 	for (FIsoSegment* Segment : StartNode->GetConnectedSegments())
 	{
-		const FPoint2D& OtherPoint = (StartNode == &Segment->GetFirstNode()) ? Segment->GetSecondNode().Get2DPoint(Space, Grid) : Segment->GetFirstNode().Get2DPoint(Space, Grid);
+		const FVector2d& OtherPoint = (StartNode == &Segment->GetFirstNode()) ? Segment->GetSecondNode().Get2DPoint(Space, Grid) : Segment->GetFirstNode().Get2DPoint(Space, Grid);
 
 		double Slope = GetSlope(StartPoint, OtherPoint, ReferenceSlope);
 		if (Slope < SMALL_NUMBER_SQUARE)
@@ -1823,8 +1823,8 @@ void FCell::CheckAllLoopsConnectedTogetherAndConnect()
 
 void FCell::TryToCreateSegment(FCellConnexion& LoopConnexion)
 {
-	const FPoint2D& ACoordinates = LoopConnexion.NodeA->Get2DPoint(EGridSpace::UniformScaled, Grid);
-	const FPoint2D& BCoordinates = LoopConnexion.NodeB->Get2DPoint(EGridSpace::UniformScaled, Grid);
+	const FVector2d& ACoordinates = LoopConnexion.NodeA->Get2DPoint(EGridSpace::UniformScaled, Grid);
+	const FVector2d& BCoordinates = LoopConnexion.NodeB->Get2DPoint(EGridSpace::UniformScaled, Grid);
 
 	LoopConnexion.Segment = Triangulator.GetOrTryToCreateSegment(*this, LoopConnexion.NodeA, ACoordinates, LoopConnexion.NodeB, BCoordinates, Slope::OneDegree);
 	if (LoopConnexion.Segment && !LoopConnexion.Segment->IsAFinalSegment())
@@ -1887,7 +1887,7 @@ void FCell::InitLoopConnexions()
 	}
 	default:
 	{
-		TArray<TPair<int32, FPoint2D>> LoopBarycenters = GetLoopBarycenters();
+		TArray<TPair<int32, FVector2d>> LoopBarycenters = GetLoopBarycenters();
 
 #ifdef DEBUG_BOWYERWATSON
 		FBowyerWatsonTriangulator::bDisplay = bDisplay;
@@ -1936,17 +1936,17 @@ void FCell::TryToConnectTwoSubLoopsWithShortestSegment(FCellConnexion& LoopConne
 
 	for (FLoopNode* NodeA : LoopA)
 	{
-		const FPoint2D& ACoordinates = NodeA->Get2DPoint(EGridSpace::UniformScaled, Grid);
+		const FVector2d& ACoordinates = NodeA->Get2DPoint(EGridSpace::UniformScaled, Grid);
 
 		for (FLoopNode* NodeB : LoopB)
 		{
-			const FPoint2D& BCoordinates = NodeB->Get2DPoint(EGridSpace::UniformScaled, Grid);
+			const FVector2d& BCoordinates = NodeB->Get2DPoint(EGridSpace::UniformScaled, Grid);
 
 #ifdef DEBUG_CONNECT_CELL_SUB_LOOPS_BY_NEIGHBORHOOD_
 			Grid.DisplayIsoSegment(EGridSpace::UniformScaled, *NodeA, *NodeB, 0, EVisuProperty::BlueCurve);
 #endif
 
-			double SquareDistance = FPoint2D::DistSquared(ACoordinates, BCoordinates);
+			double SquareDistance = FVector2d::DistSquared(ACoordinates, BCoordinates);
 			if (SquareDistance < MinDistanceSquare)
 			{
 				MinDistanceSquare = SquareDistance;
@@ -2011,8 +2011,8 @@ void FIsoTriangulator::TryToConnectTwoLoopsWithIsocelesTriangle(FCell& Cell, con
 		FIsoSegment* StartToCandiatSegment = nullptr;
 		FIsoSegment* EndToCandiatSegment = nullptr;
 
-		const FPoint2D& StartPoint2D = StartNode.Get2DPoint(EGridSpace::UniformScaled, Grid);
-		const FPoint2D& EndPoint2D = EndNode.Get2DPoint(EGridSpace::UniformScaled, Grid);
+		const FVector2d& StartPoint2D = StartNode.Get2DPoint(EGridSpace::UniformScaled, Grid);
+		const FVector2d& EndPoint2D = EndNode.Get2DPoint(EGridSpace::UniformScaled, Grid);
 
 		double StartReferenceSlope = ComputePositiveSlope(StartPoint2D, EndPoint2D, 0);
 		double EndReferenceSlope = StartReferenceSlope < 4 ? StartReferenceSlope + 4 : StartReferenceSlope - 4;
@@ -2030,7 +2030,7 @@ void FIsoTriangulator::TryToConnectTwoLoopsWithIsocelesTriangle(FCell& Cell, con
 			}
 
 			// Check if the node is inside the sector (X) or outside (Z)
-			const FPoint2D& NodePoint2D = Node->Get2DPoint(EGridSpace::UniformScaled, Grid);
+			const FVector2d& NodePoint2D = Node->Get2DPoint(EGridSpace::UniformScaled, Grid);
 
 			double PointCriteria = IsoTriangulatorImpl::IsoscelesCriteriaMax(StartPoint2D, EndPoint2D, NodePoint2D);
 
@@ -2094,8 +2094,8 @@ void FIsoTriangulator::TryToConnectTwoLoopsWithIsocelesTriangle(FCell& Cell, con
 			continue;
 		}
 
-		const FPoint2D& A1Coordinates = NodeA1->Get2DPoint(EGridSpace::UniformScaled, Grid);
-		const FPoint2D& A2Coordinates = NodeA2->Get2DPoint(EGridSpace::UniformScaled, Grid);
+		const FVector2d& A1Coordinates = NodeA1->Get2DPoint(EGridSpace::UniformScaled, Grid);
+		const FVector2d& A2Coordinates = NodeA2->Get2DPoint(EGridSpace::UniformScaled, Grid);
 
 		FIsoSegment* Segment = NodeA1->GetSegmentConnectedTo(NodeA2);
 
@@ -2111,7 +2111,7 @@ void FIsoTriangulator::TryToConnectTwoLoopsWithIsocelesTriangle(FCell& Cell, con
 
 			}
 #endif 
-			const FPoint2D& NodeCoordinates = Node->Get2DPoint(EGridSpace::UniformScaled, Grid);
+			const FVector2d& NodeCoordinates = Node->Get2DPoint(EGridSpace::UniformScaled, Grid);
 			if (!NodeA1->IsDeleteOrThinNode())
 			{
 				GetOrTryToCreateSegment(Cell, NodeA1, A1Coordinates, Node, NodeCoordinates, 0.1);
@@ -2152,10 +2152,10 @@ void FIsoTriangulator::TryToConnectVertexSubLoopWithTheMostIsoSegment(FCell& Cel
 			continue;
 		}
 
-		const FPoint2D& ACoordinates = CandidateA->Get2DPoint(EGridSpace::UniformScaled, Grid);
+		const FVector2d& ACoordinates = CandidateA->Get2DPoint(EGridSpace::UniformScaled, Grid);
 
 		FLoopNode* NextA = Loop[IndexA + 1];
-		const FPoint2D& NextACoordinates = NextA->Get2DPoint(EGridSpace::UniformScaled, Grid);
+		const FVector2d& NextACoordinates = NextA->Get2DPoint(EGridSpace::UniformScaled, Grid);
 
 		double ReferenceSlope = 0;
 		{
@@ -2184,7 +2184,7 @@ void FIsoTriangulator::TryToConnectVertexSubLoopWithTheMostIsoSegment(FCell& Cel
 				continue;
 			}
 
-			const FPoint2D& BCoordinates = NodeB->Get2DPoint(EGridSpace::UniformScaled, Grid);
+			const FVector2d& BCoordinates = NodeB->Get2DPoint(EGridSpace::UniformScaled, Grid);
 
 			double Slope = ComputeSlopeRelativeToReferenceAxis(ACoordinates, BCoordinates, ReferenceSlope);
 			if (Slope < MinSlope)
@@ -2196,7 +2196,7 @@ void FIsoTriangulator::TryToConnectVertexSubLoopWithTheMostIsoSegment(FCell& Cel
 
 		if (MinSlope < FlatSlope)
 		{
-			const FPoint2D& BCoordinates = CandidateB->Get2DPoint(EGridSpace::UniformScaled, Grid);
+			const FVector2d& BCoordinates = CandidateB->Get2DPoint(EGridSpace::UniformScaled, Grid);
 
 #ifdef DEBUG_FIND_ISO_SEGMENT_TO_LINK_OUTER_LOOP_NODES
 			Grid.DisplayIsoSegment(UniformScaled, *CandidateA, *CandidateB, 0, EVisuProperty::RedCurve);
@@ -2228,7 +2228,7 @@ void FIsoTriangulator::TryToConnectTwoSubLoopsWithTheMostIsoSegment(FCell& Cell,
 		}
 
 		FLoopNode* CandidateB = nullptr;
-		const FPoint2D& ACoordinates = CandidateA->Get2DPoint(EGridSpace::UniformScaled, Grid);
+		const FVector2d& ACoordinates = CandidateA->Get2DPoint(EGridSpace::UniformScaled, Grid);
 
 		double MinSlope = FlatSlope + DOUBLE_SMALL_NUMBER;// 0.25; // ~15 deg: The segment must make an angle less than 10 deg with the Iso
 		double MinLengthSquare = HUGE_VALUE;
@@ -2240,7 +2240,7 @@ void FIsoTriangulator::TryToConnectTwoSubLoopsWithTheMostIsoSegment(FCell& Cell,
 				continue;
 			}
 
-			const FPoint2D& BCoordinates = NodeB->Get2DPoint(EGridSpace::UniformScaled, Grid);
+			const FVector2d& BCoordinates = NodeB->Get2DPoint(EGridSpace::UniformScaled, Grid);
 
 			double Slope = ComputeSlopeRelativeToNearestAxis(ACoordinates, BCoordinates);
 			if (Slope < MinSlope)
@@ -2249,7 +2249,7 @@ void FIsoTriangulator::TryToConnectTwoSubLoopsWithTheMostIsoSegment(FCell& Cell,
 				// If the slope of the candidate segments is nearly zero, then select the shortest
 				if (MinSlope < DOUBLE_KINDA_SMALL_NUMBER)
 				{
-					double DistanceSquare = FPoint2D::DistSquared(BCoordinates, ACoordinates);
+					double DistanceSquare = FVector2d::DistSquared(BCoordinates, ACoordinates);
 					if (DistanceSquare > MinLengthSquare)
 					{
 						continue;
@@ -2262,7 +2262,7 @@ void FIsoTriangulator::TryToConnectTwoSubLoopsWithTheMostIsoSegment(FCell& Cell,
 
 		if (MinSlope < FlatSlope)
 		{
-			const FPoint2D& BCoordinates = CandidateB->Get2DPoint(EGridSpace::UniformScaled, Grid);
+			const FVector2d& BCoordinates = CandidateB->Get2DPoint(EGridSpace::UniformScaled, Grid);
 
 #ifdef DEBUG_FIND_ISO_SEGMENT_TO_LINK_OUTER_LOOP_NODES
 			Grid.DisplayIsoSegment(EGridSpace::UniformScaled, *CandidateA, *CandidateB, 0, EVisuProperty::BlueCurve);
@@ -2273,7 +2273,7 @@ void FIsoTriangulator::TryToConnectTwoSubLoopsWithTheMostIsoSegment(FCell& Cell,
 	}
 }
 
-FIsoSegment* FIsoTriangulator::GetOrTryToCreateSegment(FCell& Cell, FLoopNode* NodeA, const FPoint2D& ACoordinates, FIsoNode* NodeB, const FPoint2D& BCoordinates, const double FlatAngle)
+FIsoSegment* FIsoTriangulator::GetOrTryToCreateSegment(FCell& Cell, FLoopNode* NodeA, const FVector2d& ACoordinates, FIsoNode* NodeB, const FVector2d& BCoordinates, const double FlatAngle)
 {
 	if (FIsoSegment* Segment = NodeA->GetSegmentConnectedTo(NodeB))
 	{
@@ -2367,7 +2367,7 @@ void FIsoTriangulator::FindCandidateToConnectCellCornerToLoops(FCell& Cell)
 	TFunction<void(FCellLoop&, FCellCorner&)> FindAndTryCreateCandidateSegmentToLinkLoopToCorner = [&](FCellLoop& LoopCell, FCellCorner& CellCorner)
 	{
 		FIsoInnerNode& CornerNode = CellCorner.CornerNode;
-		const FPoint2D& CornerPoint = CellCorner.Barycenter;
+		const FVector2d& CornerPoint = CellCorner.Barycenter;
 
 		const TArray<FLoopNode*>& LoopA = LoopCell.Nodes;
 
@@ -2377,9 +2377,9 @@ void FIsoTriangulator::FindCandidateToConnectCellCornerToLoops(FCell& Cell)
 		{
 			FLoopNode* NodeA = LoopA[IndexA];
 
-			const FPoint2D& ACoordinates = NodeA->Get2DPoint(EGridSpace::UniformScaled, Grid);
+			const FVector2d& ACoordinates = NodeA->Get2DPoint(EGridSpace::UniformScaled, Grid);
 
-			double SquareDistance = FPoint2D::DistSquared(ACoordinates, CornerPoint);
+			double SquareDistance = FVector2d::DistSquared(ACoordinates, CornerPoint);
 			if (SquareDistance < MinDistanceSquare)
 			{
 				MinDistanceSquare = SquareDistance;
@@ -2390,7 +2390,7 @@ void FIsoTriangulator::FindCandidateToConnectCellCornerToLoops(FCell& Cell)
 		if (MinIndexA >= 0)
 		{
 			FLoopNode* NodeA = LoopA[MinIndexA];
-			const FPoint2D& ACoordinates = NodeA->Get2DPoint(EGridSpace::UniformScaled, Grid);
+			const FVector2d& ACoordinates = NodeA->Get2DPoint(EGridSpace::UniformScaled, Grid);
 
 			double SlopeVsCell = ComputeSlopeRelativeToNearestAxis(ACoordinates, CornerPoint);
 			if (SlopeVsCell < Slope::OneDegree)
diff --git a/Engine/Source/Runtime/Datasmith/CADKernel/Base/Private/Mesh/Meshers/IsoTriangulator.h b/Engine/Source/Runtime/Datasmith/CADKernel/Base/Private/Mesh/Meshers/IsoTriangulator.h
--- a/Engine/Source/Runtime/Datasmith/CADKernel/Base/Private/Mesh/Meshers/IsoTriangulator.h
+++ b/Engine/Source/Runtime/Datasmith/CADKernel/Base/Private/Mesh/Meshers/IsoTriangulator.h
@@ -278,8 +278,8 @@ public:
 
 		Algo::Sort(SortedLoopNodes, [this](const FLoopNode* Node1, const FLoopNode* Node2)
 			{
-				const FPoint2D& Node1Coordinates = Node1->Get2DPoint(EGridSpace::Default2D, Grid);
-				const FPoint2D& Node2Coordinates = Node2->Get2DPoint(EGridSpace::Default2D, Grid);
+				const FVector2d& Node1Coordinates = Node1->Get2DPoint(EGridSpace::Default2D, Grid);
+				const FVector2d& Node2Coordinates = Node2->Get2DPoint(EGridSpace::Default2D, Grid);
 				return (Node1Coordinates.X + Node1Coordinates.Y) < (Node2Coordinates.X + Node2Coordinates.Y);
 			});
 	}
@@ -318,7 +318,7 @@ private:
 	 *
 	 */
 	void TryToConnectVertexSubLoopWithTheMostIsoSegment(FCell& Cell, const TArray<FLoopNode*>& SubLoop);
-	FIsoSegment* GetOrTryToCreateSegment(FCell& Cell, FLoopNode* NodeA, const FPoint2D& ACoordinates, FIsoNode* NodeB, const FPoint2D& BCoordinates, const double FlatAngle);
+	FIsoSegment* GetOrTryToCreateSegment(FCell& Cell, FLoopNode* NodeA, const FVector2d& ACoordinates, FIsoNode* NodeB, const FVector2d& BCoordinates, const double FlatAngle);
 
 public:
 #ifdef CADKERNEL_DEV
@@ -354,12 +354,12 @@ namespace IsoTriangulatorImpl
  * A "Delaunay" triangle is an equilateral triangle
  * The optimal value is the smallest value.
  */
-inline double CotangentCriteria(const FPoint& APoint, const FPoint& BPoint, const FPoint& PPoint, FPoint& OutNormal)
+inline double CotangentCriteria(const FVector& APoint, const FVector& BPoint, const FVector& PPoint, FVector& OutNormal)
 {
 	const double BigValue = HUGE_VALUE;
 
-	FPoint PA = APoint - PPoint;
-	FPoint PB = BPoint - PPoint;
+	FVector PA = APoint - PPoint;
+	FVector PB = BPoint - PPoint;
 
 	// the ratio between the scalar product PA.PB (=|PA| |PB| cos (A,P,B) )
 	// with the norm of the cross product |PA^PB| (=|PA| |PB| |sin(A,P,B)|)
@@ -378,12 +378,12 @@ inline double CotangentCriteria(const FPoint& APoint, const FPoint& BPoint, cons
 	return ScalareProduct / NormOFScalarProduct;
 }
 
-inline double CotangentCriteria(const FPoint2D& APoint, const FPoint2D& BPoint, const FPoint2D& PPoint)
+inline double CotangentCriteria(const FVector2d& APoint, const FVector2d& BPoint, const FVector2d& PPoint)
 {
 	const double BigValue = HUGE_VALUE;
 
-	FPoint2D PA = APoint - PPoint;
-	FPoint2D PB = BPoint - PPoint;
+	FVector2d PA = APoint - PPoint;
+	FVector2d PB = BPoint - PPoint;
 
 	// the ratio between the scalar product PA.PB (=|PA| |PB| cos (A,P,B) )
 	// with the norm of the cross product |PA^PB| (=|PA| |PB| |sin(A,P,B)|)
@@ -408,7 +408,7 @@ struct FPairOfDouble
 	double Value2;
 };
 
-inline FPairOfDouble IsoscelesCriteria(const FPoint2D& APoint, const FPoint2D& BPoint, const FPoint2D& CPoint)
+inline FPairOfDouble IsoscelesCriteria(const FVector2d& APoint, const FVector2d& BPoint, const FVector2d& CPoint)
 {
 	const double SlopeAB = ComputeSlope(APoint, BPoint);
 	const double SlopeAC = ComputeSlope(APoint, CPoint);
@@ -418,19 +418,19 @@ inline FPairOfDouble IsoscelesCriteria(const FPoint2D& APoint, const FPoint2D& B
 	return { TransformIntoOrientedSlope(SlopeAC - SlopeAB), TransformIntoOrientedSlope(SlopeBA - SlopeBC) };
 }
 
-inline double IsoscelesCriteriaMax(const FPoint2D& APoint, const FPoint2D& BPoint, const FPoint2D& CPoint)
+inline double IsoscelesCriteriaMax(const FVector2d& APoint, const FVector2d& BPoint, const FVector2d& CPoint)
 {
 	const FPairOfDouble Criteria = IsoscelesCriteria(APoint, BPoint, CPoint);
 	return FMath::Max(Criteria.Value1, Criteria.Value2);
 }
 
-inline double IsoscelesCriteriaMin(const FPoint2D& APoint, const FPoint2D& BPoint, const FPoint2D& CPoint)
+inline double IsoscelesCriteriaMin(const FVector2d& APoint, const FVector2d& BPoint, const FVector2d& CPoint)
 {
 	const FPairOfDouble Criteria = IsoscelesCriteria(APoint, BPoint, CPoint);
 	return FMath::Min(Criteria.Value1, Criteria.Value2);
 }
 
-inline double EquilateralSlopeCriteria(const FPoint2D& APoint, const FPoint2D& BPoint, const FPoint2D& CPoint)
+inline double EquilateralSlopeCriteria(const FVector2d& APoint, const FVector2d& BPoint, const FVector2d& CPoint)
 {
 	const double SlopeAB = ComputeSlope(APoint, BPoint);
 	const double SlopeAC = ComputeSlope(APoint, CPoint);
@@ -449,9 +449,9 @@ inline double EquilateralSlopeCriteria(const FPoint2D& APoint, const FPoint2D& B
 template<class PointType>
 inline double EquilateralCriteria(const PointType& SegmentA, const PointType& SegmentB, const PointType& Point)
 {
-	double Criteria1 = FMath::Abs(CoordinateOfProjectedPointOnSegment(SegmentA, SegmentB, Point, false) - 0.5);
-	double Criteria2 = FMath::Abs(CoordinateOfProjectedPointOnSegment(Point, SegmentA, SegmentB, false) - 0.5);
-	double Criteria3 = FMath::Abs(CoordinateOfProjectedPointOnSegment(SegmentB, Point, SegmentA, false) - 0.5);
+	double Criteria1 = FMath::Abs(CoordinateOfProjectedPointOnSegment<PointType>(SegmentA, SegmentB, Point, false) - 0.5);
+	double Criteria2 = FMath::Abs(CoordinateOfProjectedPointOnSegment<PointType>(Point, SegmentA, SegmentB, false) - 0.5);
+	double Criteria3 = FMath::Abs(CoordinateOfProjectedPointOnSegment<PointType>(SegmentB, Point, SegmentA, false) - 0.5);
 	return Criteria1 + Criteria2 + Criteria3;
 }
 
diff --git a/Engine/Source/Runtime/Datasmith/CADKernel/Base/Private/Mesh/Meshers/IsoTriangulator/IntersectionIsoSegmentTool.cpp b/Engine/Source/Runtime/Datasmith/CADKernel/Base/Private/Mesh/Meshers/IsoTriangulator/IntersectionIsoSegmentTool.cpp
--- a/Engine/Source/Runtime/Datasmith/CADKernel/Base/Private/Mesh/Meshers/IsoTriangulator/IntersectionIsoSegmentTool.cpp
+++ b/Engine/Source/Runtime/Datasmith/CADKernel/Base/Private/Mesh/Meshers/IsoTriangulator/IntersectionIsoSegmentTool.cpp
@@ -26,7 +26,7 @@ FIntersectionIsoSegmentTool::FIntersectionIsoSegmentTool(const FGrid& InGrid, co
 	VSegments.Reserve(MaxSegmentCount);
 }
 
-void FIntersectionIsoSegmentTool::AddIsoSegment(const FPoint2D& StartPoint, const FPoint2D& EndPoint, const ESegmentType InType)
+void FIntersectionIsoSegmentTool::AddIsoSegment(const FVector2d& StartPoint, const FVector2d& EndPoint, const ESegmentType InType)
 {
 	switch(InType)
 	{ 
diff --git a/Engine/Source/Runtime/Datasmith/CADKernel/Base/Private/Mesh/Meshers/IsoTriangulator/IntersectionIsoSegmentTool.h b/Engine/Source/Runtime/Datasmith/CADKernel/Base/Private/Mesh/Meshers/IsoTriangulator/IntersectionIsoSegmentTool.h
--- a/Engine/Source/Runtime/Datasmith/CADKernel/Base/Private/Mesh/Meshers/IsoTriangulator/IntersectionIsoSegmentTool.h
+++ b/Engine/Source/Runtime/Datasmith/CADKernel/Base/Private/Mesh/Meshers/IsoTriangulator/IntersectionIsoSegmentTool.h
@@ -17,15 +17,15 @@ namespace IntersectionIsoSegmentTool
 class FIntersectionIsoSegment
 {
 private:
-	const FPoint2D& Point0;
-	const FPoint2D& Point1;
+	const FVector2d& Point0;
+	const FVector2d& Point1;
 	double IsoCoordinate;
 	double MinCoordinate;
 	double MaxCoordinate;
 
 public:
 	
-	FIntersectionIsoSegment(const FPoint2D& StartPoint, const FPoint2D& EndPoint, const double InIsoCoordinate, const double InStartCoordinate, const double InEndCoordinate)
+	FIntersectionIsoSegment(const FVector2d& StartPoint, const FVector2d& EndPoint, const double InIsoCoordinate, const double InStartCoordinate, const double InEndCoordinate)
 		: Point0(StartPoint)
 		, Point1(EndPoint)
 		, IsoCoordinate(InIsoCoordinate)
@@ -36,8 +36,8 @@ public:
 	virtual ~FIntersectionIsoSegment()
 	{}
 
-	virtual FPoint2D GetMinPoint() const = 0;
-	virtual FPoint2D GetMaxPoint() const = 0;
+	virtual FVector2d GetMinPoint() const = 0;
+	virtual FVector2d GetMaxPoint() const = 0;
 
 	double GetIsoCoordinate() const 
 	{
@@ -75,35 +75,35 @@ public:
 class FIsoUSegment : public FIntersectionIsoSegment
 {
 public:
-	FIsoUSegment(const FPoint2D& StartPoint, const FPoint2D& EndPoint)
+	FIsoUSegment(const FVector2d& StartPoint, const FVector2d& EndPoint)
 		: FIntersectionIsoSegment(StartPoint, EndPoint, StartPoint.X, StartPoint.Y, EndPoint.Y)
 	{}
 
-	virtual FPoint2D GetMinPoint() const override
+	virtual FVector2d GetMinPoint() const override
 	{
-		return FPoint2D(GetIsoCoordinate(), GetMinCoordinate());
+		return FVector2d(GetIsoCoordinate(), GetMinCoordinate());
 	}
-	virtual FPoint2D GetMaxPoint() const override
+	virtual FVector2d GetMaxPoint() const override
 	{
-		return FPoint2D(GetIsoCoordinate(), GetMaxCoordinate());
+		return FVector2d(GetIsoCoordinate(), GetMaxCoordinate());
 	}
 };
 
 class FIsoVSegment : public FIntersectionIsoSegment
 {
 public:
-	FIsoVSegment(const FPoint2D& StartPoint, const FPoint2D& EndPoint)
+	FIsoVSegment(const FVector2d& StartPoint, const FVector2d& EndPoint)
 		: FIntersectionIsoSegment(StartPoint, EndPoint, StartPoint.Y, StartPoint.X, EndPoint.X)
 	{}
 
-	virtual FPoint2D GetMinPoint() const override
+	virtual FVector2d GetMinPoint() const override
 	{
-		return FPoint2D(GetMinCoordinate(), GetIsoCoordinate());
+		return FVector2d(GetMinCoordinate(), GetIsoCoordinate());
 	}
 
-	virtual FPoint2D GetMaxPoint() const override
+	virtual FVector2d GetMaxPoint() const override
 	{
-		return FPoint2D(GetMaxCoordinate(), GetIsoCoordinate());
+		return FVector2d(GetMaxCoordinate(), GetIsoCoordinate());
 	}
 
 };
@@ -125,7 +125,7 @@ private:
 public:
 	FIntersectionIsoSegmentTool(const FGrid& InGrid, const double Tolerance);
 
-	void AddIsoSegment(const FPoint2D& StartPoint, const FPoint2D& EndPoint, const ESegmentType InType);
+	void AddIsoSegment(const FVector2d& StartPoint, const FVector2d& EndPoint, const ESegmentType InType);
 
 	bool DoesIntersect(const FIsoNode& StartNode, const FIsoNode& EndNode) const;
 
diff --git a/Engine/Source/Runtime/Datasmith/CADKernel/Base/Private/Mesh/Meshers/IsoTriangulator/IntersectionSegmentTool.cpp b/Engine/Source/Runtime/Datasmith/CADKernel/Base/Private/Mesh/Meshers/IsoTriangulator/IntersectionSegmentTool.cpp
--- a/Engine/Source/Runtime/Datasmith/CADKernel/Base/Private/Mesh/Meshers/IsoTriangulator/IntersectionSegmentTool.cpp
+++ b/Engine/Source/Runtime/Datasmith/CADKernel/Base/Private/Mesh/Meshers/IsoTriangulator/IntersectionSegmentTool.cpp
@@ -29,7 +29,7 @@ IntersectionSegmentTool::FSegment::FSegment(const FGrid& Grid, const double Tole
 {
 }
 
-IntersectionSegmentTool::FSegment::FSegment(const FGrid& Grid, const double Tolerance, const FIsoNode& StartNode, const FPoint2D& EndPoint)
+IntersectionSegmentTool::FSegment::FSegment(const FGrid& Grid, const double Tolerance, const FIsoNode& StartNode, const FVector2d& EndPoint)
 	: IntersectionToolBase::FSegment(Tolerance, StartNode.Get2DPoint(EGridSpace::UniformScaled, Grid), EndPoint)
 	, IsoSegment(nullptr)
 {
diff --git a/Engine/Source/Runtime/Datasmith/CADKernel/Base/Private/Mesh/Meshers/IsoTriangulator/IntersectionSegmentTool.h b/Engine/Source/Runtime/Datasmith/CADKernel/Base/Private/Mesh/Meshers/IsoTriangulator/IntersectionSegmentTool.h
--- a/Engine/Source/Runtime/Datasmith/CADKernel/Base/Private/Mesh/Meshers/IsoTriangulator/IntersectionSegmentTool.h
+++ b/Engine/Source/Runtime/Datasmith/CADKernel/Base/Private/Mesh/Meshers/IsoTriangulator/IntersectionSegmentTool.h
@@ -42,7 +42,7 @@ struct FSegment
 	/**
 	 * WARNING StartPoint, EndPoint must be defined in EGridSpace::UniformScaled
 	 */
-	FSegment(const double Tolerance, const FPoint2D& StartPoint, const FPoint2D& EndPoint)
+	FSegment(const double Tolerance, const FVector2d& StartPoint, const FVector2d& EndPoint)
 		: Segment2D(StartPoint, EndPoint)
 		, Boundary(StartPoint, EndPoint, Tolerance)
 	{
@@ -64,8 +64,8 @@ struct FSegment
 
 	static EConnectionType IsSuperimposed(const FSegment2D& SegmentAB, const FSegment2D& SegmentCD, bool bSameOrientation)
 	{
-		const FPoint2D AB = SegmentAB.GetVector().GetSafeNormal();
-		const FPoint2D CD = SegmentCD.GetVector().GetSafeNormal();
+		const FVector2d AB = SegmentAB.GetVector().GetSafeNormal();
+		const FVector2d CD = SegmentCD.GetVector().GetSafeNormal();
 		const double ParallelCoef = AB ^ CD;
 		if (FMath::IsNearlyZero(ParallelCoef, DOUBLE_KINDA_SMALL_NUMBER))
 		{
@@ -78,7 +78,7 @@ struct FSegment
 		return EConnectionType::StartFrom;
 	};
 
-	EConnectionType DoesItStartFromAndSuperimposed(const FIsoNode* StartNode, const FPoint2D* EndPoint, const FSegment2D& InSegment) const
+	EConnectionType DoesItStartFromAndSuperimposed(const FIsoNode* StartNode, const FVector2d* EndPoint, const FSegment2D& InSegment) const
 	{
 		if (GetFirstNode() == StartNode)
 		{
@@ -94,7 +94,7 @@ struct FSegment
 		return EConnectionType::DoesntStartFrom;
 	}
 
-	EConnectionType DoesItStartFromAndSuperimposed(const FPoint2D* StartPoint, const FPoint2D* EndPoint, const FSegment2D& InSegment) const
+	EConnectionType DoesItStartFromAndSuperimposed(const FVector2d* StartPoint, const FVector2d* EndPoint, const FSegment2D& InSegment) const
 	{
 		return EConnectionType::DoesntStartFrom;
 	}
@@ -182,8 +182,8 @@ struct FSegment : public IntersectionToolBase::FSegment
 
 	FSegment(const FGrid& Grid, const double Tolerance, const FIsoSegment& InSegment);
 	FSegment(const FGrid& Grid, const double Tolerance, const FIsoNode& StartNode, const FIsoNode& EndNode);
-	FSegment(const FGrid& Grid, const double Tolerance, const FIsoNode& StartNode, const FPoint2D& EndPoint);
-	FSegment(const FGrid& Grid, const double Tolerance, const FPoint2D& StartPoint, const FPoint2D& EndPoint)
+	FSegment(const FGrid& Grid, const double Tolerance, const FIsoNode& StartNode, const FVector2d& EndPoint);
+	FSegment(const FGrid& Grid, const double Tolerance, const FVector2d& StartPoint, const FVector2d& EndPoint)
 		: IntersectionToolBase::FSegment(Tolerance, StartPoint, EndPoint)
 		, IsoSegment(nullptr)
 	{
@@ -443,7 +443,7 @@ public:
 		Segments.Emplace(Grid, Tolerance, Segment);
 	}
 
-	void AddSegment(const FPoint2D& StartPoint, const FPoint2D& EndPoint)
+	void AddSegment(const FVector2d& StartPoint, const FVector2d& EndPoint)
 	{
 		bSegmentsAreSorted = false;
 		Segments.Emplace(Grid, Tolerance, StartPoint, EndPoint);
@@ -478,7 +478,7 @@ public:
 	/**
 	 * WARNING StartPoint, EndPoint must be defined in EGridSpace::UniformScaled
 	 */
-	bool DoesIntersect(const FPoint2D& StartPoint, const FPoint2D& EndPoint) const
+	bool DoesIntersect(const FVector2d& StartPoint, const FVector2d& EndPoint) const
 	{
 		return TIntersectionSegmentTool<IntersectionSegmentTool::FSegment>::FindIntersectingSegment(&StartPoint, &EndPoint) != nullptr;
 	}
@@ -487,7 +487,7 @@ public:
 	 * Allow StartNode to be connected to one segment
 	 * WARNING EndPoint must be defined in EGridSpace::UniformScaled
 	 */
-	bool DoesIntersect(const FIsoNode& StartNode, const FPoint2D& EndPoint) const
+	bool DoesIntersect(const FIsoNode& StartNode, const FVector2d& EndPoint) const
 	{
 		return TIntersectionSegmentTool<IntersectionSegmentTool::FSegment>::FindIntersectingSegment(&StartNode, &EndPoint) != nullptr;
 	}
diff --git a/Engine/Source/Runtime/Datasmith/CADKernel/Base/Private/Mesh/Meshers/IsoTriangulator/IsoCell.h b/Engine/Source/Runtime/Datasmith/CADKernel/Base/Private/Mesh/Meshers/IsoTriangulator/IsoCell.h
--- a/Engine/Source/Runtime/Datasmith/CADKernel/Base/Private/Mesh/Meshers/IsoTriangulator/IsoCell.h
+++ b/Engine/Source/Runtime/Datasmith/CADKernel/Base/Private/Mesh/Meshers/IsoTriangulator/IsoCell.h
@@ -20,7 +20,7 @@ struct FLoopConnexion;
 struct FCellLoop
 {
 	int32 Id;
-	FPoint2D Barycenter;
+	FVector2d Barycenter;
 	TArray<FLoopNode*> Nodes;
 	TArray<FCellConnexion*> Connexions;
 	bool bIsOuterLoop = false;
@@ -41,7 +41,7 @@ struct FCellLoop
 
 	FCellLoop(const int32 InIndex)
 		: Id(InIndex)
-		, Barycenter(FPoint::ZeroVector)
+		, Barycenter(FVector::ZeroVector)
 	{
 	}
 
@@ -465,9 +465,9 @@ struct FCell
 
 	void InitLoopConnexions();
 
-	TArray<TPair<int32, FPoint2D>> GetLoopBarycenters()
+	TArray<TPair<int32, FVector2d>> GetLoopBarycenters()
 	{
-		TArray<TPair<int32, FPoint2D>> LoopBarycenters;
+		TArray<TPair<int32, FVector2d>> LoopBarycenters;
 		LoopBarycenters.Reserve(CellLoops.Num());
 		for (const FCellLoop& LoopCell : CellLoops)
 		{
diff --git a/Engine/Source/Runtime/Datasmith/CADKernel/Base/Private/Mesh/Meshers/IsoTriangulator/IsoNode.h b/Engine/Source/Runtime/Datasmith/CADKernel/Base/Private/Mesh/Meshers/IsoTriangulator/IsoNode.h
--- a/Engine/Source/Runtime/Datasmith/CADKernel/Base/Private/Mesh/Meshers/IsoTriangulator/IsoNode.h
+++ b/Engine/Source/Runtime/Datasmith/CADKernel/Base/Private/Mesh/Meshers/IsoTriangulator/IsoNode.h
@@ -264,17 +264,12 @@ public:
 	/**
 	 * Return the 2d coordinate of the node according to the space
 	 */
-	virtual const FPoint2D& Get2DPoint(EGridSpace Space, const FGrid& Grid) const = 0;
-	virtual void Set2DPoint(EGridSpace Space, FGrid& Grid, const FPoint2D& NewCoordinate) = 0;
+	virtual const FVector2d& Get2DPoint(EGridSpace Space, const FGrid& Grid) const = 0;
+	virtual void Set2DPoint(EGridSpace Space, FGrid& Grid, const FVector2d& NewCoordinate) = 0;
 
-	virtual const FPoint& Get3DPoint(const FGrid& Grid) const = 0;
+	virtual const FVector& Get3DPoint(const FGrid& Grid) const = 0;
 	virtual const FVector3f& GetNormal(const FGrid& Grid) const = 0;
 
-	/**
-	 * Only for display purpose as it return a copy of the point
-	 */
-	virtual const FPoint GetPoint(EGridSpace Space, const FGrid& Grid) const = 0;
-
 	virtual bool operator==(const FIsoNode& OtherNode) const = 0;
 	virtual bool IsEqualTo(const FLoopNode& OtherNode) const = 0;
 	virtual bool IsEqualTo(const FIsoInnerNode& OtherNode) const = 0;
@@ -366,17 +361,17 @@ public:
 	/**
 	 * Return the 2d coordinate of the node according to the space
 	 */
-	virtual const FPoint2D& Get2DPoint(EGridSpace Space, const FGrid& Grid) const override
+	virtual const FVector2d& Get2DPoint(EGridSpace Space, const FGrid& Grid) const override
 	{
 		return Grid.GetLoop2DPoint(Space, LoopIndex, LocalIndex);
 	}
 
-	virtual void Set2DPoint(EGridSpace Space, FGrid& Grid, const FPoint2D& NewCoordinate) override
+	virtual void Set2DPoint(EGridSpace Space, FGrid& Grid, const FVector2d& NewCoordinate) override
 	{
 		Grid.SetLoop2DPoint(Space, LoopIndex, LocalIndex, NewCoordinate);
 	}
 
-	virtual const FPoint& Get3DPoint(const FGrid& Grid) const override
+	virtual const FVector& Get3DPoint(const FGrid& Grid) const override
 	{
 		return Grid.GetLoops3D()[LoopIndex][LocalIndex];
 	}
@@ -389,17 +384,17 @@ public:
 	/**
 	 * Only for display purpose as it return a copy of the point
 	 */
-	virtual const FPoint GetPoint(EGridSpace Space, const FGrid& Grid) const override
-	{
-		switch (Space)
-		{
-		case EGridSpace::Default2D:
-		case EGridSpace::Scaled:
-		case EGridSpace::UniformScaled:
-			return Grid.GetLoop2DPoint(Space, LoopIndex, LocalIndex);
-		}
-		return FPoint::ZeroVector;
-	}
+	//virtual const FVector GetPoint(EGridSpace Space, const FGrid& Grid) const override
+	//{
+	//	switch (Space)
+	//	{
+	//	case EGridSpace::Default2D:
+	//	case EGridSpace::Scaled:
+	//	case EGridSpace::UniformScaled:
+	//		return Grid.GetLoop2DPoint(Space, LoopIndex, LocalIndex);
+	//	}
+	//	return FVector::ZeroVector;
+	//}
 
 	/**
 	 * Check if the segment starting from this loop node (S) to EndSegmentCoordinate (E) is inside the face
@@ -427,7 +422,7 @@ public:
 	 *    |                                |        |                              |
 	 *    X---------------X----------------X        X------------------X-----------X
 	 */
-	bool IsSegmentBeInsideFace(const FPoint2D& EndSegmentCoordinate, const FGrid& Grid, const double FlatAngle) const
+	bool IsSegmentBeInsideFace(const FVector2d& EndSegmentCoordinate, const FGrid& Grid, const double FlatAngle) const
 	{
 		return (!IsPointPInsideSectorABC(GetPreviousNode().Get2DPoint(EGridSpace::UniformScaled, Grid), Get2DPoint(EGridSpace::UniformScaled, Grid), GetNextNode().Get2DPoint(EGridSpace::UniformScaled, Grid), EndSegmentCoordinate, FlatAngle));
 	}
@@ -507,17 +502,17 @@ public:
 		return (States & EIsoNodeStates::LinkedToLoop) == EIsoNodeStates::LinkedToLoop;
 	}
 
-	virtual const FPoint2D& Get2DPoint(EGridSpace Space, const FGrid& Grid) const override
+	virtual const FVector2d& Get2DPoint(EGridSpace Space, const FGrid& Grid) const override
 	{
 		return Grid.GetInner2DPoint(Space, LocalIndex);
 	}
 
-	virtual void Set2DPoint(EGridSpace Space, FGrid& Grid, const FPoint2D& NewCoordinate) override
+	virtual void Set2DPoint(EGridSpace Space, FGrid& Grid, const FVector2d& NewCoordinate) override
 	{
 		return Grid.SetInner2DPoint(Space, LocalIndex, NewCoordinate);
 	}
 
-	virtual const FPoint& Get3DPoint(const FGrid& Grid) const override
+	virtual const FVector& Get3DPoint(const FGrid& Grid) const override
 	{
 		return Grid.GetInner3DPoint(LocalIndex);
 	}
@@ -530,17 +525,17 @@ public:
 	/**
 	 * Only for display purpose as it return a copy of the point
 	 */
-	virtual const FPoint GetPoint(EGridSpace Space, const FGrid& Grid) const override
-	{
-		switch (Space)
-		{
-		case EGridSpace::Default2D:
-		case EGridSpace::Scaled:
-		case EGridSpace::UniformScaled:
-			return Grid.GetInner2DPoint(Space, LocalIndex);
-		}
-		return FPoint::ZeroVector;
-	}
+	//virtual const FVector GetPoint(EGridSpace Space, const FGrid& Grid) const override
+	//{
+	//	switch (Space)
+	//	{
+	//	case EGridSpace::Default2D:
+	//	case EGridSpace::Scaled:
+	//	case EGridSpace::UniformScaled:
+	//		return Grid.GetInner2DPoint(Space, LocalIndex);
+	//	}
+	//	return FVector::ZeroVector;
+	//}
 
 	void OffsetId(int32 StartId)
 	{
diff --git a/Engine/Source/Runtime/Datasmith/CADKernel/Base/Private/Mesh/Meshers/IsoTriangulator/IsoSegment.h b/Engine/Source/Runtime/Datasmith/CADKernel/Base/Private/Mesh/Meshers/IsoTriangulator/IsoSegment.h
--- a/Engine/Source/Runtime/Datasmith/CADKernel/Base/Private/Mesh/Meshers/IsoTriangulator/IsoSegment.h
+++ b/Engine/Source/Runtime/Datasmith/CADKernel/Base/Private/Mesh/Meshers/IsoTriangulator/IsoSegment.h
@@ -331,16 +331,16 @@ public:
 
 	double Get2DLengthSquare(EGridSpace Space, const FGrid& Grid) const
 	{
-		const FPoint2D& Point1 = FirstNode->Get2DPoint(Space, Grid);
-		const FPoint2D& Point2 = SecondNode->Get2DPoint(Space, Grid);
-		return FPoint2D::DistSquared(Point1, Point2);
+		const FVector2d& Point1 = FirstNode->Get2DPoint(Space, Grid);
+		const FVector2d& Point2 = SecondNode->Get2DPoint(Space, Grid);
+		return FVector2d::DistSquared(Point1, Point2);
 	}
 
 	double Get3DLengthSquare(const FGrid& Grid) const
 	{
-		const FPoint& Point1 = FirstNode->Get3DPoint(Grid);
-		const FPoint& Point2 = SecondNode->Get3DPoint(Grid);
-		return FPoint::DistSquared(Point1, Point2);
+		const FVector& Point1 = FirstNode->Get3DPoint(Grid);
+		const FVector& Point2 = SecondNode->Get3DPoint(Grid);
+		return FVector::DistSquared(Point1, Point2);
 	}
 };
 
diff --git a/Engine/Source/Runtime/Datasmith/CADKernel/Base/Private/Mesh/Meshers/IsoTriangulatorRemoveIntersection.cpp b/Engine/Source/Runtime/Datasmith/CADKernel/Base/Private/Mesh/Meshers/IsoTriangulatorRemoveIntersection.cpp
--- a/Engine/Source/Runtime/Datasmith/CADKernel/Base/Private/Mesh/Meshers/IsoTriangulatorRemoveIntersection.cpp
+++ b/Engine/Source/Runtime/Datasmith/CADKernel/Base/Private/Mesh/Meshers/IsoTriangulatorRemoveIntersection.cpp
@@ -31,14 +31,14 @@ bool FIsoTriangulator::RemovePickToOutside(const TArray<FLoopNode*>& NodesOfLoop
 
 	const int32 NodeCount = NodesOfLoop.Num();
 
-	TFunction<FPoint2D(double, FLoopNode*, FLoopNode*)> IntersectingPoint = [&](double Coordinate, FLoopNode* Start, FLoopNode* End) -> FPoint2D
+	TFunction<FVector2d(double, FLoopNode*, FLoopNode*)> IntersectingPoint = [&](double Coordinate, FLoopNode* Start, FLoopNode* End) -> FVector2d
 	{
-		const FPoint2D& StartPoint = Start->Get2DPoint(EGridSpace::UniformScaled, Grid);
-		const FPoint2D& EndPoint = End->Get2DPoint(EGridSpace::UniformScaled, Grid);
+		const FVector2d& StartPoint = Start->Get2DPoint(EGridSpace::UniformScaled, Grid);
+		const FVector2d& EndPoint = End->Get2DPoint(EGridSpace::UniformScaled, Grid);
 
 		int32 StartIndex = (int32)Coordinate;
 		Coordinate -= StartIndex;
-		FPoint2D Intersection = PointOnSegment(StartPoint, EndPoint, Coordinate);
+		FVector2d Intersection = PointOnSegment<FVector2d>(StartPoint, EndPoint, Coordinate);
 
 		//{
 		//	F3DDebugSession _(TEXT("Intersected Point"));
@@ -58,12 +58,12 @@ bool FIsoTriangulator::RemovePickToOutside(const TArray<FLoopNode*>& NodesOfLoop
 	}
 
 	FLoopNode* StartNode = GetPrevious(TmpNode);
-	FPoint2D FirstIntersection = IntersectingPoint(Intersection.Value, StartNode, TmpNode);
+	FVector2d FirstIntersection = IntersectingPoint(Intersection.Value, StartNode, TmpNode);
 
 	TmpNode = GetNext(EndNode);
-	FPoint2D SecondIntersection = IntersectingPoint(NextIntersection.Value, EndNode, TmpNode);
+	FVector2d SecondIntersection = IntersectingPoint(NextIntersection.Value, EndNode, TmpNode);
 
-	FPoint2D MiddlePoint = (FirstIntersection + SecondIntersection) * 0.5;
+	FVector2d MiddlePoint = (FirstIntersection + SecondIntersection) * 0.5;
 
 	TmpNode = GetNext(StartNode);
 	while (TmpNode && (TmpNode != EndNode) && !TmpNode->IsDelete())
@@ -75,7 +75,7 @@ bool FIsoTriangulator::RemovePickToOutside(const TArray<FLoopNode*>& NodesOfLoop
 		TmpNode = GetNext(StartNode);
 	}
 
-	FPoint2D MoveDirection = SecondIntersection - FirstIntersection;
+	FVector2d MoveDirection = SecondIntersection - FirstIntersection;
 	double Length = MoveDirection.Length();
 	if (FMath::IsNearlyZero(Length))
 	{
@@ -89,7 +89,7 @@ bool FIsoTriangulator::RemovePickToOutside(const TArray<FLoopNode*>& NodesOfLoop
 	}
 
 	MoveDirection /= Length;
-	MoveDirection = FPoint2D(-MoveDirection.Y, MoveDirection.X);
+	MoveDirection = FVector2d(-MoveDirection.Y, MoveDirection.X);
 	MoveDirection *= GeometricTolerance;
 	MiddlePoint += MoveDirection;
 
@@ -171,10 +171,10 @@ bool FIsoTriangulator::SpreadCoincidentNodes(const TArray<FLoopNode*>& NodesOfLo
 	//FLoopNode* Node1 = GetNext(Node0);
 	//FLoopNode* NextNode = GetNext(Node1);
 
-	//const FPoint2D* PreviousPoint = &PreviousNode->Get2DPoint(EGridSpace::UniformScaled, Grid);
-	//const FPoint2D* Point0 = &Node0->Get2DPoint(EGridSpace::UniformScaled, Grid);
-	//const FPoint2D* Point1 = &Node1->Get2DPoint(EGridSpace::UniformScaled, Grid);
-	//const FPoint2D* NextPoint = &NextNode->Get2DPoint(EGridSpace::UniformScaled, Grid);
+	//const FVector2d* PreviousPoint = &PreviousNode->Get2DPoint(EGridSpace::UniformScaled, Grid);
+	//const FVector2d* Point0 = &Node0->Get2DPoint(EGridSpace::UniformScaled, Grid);
+	//const FVector2d* Point1 = &Node1->Get2DPoint(EGridSpace::UniformScaled, Grid);
+	//const FVector2d* NextPoint = &NextNode->Get2DPoint(EGridSpace::UniformScaled, Grid);
 	return true;
 }
 
@@ -350,8 +350,8 @@ bool FIsoTriangulator::TryToRemoveIntersectionOfTwoConsecutiveIntersectingSegmen
 	FLoopNode* PreviousNode = nullptr;
 	FLoopNode* NextNode = nullptr;
 
-	TSegment<FPoint2D> IntersectingSegment2D(IntersectingSegment.GetFirstNode().Get2DPoint(EGridSpace::UniformScaled, Grid), IntersectingSegment.GetSecondNode().Get2DPoint(EGridSpace::UniformScaled, Grid));
-	TSegment<FPoint2D> Segment2D(Segment.GetFirstNode().Get2DPoint(EGridSpace::UniformScaled, Grid), Segment.GetSecondNode().Get2DPoint(EGridSpace::UniformScaled, Grid));
+	TSegment<FVector2d> IntersectingSegment2D(IntersectingSegment.GetFirstNode().Get2DPoint(EGridSpace::UniformScaled, Grid), IntersectingSegment.GetSecondNode().Get2DPoint(EGridSpace::UniformScaled, Grid));
+	TSegment<FVector2d> Segment2D(Segment.GetFirstNode().Get2DPoint(EGridSpace::UniformScaled, Grid), Segment.GetSecondNode().Get2DPoint(EGridSpace::UniformScaled, Grid));
 
 	double IntersectingSegmentSlop = ComputeOrientedSlope(IntersectingSegment2D.Point0, IntersectingSegment2D.Point1, 0);
 	double SegmentSlop = ComputeUnorientedSlope(Segment2D.Point1, Segment2D.Point0, IntersectingSegmentSlop);
@@ -363,8 +363,8 @@ bool FIsoTriangulator::TryToRemoveIntersectionOfTwoConsecutiveIntersectingSegmen
 	// if the segment and IntersectingSegment are parallel, segment are moved inside
 	if (SegmentSlop < 0.01)
 	{
-		double StartPointSquareDistance = SquareDistanceOfPointToSegment(Segment2D.Point0, IntersectingSegment2D.Point0, IntersectingSegment2D.Point1);
-		double EndPointSquareDistance = SquareDistanceOfPointToSegment(Segment2D.Point1, IntersectingSegment2D.Point0, IntersectingSegment2D.Point1);
+		double StartPointSquareDistance = SquareDistanceOfPointToSegment<FVector2d>(Segment2D.Point0, IntersectingSegment2D.Point0, IntersectingSegment2D.Point1);
+		double EndPointSquareDistance = SquareDistanceOfPointToSegment<FVector2d>(Segment2D.Point1, IntersectingSegment2D.Point0, IntersectingSegment2D.Point1);
 		if (StartPointSquareDistance < SquareGeometricTolerance && EndPointSquareDistance < SquareGeometricTolerance)
 		{
 			OffsetSegment(Segment, Segment2D, IntersectingSegment2D);
@@ -415,19 +415,19 @@ bool FIsoTriangulator::TryToRemoveIntersectionOfTwoConsecutiveIntersectingSegmen
 	}
 #endif
 
-	TSegment<FPoint2D> NextSegment2D(Node->Get2DPoint(EGridSpace::UniformScaled, Grid), NextNode->Get2DPoint(EGridSpace::UniformScaled, Grid));
+	TSegment<FVector2d> NextSegment2D(Node->Get2DPoint(EGridSpace::UniformScaled, Grid), NextNode->Get2DPoint(EGridSpace::UniformScaled, Grid));
 	if (!FastIntersectSegments2D(NextSegment2D, IntersectingSegment2D))
 	{
 		return false;
 	}
 
-	TSegment<FPoint2D>  PreviousSegment2D(Node->Get2DPoint(EGridSpace::UniformScaled, Grid), PreviousNode->Get2DPoint(EGridSpace::UniformScaled, Grid));
+	TSegment<FVector2d>  PreviousSegment2D(Node->Get2DPoint(EGridSpace::UniformScaled, Grid), PreviousNode->Get2DPoint(EGridSpace::UniformScaled, Grid));
 
-	FPoint2D Intersection1 = FindIntersectionOfSegments2D(PreviousSegment2D, IntersectingSegment2D);
-	FPoint2D Intersection2 = FindIntersectionOfSegments2D(NextSegment2D, IntersectingSegment2D);
+	FVector2d Intersection1 = FindIntersectionOfSegments2D(PreviousSegment2D, IntersectingSegment2D);
+	FVector2d Intersection2 = FindIntersectionOfSegments2D(NextSegment2D, IntersectingSegment2D);
 
 	double Coordinate;
-	FPoint2D ProjectedPoint = ProjectPointOnSegment(Node->Get2DPoint(EGridSpace::UniformScaled, Grid), Intersection1, Intersection2, Coordinate);
+	FVector2d ProjectedPoint = ProjectPointOnSegment<FVector2d>(Node->Get2DPoint(EGridSpace::UniformScaled, Grid), Intersection1, Intersection2, Coordinate);
 
 #ifdef DEBUG_TWO_CONSECUTIVE_INTERSECTING
 	if (bDisplay)
@@ -437,13 +437,13 @@ bool FIsoTriangulator::TryToRemoveIntersectionOfTwoConsecutiveIntersectingSegmen
 	}
 #endif
 
-	FPoint2D SegmentTangent = IntersectingSegment2D.Point1 - IntersectingSegment2D.Point0;
+	FVector2d SegmentTangent = IntersectingSegment2D.Point1 - IntersectingSegment2D.Point0;
 	SegmentTangent.Normalize();
-	FPoint2D MoveDirection = FPoint2D(-SegmentTangent.Y, SegmentTangent.X);
+	FVector2d MoveDirection = FVector2d(-SegmentTangent.Y, SegmentTangent.X);
 
 	MoveDirection *= GeometricTolerance;
 
-	FPoint2D NewCoordinate = ProjectedPoint + MoveDirection;
+	FVector2d NewCoordinate = ProjectedPoint + MoveDirection;
 
 #ifdef DEBUG_TWO_CONSECUTIVE_INTERSECTING
 	if (bDisplay)
@@ -638,10 +638,10 @@ void FIsoTriangulator::RemovePickOfLoop(FIsoSegment& Segment)
 	FLoopNode* PreviousNode = &Node0->GetPreviousNode();
 	FLoopNode* NextNode = &Node1->GetNextNode();
 
-	const FPoint2D* PreviousPoint = &PreviousNode->Get2DPoint(EGridSpace::UniformScaled, Grid);
-	const FPoint2D* Point0 = &Node0->Get2DPoint(EGridSpace::UniformScaled, Grid);
-	const FPoint2D* Point1 = &Node1->Get2DPoint(EGridSpace::UniformScaled, Grid);
-	const FPoint2D* NextPoint = &NextNode->Get2DPoint(EGridSpace::UniformScaled, Grid);
+	const FVector2d* PreviousPoint = &PreviousNode->Get2DPoint(EGridSpace::UniformScaled, Grid);
+	const FVector2d* Point0 = &Node0->Get2DPoint(EGridSpace::UniformScaled, Grid);
+	const FVector2d* Point1 = &Node1->Get2DPoint(EGridSpace::UniformScaled, Grid);
+	const FVector2d* NextPoint = &NextNode->Get2DPoint(EGridSpace::UniformScaled, Grid);
 
 	bool bNodeRemoved = true;
 	bool bPickRemoved = false;
diff --git a/Engine/Source/Runtime/Datasmith/CADKernel/Base/Private/Mesh/Meshers/ParametricFaceMesher.cpp b/Engine/Source/Runtime/Datasmith/CADKernel/Base/Private/Mesh/Meshers/ParametricFaceMesher.cpp
--- a/Engine/Source/Runtime/Datasmith/CADKernel/Base/Private/Mesh/Meshers/ParametricFaceMesher.cpp
+++ b/Engine/Source/Runtime/Datasmith/CADKernel/Base/Private/Mesh/Meshers/ParametricFaceMesher.cpp
@@ -355,10 +355,10 @@ void FParametricFaceMesher::DefineImposedCuttingPointsBasedOnOtherSideMesh(FThin
 
 	FThinZoneSide& FrontSide = SideToConstrain.GetFrontThinZoneSide();
 
-	TMap<int32, FPoint2D> ExistingMeshNodes;
+	TMap<int32, FVector2d> ExistingMeshNodes;
 	TArray<FCrossZoneElement> CrossZoneElements;
 
-	FAddMeshNodeFunc AddToCrossZoneElements = [&CrossZoneElements](const int32 NodeIndice, const FPoint2D& MeshNode2D, double MeshingTolerance3D, const FEdgeSegment& EdgeSegment, const FPairOfIndex& OppositeNodeIndices)
+	FAddMeshNodeFunc AddToCrossZoneElements = [&CrossZoneElements](const int32 NodeIndice, const FVector2d& MeshNode2D, double MeshingTolerance3D, const FEdgeSegment& EdgeSegment, const FPairOfIndex& OppositeNodeIndices)
 	{
 		if (CrossZoneElements.Num() && CrossZoneElements.Last().VertexId >= 0 && CrossZoneElements.Last().VertexId == NodeIndice)
 		{
@@ -370,7 +370,7 @@ void FParametricFaceMesher::DefineImposedCuttingPointsBasedOnOtherSideMesh(FThin
 		}
 	};
 
-	FAddMeshNodeFunc AddToExistingMeshNodes = [&ExistingMeshNodes](const int32 NodeIndice, const FPoint2D& MeshNode2D, double MeshingTolerance3D, const FEdgeSegment& EdgeSegment, const FPairOfIndex& OppositeNodeIndices)
+	FAddMeshNodeFunc AddToExistingMeshNodes = [&ExistingMeshNodes](const int32 NodeIndice, const FVector2d& MeshNode2D, double MeshingTolerance3D, const FEdgeSegment& EdgeSegment, const FPairOfIndex& OppositeNodeIndices)
 	{
 		ExistingMeshNodes.Emplace(NodeIndice, MeshNode2D);
 	};
@@ -395,7 +395,7 @@ void FParametricFaceMesher::DefineImposedCuttingPointsBasedOnOtherSideMesh(FThin
 	{
 		if (CrossZoneElement.OppositeVertexIndices[0] >= 0)
 		{
-			FPoint2D* OppositeVertex = ExistingMeshNodes.Find(CrossZoneElement.OppositeVertexIndices[0]);
+			FVector2d* OppositeVertex = ExistingMeshNodes.Find(CrossZoneElement.OppositeVertexIndices[0]);
 			if (OppositeVertex)
 			{
 				CrossZoneElement.OppositePoint2D = *OppositeVertex;
@@ -405,11 +405,11 @@ void FParametricFaceMesher::DefineImposedCuttingPointsBasedOnOtherSideMesh(FThin
 		}
 
 		double MinSquareThickness = MaxSquareThickness;
-		FPoint2D ClosePoint;
+		FVector2d ClosePoint = FVector2d::ZeroVector;
 		FEdgeSegment* CloseSegment = nullptr;
 		double ClosePointCoordinate = -1;
 
-		const FPoint2D MeshNodeCoordinate = CrossZoneElement.VertexPoint2D;
+		const FVector2d MeshNodeCoordinate = CrossZoneElement.VertexPoint2D;
 
 		for (FEdgeSegment& Segment : SideToConstrain.GetSegments())
 		{
@@ -422,9 +422,9 @@ void FParametricFaceMesher::DefineImposedCuttingPointsBasedOnOtherSideMesh(FThin
 			}
 
 			double CoordSegmentU;
-			FPoint2D Projection = Segment.ProjectPoint(MeshNodeCoordinate, CoordSegmentU);
+			FVector2d Projection = Segment.ProjectPoint(MeshNodeCoordinate, CoordSegmentU);
 
-			const double SquareDistance = FPoint2D::DistSquared(MeshNodeCoordinate, Projection);
+			const double SquareDistance = FVector2d::DistSquared(MeshNodeCoordinate, Projection);
 			if (MinSquareThickness > SquareDistance)
 			{
 				FEdgeSegment* SegmentPtr = &Segment;
@@ -697,7 +697,7 @@ void FParametricFaceMesher::Mesh(FTopologicalEdge& InEdge, bool bFinalMeshing)
 	}
 
 	// Add Edge intersection with inner surface grid Iso
-	FPoint2D ExtremityTolerances = InEdge.GetCurve()->GetExtremityTolerances(EdgeBounds);
+	FVector2d ExtremityTolerances = InEdge.GetCurve()->GetExtremityTolerances(EdgeBounds);
 	double EdgeTolerance = FMath::Min(ExtremityTolerances[0], ExtremityTolerances[1]);
 	if (!EdgeIntersectionWithIsoU_Coordinates.IsEmpty())
 	{
@@ -900,7 +900,7 @@ bool FParametricFaceMesher::MeshPlanarFace()
 	TArray<FGeneralPolygon2d> PolygonsOut;
 	PolygonsOut.Reserve(Face.GetLoops().Num());
 
-	const TArray<TArray<FPoint2D>>& Loops2D = Grid.GetLoops2D(EGridSpace::Default2D);
+	const TArray<TArray<FVector2d>>& Loops2D = Grid.GetLoops2D(EGridSpace::Default2D);
 	const TArray<TSharedPtr<FTopologicalLoop>>& Loops = Face.GetLoops();
 	ensureCADKernel(Loops.Num() == Loops2D.Num());
 
@@ -908,10 +908,10 @@ bool FParametricFaceMesher::MeshPlanarFace()
 	bool bOuterIsCW = false;
 	for (int32 Index = 0; Index < Loops.Num(); ++Index)
 	{
-		const TArray<FPoint2D>& PointList = Loops2D[Index];
+		const TArray<FVector2d>& PointList = Loops2D[Index];
 		TArray<FVector2d> VertexList;
 		VertexList.Reserve(PointList.Num());
-		for (const FPoint2D& Point : PointList)
+		for (const FVector2d& Point : PointList)
 		{
 			VertexList.Emplace(Point.X, Point.Y);
 		}
@@ -947,7 +947,7 @@ bool FParametricFaceMesher::MeshPlanarFace()
 
 	FFaceMesh& FaceMesh = Face.GetOrCreateMesh(MeshModel);
 
-	TArray<FPoint>& NodeCoordinates = FaceMesh.GetNodeCoordinates();
+	TArray<FVector>& NodeCoordinates = FaceMesh.GetNodeCoordinates();
 	TArray<int32>& TrianglesVerticesIndex = FaceMesh.TrianglesVerticesIndex;
 	TArray<int32>& VerticesGlobalIndex = FaceMesh.VerticesGlobalIndex;
 	TArray<FVector3f>& Normals = FaceMesh.Normals;
@@ -995,9 +995,9 @@ bool FParametricFaceMesher::MeshPlanarFace()
 		int32 Index = StartIndex;
 		for (const FVector2d& Vertex : Vertices)
 		{
-			FPoint Point;
+			FVector Point;
 			FVector3f Normal;
-			CarrierSurface.EvaluatePointAndNormal(FPoint2D(Vertex.X, Vertex.Y), Point, Normal);
+			CarrierSurface.EvaluatePointAndNormal(FVector2d(Vertex.X, Vertex.Y), Point, Normal);
 			
 			NodeCoordinates.Emplace(Point);
 			Normals.Emplace(Normal);
diff --git a/Engine/Source/Runtime/Datasmith/CADKernel/Base/Private/Mesh/Meshers/ParametricFaceMesher.h b/Engine/Source/Runtime/Datasmith/CADKernel/Base/Private/Mesh/Meshers/ParametricFaceMesher.h
--- a/Engine/Source/Runtime/Datasmith/CADKernel/Base/Private/Mesh/Meshers/ParametricFaceMesher.h
+++ b/Engine/Source/Runtime/Datasmith/CADKernel/Base/Private/Mesh/Meshers/ParametricFaceMesher.h
@@ -21,12 +21,12 @@ namespace ParametricMesherTool
 struct FCrossZoneElement
 {
 	const int32 VertexId;
-	FPoint2D VertexPoint2D;
+	FVector2d VertexPoint2D;
 	const double Tolerance3D;
 	const FEdgeSegment* Segment;
 	FPairOfIndex OppositeVertexIndices;
 
-	FPoint2D OppositePoint2D;
+	FVector2d OppositePoint2D;
 	FEdgeSegment* OppositeSegment;
 	double OppositePointCoordinate = -1.;
 
@@ -34,7 +34,7 @@ struct FCrossZoneElement
 
 	bool bIsSelected = false;
 
-	FCrossZoneElement(const int32 InVertexId, const FPoint2D InPoint2D, const double InTolerance3D, const FEdgeSegment* InSegment, const FPairOfIndex& InOppositeVertexIndices)
+	FCrossZoneElement(const int32 InVertexId, const FVector2d InPoint2D, const double InTolerance3D, const FEdgeSegment* InSegment, const FPairOfIndex& InOppositeVertexIndices)
 		: VertexId(InVertexId)
 		, VertexPoint2D(InPoint2D)
 		, Tolerance3D(InTolerance3D)
diff --git a/Engine/Source/Runtime/Datasmith/CADKernel/Base/Private/Mesh/Structure/FaceMesh.cpp b/Engine/Source/Runtime/Datasmith/CADKernel/Base/Private/Mesh/Structure/FaceMesh.cpp
--- a/Engine/Source/Runtime/Datasmith/CADKernel/Base/Private/Mesh/Structure/FaceMesh.cpp
+++ b/Engine/Source/Runtime/Datasmith/CADKernel/Base/Private/Mesh/Structure/FaceMesh.cpp
@@ -12,7 +12,7 @@
 namespace UE::CADKernel
 {
 
-void FFaceMesh::GetNodeIdToCoordinates(TMap<int32, const FPoint*>& NodeIdToCoordinates) const
+void FFaceMesh::GetNodeIdToCoordinates(TMap<int32, const FVector*>& NodeIdToCoordinates) const
 {
 	const FTopologicalFace& Face = (const FTopologicalFace&) GetGeometricEntity();
 
@@ -43,7 +43,7 @@ void FFaceMesh::GetNodeIdToCoordinates(TMap<int32, const FPoint*>& NodeIdToCoord
 			const FEdgeMesh* EdgeMesh = ActiveEdge->GetMesh();
 
 			TArray<int32> NodeIds = EdgeMesh->EdgeVerticesIndex;
-			const FPoint* EdgeNodeCoordinates = EdgeMesh->GetNodeCoordinates().GetData();
+			const FVector* EdgeNodeCoordinates = EdgeMesh->GetNodeCoordinates().GetData();
 			if(EdgeNodeCoordinates)
 			{
 				for (int32 Index = 1; Index <NodeIds.Num() - 1; Index++)
@@ -54,7 +54,7 @@ void FFaceMesh::GetNodeIdToCoordinates(TMap<int32, const FPoint*>& NodeIdToCoord
 		}
 	}
 
-	const FPoint* Coordinates = GetNodeCoordinates().GetData();
+	const FVector* Coordinates = GetNodeCoordinates().GetData();
 	for (int32 Index = 0; Index < GetNodeCoordinates().Num(); ++Index)
 	{
 		NodeIdToCoordinates.Add(StartNodeId+Index, Coordinates + Index);
diff --git a/Engine/Source/Runtime/Datasmith/CADKernel/Base/Private/Mesh/Structure/Grid.cpp b/Engine/Source/Runtime/Datasmith/CADKernel/Base/Private/Mesh/Structure/Grid.cpp
--- a/Engine/Source/Runtime/Datasmith/CADKernel/Base/Private/Mesh/Structure/Grid.cpp
+++ b/Engine/Source/Runtime/Datasmith/CADKernel/Base/Private/Mesh/Structure/Grid.cpp
@@ -121,7 +121,7 @@ void FGrid::GetPreferredUVCuttingParametersFromLoops(FCuttingGrid& CuttingParame
 					continue;
 				}
 
-				const TArray<FPoint>& EdgeMeshNodes = EdgeMesh->GetNodeCoordinates();
+				const TArray<FVector>& EdgeMeshNodes = EdgeMesh->GetNodeCoordinates();
 				if (EdgeMeshNodes.Num() == 0)
 				{
 					continue;
@@ -140,7 +140,7 @@ void FGrid::GetPreferredUVCuttingParametersFromLoops(FCuttingGrid& CuttingParame
 				ProjectedPointCoords.Add(Edge->GetBoundary().GetMax());
 			}
 
-			TArray<FPoint2D> EdgePoints2D;
+			TArray<FVector2d> EdgePoints2D;
 			Edge->Approximate2DPoints(ProjectedPointCoords, EdgePoints2D);
 
 			for (int32 Index = 0; Index < EdgePoints2D.Num(); ++Index)
@@ -220,8 +220,8 @@ void FGrid::FindPointsCloseToLoop()
 	int32 Index[2] = { 1, 1 };
 	int32 GlobalIndex = 1;
 
-	const FPoint2D* PointA = nullptr;
-	const FPoint2D* PointB = nullptr;
+	const FVector2d* PointA = nullptr;
+	const FVector2d* PointB = nullptr;
 
 	// Find start index
 	TFunction<void(EIso)> FindPointAIndex = [&](EIso Iso)
@@ -356,7 +356,7 @@ void FGrid::FindPointsCloseToLoop()
 
 	double ABv = 0.;
 	double ABu = 0.;
-	for (const TArray<FPoint2D>& Loop : FaceLoops2D[EGridSpace::UniformScaled])
+	for (const TArray<FVector2d>& Loop : FaceLoops2D[EGridSpace::UniformScaled])
 	{
 		PointA = &Loop.Last();
 
@@ -449,8 +449,8 @@ void FGrid::RemovePointsCloseToLoop()
 
 	struct FGridSegment
 	{
-		const FPoint2D* StartPoint;
-		const FPoint2D* EndPoint;
+		const FVector2d* StartPoint;
+		const FVector2d* EndPoint;
 		double StartPointWeight;
 		double EndPointWeight;
 		double UMin;
@@ -458,7 +458,7 @@ void FGrid::RemovePointsCloseToLoop()
 		double UMax;
 		double VMax;
 
-		FGridSegment(const FPoint2D& SPoint, const FPoint2D& EPoint)
+		FGridSegment(const FVector2d& SPoint, const FVector2d& EPoint)
 			: StartPoint(&SPoint)
 			, EndPoint(&EPoint)
 		{
@@ -495,13 +495,13 @@ void FGrid::RemovePointsCloseToLoop()
 	TArray<FGridSegment> LoopSegments;
 	{
 		int32 SegmentNum = 0;
-		for (const TArray<FPoint2D>& Loop : FaceLoops2D[EGridSpace::UniformScaled])
+		for (const TArray<FVector2d>& Loop : FaceLoops2D[EGridSpace::UniformScaled])
 		{
 			SegmentNum += Loop.Num();
 		}
 		LoopSegments.Reserve(SegmentNum);
 
-		for (TArray<FPoint2D>& Loop : FaceLoops2D[EGridSpace::UniformScaled])
+		for (TArray<FVector2d>& Loop : FaceLoops2D[EGridSpace::UniformScaled])
 		{
 			for (int32 Index = 0; Index < Loop.Num() - 1; ++Index)
 			{
@@ -599,7 +599,7 @@ void FGrid::RemovePointsCloseToLoop()
 	for (const int32& SortedIndex : SortedPointIndexes)
 	{
 		int32 Index = IndexOfPointsNearAndInsideLoop[SortedIndex];
-		const FPoint2D& Point2D = Points2D[EGridSpace::UniformScaled][Index];
+		const FVector2d& Point2D = Points2D[EGridSpace::UniformScaled][Index];
 
 		double DeltaU;
 		double DeltaV;
@@ -640,10 +640,10 @@ void FGrid::RemovePointsCloseToLoop()
 			}
 
 			double Coordinate;
-			FPoint2D Projection = ProjectPointOnSegment(Point2D, *Segment.StartPoint, *Segment.EndPoint, Coordinate, /*bRestrictCoodinateToInside*/ true);
+			FVector2d Projection = ProjectPointOnSegment<FVector2d>(Point2D, *Segment.StartPoint, *Segment.EndPoint, Coordinate, /*bRestrictCoodinateToInside*/ true);
 
 			// If Projected point is in the oval center on Point2D => the node is too close
-			FPoint2D ProjectionToPoint = (Point2D - Projection).GetAbs();
+			FVector2d ProjectionToPoint = (Point2D - Projection).GetAbs();
 
 			if (ProjectionToPoint.X > DeltaU * 0.05 || ProjectionToPoint.Y > DeltaV * 0.05)
 			{
@@ -661,10 +661,10 @@ void FGrid::RemovePointsCloseToLoop()
  * The normal is swap if StartPoint is too close to the Boundary
  * The norm of the normal is defined as 1/20 of the parallel boundary Length
  */
-void ScaleAndSwap(FPoint2D& Normal, const FPoint2D& StartPoint, const FSurfacicBoundary& Boundary)
+void ScaleAndSwap(FVector2d& Normal, const FVector2d& StartPoint, const FSurfacicBoundary& Boundary)
 {
 	Normal.Normalize();
-	FPoint2D MainDirection = Normal;
+	FVector2d MainDirection = Normal;
 	MainDirection.X /= Boundary[EIso::IsoU].Length();
 	MainDirection.Y /= Boundary[EIso::IsoV].Length();
 
@@ -700,13 +700,13 @@ void ScaleAndSwap(FPoint2D& Normal, const FPoint2D& StartPoint, const FSurfacicB
 /**
  * Displace loop nodes inside to avoid that the nodes are outside the surface Boundary, so outside the grid
  */
-void SlightlyDisplacedPolyline(TArray<FPoint2D>& D2Points, const FSurfacicBoundary& Boundary)
+void SlightlyDisplacedPolyline(TArray<FVector2d>& D2Points, const FSurfacicBoundary& Boundary)
 {
-	FPoint2D Normal;
+	FVector2d Normal;
 	for (int32 Index = 0; Index < D2Points.Num() - 1; ++Index)
 	{
-		const FPoint2D Tangent = D2Points[Index + 1] - D2Points[Index];
-		Normal = FPoint2D(-Tangent.Y, Tangent.X);
+		const FVector2d Tangent = D2Points[Index + 1] - D2Points[Index];
+		Normal = FVector2d(-Tangent.Y, Tangent.X);
 
 		ScaleAndSwap(Normal, D2Points[Index], Boundary);
 		D2Points[Index] += Normal;
@@ -725,10 +725,10 @@ void FGrid::GetMeshOfLoop(const FTopologicalLoop& Loop)
 		LoopNodeCount += Edge.Entity->GetLinkActiveEdge()->GetMesh()->GetNodeCount() + 2;
 	}
 
-	TArray<FPoint2D>& Loop2D = FaceLoops2D[EGridSpace::Default2D].Emplace_GetRef();
+	TArray<FVector2d>& Loop2D = FaceLoops2D[EGridSpace::Default2D].Emplace_GetRef();
 	Loop2D.Empty(LoopNodeCount);
 
-	TArray<FPoint>& Loop3D = FaceLoops3D.Emplace_GetRef();
+	TArray<FVector>& Loop3D = FaceLoops3D.Emplace_GetRef();
 	Loop3D.Reserve(LoopNodeCount);
 
 	TArray<FVector3f>& LoopNormals = NormalsOfFaceLoops.Emplace_GetRef();
@@ -776,7 +776,7 @@ void FGrid::GetMeshOfLoop(const FTopologicalLoop& Loop)
 
 			CuttingPolyline.Points3D.Init(ActiveEdge->GetStartBarycenter(), CuttingPolyline.Coordinates.Num());
 
-			TArray<FPoint2D> D2Points = CuttingPolyline.Points2D;
+			TArray<FVector2d> D2Points = CuttingPolyline.Points2D;
 			const FSurfacicBoundary& Boundary = Edge->GetCurve()->GetCarrierSurface()->GetBoundary();
 			// to compute the normals, the 2D points are slightly displaced perpendicular to the curve
 			SlightlyDisplacedPolyline(D2Points, Boundary);
@@ -786,7 +786,7 @@ void FGrid::GetMeshOfLoop(const FTopologicalLoop& Loop)
 		{
 			if (EdgeCuttingPointCoordinates.IsEmpty())
 			{
-				const TArray<FPoint>& MeshVertex3D = ActiveEdge->GetOrCreateMesh(MeshModel).GetNodeCoordinates();
+				const TArray<FVector>& MeshVertex3D = ActiveEdge->GetOrCreateMesh(MeshModel).GetNodeCoordinates();
 				TArray<double> ProjectedPointCoords;
 
 				CuttingPolyline.Coordinates.Reserve(MeshVertex3D.Num() + 2);
@@ -938,7 +938,7 @@ void FGrid::GetMeshOfThinZone(const FThinZone2D& ThinZone)
 	// Return the mesh defined by the node index
 	TFunction<void(const FThinZoneSide&, TArray<int32>&)> GetThinZoneSideMesh = [this](const FThinZoneSide& Side, TArray<int32>& MeshIndices)
 	{
-		FAddMeshNodeFunc AddMeshNode = [&MeshIndices](int32 NodeIndice, const FPoint2D& MeshNode2D, double MeshingTolerance, const FEdgeSegment& EdgeSegment, const FPairOfIndex& OppositeNodeIndices)
+		FAddMeshNodeFunc AddMeshNode = [&MeshIndices](int32 NodeIndice, const FVector2d& MeshNode2D, double MeshingTolerance, const FEdgeSegment& EdgeSegment, const FPairOfIndex& OppositeNodeIndices)
 		{
 			if (MeshIndices.IsEmpty() || MeshIndices.Last() != NodeIndice)
 			{
@@ -973,7 +973,7 @@ void FGrid::GetMeshOfThinZone(const FThinZone2D& ThinZone)
 		return false;
 	};
 
-	TFunction<void(const TArray<int32>&, TArray<FPoint2D>&)> GetThinZoneMeshCoordinates = [&NodeIdsOfLoops = NodeIdsOfFaceLoops, &Loops2D = FaceLoops2D, &FindLoopAndNodeIndex](const TArray<int32>& ThinZoneNodeIds, TArray<FPoint2D>& ThinZonePoints)
+	TFunction<void(const TArray<int32>&, TArray<FVector2d>&)> GetThinZoneMeshCoordinates = [&NodeIdsOfLoops = NodeIdsOfFaceLoops, &Loops2D = FaceLoops2D, &FindLoopAndNodeIndex](const TArray<int32>& ThinZoneNodeIds, TArray<FVector2d>& ThinZonePoints)
 	{
 		ThinZonePoints.Reset(ThinZoneNodeIds.Num());
 
@@ -1012,7 +1012,7 @@ void FGrid::GetMeshOfThinZone(const FThinZone2D& ThinZone)
 		}
 	};
 
-	TFunction<void(TArray<FPoint2D>&)> RemoveDuplicatedNode = [](TArray<FPoint2D>& ThinZoneMesh)
+	TFunction<void(TArray<FVector2d>&)> RemoveDuplicatedNode = [](TArray<FVector2d>& ThinZoneMesh)
 	{
 		if (ThinZoneMesh.Num() > 1)
 		{
@@ -1020,7 +1020,7 @@ void FGrid::GetMeshOfThinZone(const FThinZone2D& ThinZone)
 			constexpr double Tolerance = DOUBLE_SMALL_NUMBER;
 			for (int32 Index = ThinZoneMesh.Num() - 1; Index > 0; --Index)
 			{
-				double SquareDist = FPoint2D::DistSquared(ThinZoneMesh[Index - 1], ThinZoneMesh[Index]);
+				double SquareDist = FVector2d::DistSquared(ThinZoneMesh[Index - 1], ThinZoneMesh[Index]);
 				if (SquareDist < Tolerance)
 				{
 					ThinZoneMesh.RemoveAt(Index);
@@ -1028,7 +1028,7 @@ void FGrid::GetMeshOfThinZone(const FThinZone2D& ThinZone)
 			}
 			if (ThinZoneMesh.Num() > 2)
 			{
-				double SquareDist = FPoint2D::DistSquared(ThinZoneMesh[0], ThinZoneMesh.Last());
+				double SquareDist = FVector2d::DistSquared(ThinZoneMesh[0], ThinZoneMesh.Last());
 				if (SquareDist < Tolerance)
 				{
 					ThinZoneMesh.RemoveAt(ThinZoneMesh.Num() - 1);
@@ -1051,7 +1051,7 @@ void FGrid::GetMeshOfThinZone(const FThinZone2D& ThinZone)
 
 	if (FirstSideIsClosed && SecondSideIsClosed)
 	{
-		TArray<FPoint2D>& ThinZoneMesh = FaceLoops2D[EGridSpace::UniformScaled].Emplace_GetRef();
+		TArray<FVector2d>& ThinZoneMesh = FaceLoops2D[EGridSpace::UniformScaled].Emplace_GetRef();
 		GetThinZoneMeshCoordinates(ThinZoneNodeIds, ThinZoneMesh);
 		RemoveDuplicatedNode(ThinZoneMesh);
 		ThinZoneNodeIds.Reset();
@@ -1068,7 +1068,7 @@ void FGrid::GetMeshOfThinZone(const FThinZone2D& ThinZone)
 		return;
 	}
 
-	TArray<FPoint2D>& ThinZoneMesh = FaceLoops2D[EGridSpace::UniformScaled].Emplace_GetRef();
+	TArray<FVector2d>& ThinZoneMesh = FaceLoops2D[EGridSpace::UniformScaled].Emplace_GetRef();
 	GetThinZoneMeshCoordinates(ThinZoneNodeIds, ThinZoneMesh);
 	RemoveDuplicatedNode(ThinZoneMesh);
 
@@ -1124,9 +1124,9 @@ bool FGrid::GetMeshOfLoops()
 
 	// Fit loops to Surface bounds.
 	const FSurfacicBoundary& Bounds = Face.GetBoundary();
-	for (TArray<FPoint2D>& Loop : FaceLoops2D[EGridSpace::Default2D])
+	for (TArray<FVector2d>& Loop : FaceLoops2D[EGridSpace::Default2D])
 	{
-		for (FPoint2D& Point : Loop)
+		for (FVector2d& Point : Loop)
 		{
 			Bounds.MoveInsideIfNot(Point);
 		}
@@ -1137,9 +1137,9 @@ bool FGrid::GetMeshOfLoops()
 	UniformScaleBounds[IsoV].Set(UniformCuttingCoordinates[IsoV][0], UniformCuttingCoordinates[IsoV].Last());
 
 	// Fit loops in UniformScaled to UniformScale bounds.
-	for (TArray<FPoint2D>& Loop : FaceLoops2D[EGridSpace::UniformScaled])
+	for (TArray<FVector2d>& Loop : FaceLoops2D[EGridSpace::UniformScaled])
 	{
-		for (FPoint2D& Point : Loop)
+		for (FVector2d& Point : Loop)
 		{
 			UniformScaleBounds.MoveInsideIfNot(Point);
 		}
@@ -1162,9 +1162,9 @@ void FGrid::ScaleLoops()
 
 	for (int32 IndexBoudnary = 0; IndexBoudnary < FaceLoops2D[EGridSpace::Default2D].Num(); ++IndexBoudnary)
 	{
-		const TArray<FPoint2D>& Loop = FaceLoops2D[EGridSpace::Default2D][IndexBoudnary];
-		TArray<FPoint2D>& ScaledLoop = FaceLoops2D[EGridSpace::Scaled][IndexBoudnary];
-		TArray<FPoint2D>& UniformScaledLoop = FaceLoops2D[EGridSpace::UniformScaled][IndexBoudnary];
+		const TArray<FVector2d>& Loop = FaceLoops2D[EGridSpace::Default2D][IndexBoudnary];
+		TArray<FVector2d>& ScaledLoop = FaceLoops2D[EGridSpace::Scaled][IndexBoudnary];
+		TArray<FVector2d>& UniformScaledLoop = FaceLoops2D[EGridSpace::UniformScaled][IndexBoudnary];
 
 		ScaledLoop.SetNum(Loop.Num());
 		UniformScaledLoop.SetNum(Loop.Num());
@@ -1173,7 +1173,7 @@ void FGrid::ScaleLoops()
 		int32 IndexV = 0;
 		for (int32 Index = 0; Index < Loop.Num(); ++Index)
 		{
-			const FPoint2D& Point = Loop[Index];
+			const FVector2d& Point = Loop[Index];
 
 			ArrayUtils::FindCoordinateIndex(CoordinateGrid[EIso::IsoU], Point.X, IndexU);
 			ArrayUtils::FindCoordinateIndex(CoordinateGrid[EIso::IsoV], Point.Y, IndexV);
@@ -1247,9 +1247,9 @@ void FGrid::FindInnerFacePoints()
 	{
 		int32 IndexV = 0;
 		int32 IndexU = 0;
-		for (TArray<FPoint2D>& Loop : FaceLoops2D[EGridSpace::UniformScaled])
+		for (TArray<FVector2d>& Loop : FaceLoops2D[EGridSpace::UniformScaled])
 		{
-			for (FPoint2D& Point : Loop)
+			for (FVector2d& Point : Loop)
 			{
 				while (IndexV != 0 && (Point.Y < UniformCuttingCoordinates[EIso::IsoV][IndexV]))
 				{
@@ -1316,12 +1316,12 @@ void FGrid::FindInnerFacePoints()
 	}
 
 	// Intersection along U axis
-	for (const TArray<FPoint2D>& Loop : FaceLoops2D[EGridSpace::UniformScaled])
+	for (const TArray<FVector2d>& Loop : FaceLoops2D[EGridSpace::UniformScaled])
 	{
-		const FPoint2D* FirstSegmentPoint = &Loop.Last();
-		for (const FPoint2D& LoopPoint : Loop)
+		const FVector2d* FirstSegmentPoint = &Loop.Last();
+		for (const FVector2d& LoopPoint : Loop)
 		{
-			const FPoint2D* SecondSegmentPoint = &LoopPoint;
+			const FVector2d* SecondSegmentPoint = &LoopPoint;
 			double UMin = FirstSegmentPoint->X;
 			double VMin = FirstSegmentPoint->Y;
 			double Umax = SecondSegmentPoint->X;
@@ -1333,8 +1333,8 @@ void FGrid::FindInnerFacePoints()
 			// AB^AP = ABu*(Pv-Av) - ABv*(Pu-Au)
 			// AB^AP = Pv*ABu - Pu*ABv + Au*ABv - Av*ABu
 			// AB^AP = Pv*ABu - Pu*ABv + AuABVMinusAvABu
-			FPoint2D PointA;
-			FPoint2D PointB;
+			FVector2d PointA;
+			FVector2d PointB;
 			if (FirstSegmentPoint->Y < SecondSegmentPoint->Y)
 			{
 				PointA = *FirstSegmentPoint;
@@ -1406,12 +1406,12 @@ void FGrid::FindInnerFacePoints()
 	}
 
 	// Intersection along V axis
-	for (const TArray<FPoint2D>& Loop : FaceLoops2D[EGridSpace::UniformScaled])
+	for (const TArray<FVector2d>& Loop : FaceLoops2D[EGridSpace::UniformScaled])
 	{
-		const FPoint2D* FirstSegmentPoint = &Loop.Last();
-		for (const FPoint2D& LoopPoint : Loop)
+		const FVector2d* FirstSegmentPoint = &Loop.Last();
+		for (const FVector2d& LoopPoint : Loop)
 		{
-			const FPoint2D* SecondSegmentPoint = &LoopPoint;
+			const FVector2d* SecondSegmentPoint = &LoopPoint;
 			double UMin = FirstSegmentPoint->X;
 			double VMin = FirstSegmentPoint->Y;
 			double Umax = SecondSegmentPoint->X;
@@ -1423,8 +1423,8 @@ void FGrid::FindInnerFacePoints()
 			// AB^AP = ABu*(Pv-Av) - ABv*(Pu-Au)
 			// AB^AP = Pv*ABu - Pu*ABv + Au*ABv - Av*ABu
 			// AB^AP = Pv*ABu - Pu*ABv + AuABVMinusAvABu
-			FPoint2D PointA;
-			FPoint2D PointB;
+			FVector2d PointA;
+			FVector2d PointB;
 			if (FirstSegmentPoint->X < SecondSegmentPoint->X)
 			{
 				PointA = *FirstSegmentPoint;
diff --git a/Engine/Source/Runtime/Datasmith/CADKernel/Base/Private/Mesh/Structure/Grid.h b/Engine/Source/Runtime/Datasmith/CADKernel/Base/Private/Mesh/Structure/Grid.h
--- a/Engine/Source/Runtime/Datasmith/CADKernel/Base/Private/Mesh/Structure/Grid.h
+++ b/Engine/Source/Runtime/Datasmith/CADKernel/Base/Private/Mesh/Structure/Grid.h
@@ -65,12 +65,12 @@ protected:
 	/**
 	 * 2D Coordinate of Loop's nodes in each space
 	 */
-	TArray<TArray<FPoint2D>> FaceLoops2D[EGridSpace::EndGridSpace];
+	TArray<TArray<FVector2d>> FaceLoops2D[EGridSpace::EndGridSpace];
 
 	/**
 	 * 3D Coordinate of Loop nodes in each space
 	 */
-	TArray<TArray<FPoint>> FaceLoops3D;
+	TArray<TArray<FVector>> FaceLoops3D;
 
 	/**
 	 * Surface Normal at each boundary nodes
@@ -308,7 +308,7 @@ public:
 		NodeMarkers[Index] &= ~ENodeMarker::IsInside;
 	}
 
-	void SetInner2DPoint(EGridSpace Space, int32 Index, const FPoint2D& NewCoordinate)
+	void SetInner2DPoint(EGridSpace Space, int32 Index, const FVector2d& NewCoordinate)
 	{
 		Points2D[(int32)Space][Index] = NewCoordinate;
 	}
@@ -352,17 +352,17 @@ public:
 		return NodeIdsOfFaceLoops;
 	}
 
-	const FPoint2D& GetLoop2DPoint(EGridSpace Space, int32 LoopIndex, int32 Index) const
+	const FVector2d& GetLoop2DPoint(EGridSpace Space, int32 LoopIndex, int32 Index) const
 	{
 		return FaceLoops2D[(int32)Space][LoopIndex][Index];
 	}
 
-	void SetLoop2DPoint(EGridSpace Space, int32 LoopIndex, int32 Index, const FPoint2D& NewCoordinate)
+	void SetLoop2DPoint(EGridSpace Space, int32 LoopIndex, int32 Index, const FVector2d& NewCoordinate)
 	{
 		FaceLoops2D[(int32)Space][LoopIndex][Index] = NewCoordinate;
 	}
 
-	const FPoint& GetLoop3DPoint(int32 LoopIndex, int32 Index) const
+	const FVector& GetLoop3DPoint(int32 LoopIndex, int32 Index) const
 	{
 		return FaceLoops3D[LoopIndex][Index];
 	}
@@ -375,7 +375,7 @@ public:
 	/**
 	 * @return the array of array of 2d points of the loops according to the defined grid space
 	 */
-	const TArray<TArray<FPoint2D>>& GetLoops2D(EGridSpace Space) const
+	const TArray<TArray<FVector2d>>& GetLoops2D(EGridSpace Space) const
 	{
 		return FaceLoops2D[(int32)Space];
 	}
@@ -383,7 +383,7 @@ public:
 	/**
 	 * @return the array of array of 3d points of the loops
 	 */
-	TArray<TArray<FPoint>>& GetLoops3D()
+	TArray<TArray<FVector>>& GetLoops3D()
 	{
 		return FaceLoops3D;
 	}
@@ -391,7 +391,7 @@ public:
 	/**
 	 * @return the array of array of 3d points of the loops
 	 */
-	const TArray<TArray<FPoint>>& GetLoops3D() const
+	const TArray<TArray<FVector>>& GetLoops3D() const
 	{
 		return FaceLoops3D;
 	}
diff --git a/Engine/Source/Runtime/Datasmith/CADKernel/Base/Private/Mesh/Structure/GridBase.cpp b/Engine/Source/Runtime/Datasmith/CADKernel/Base/Private/Mesh/Structure/GridBase.cpp
--- a/Engine/Source/Runtime/Datasmith/CADKernel/Base/Private/Mesh/Structure/GridBase.cpp
+++ b/Engine/Source/Runtime/Datasmith/CADKernel/Base/Private/Mesh/Structure/GridBase.cpp
@@ -153,7 +153,7 @@ bool FGridBase::ScaleGrid()
 		for (int32 IPointU = 0; IPointU < CuttingCount[EIso::IsoU]; ++IPointU)
 		{
 			double Length = 0;
-			FPoint LastPoint = MiddlePoints.Points3D[IPointU];
+			FVector LastPoint = MiddlePoints.Points3D[IPointU];
 			for (int32 IPointV = StartIndexUp; IPointV < CuttingCount[EIso::IsoV]; ++IPointV)
 			{
 				NumUV = IPointV * CuttingCount[EIso::IsoU] + IPointU;
@@ -191,7 +191,7 @@ bool FGridBase::ScaleGrid()
 		for (int32 IPointV = 0; IPointV < CuttingCount[EIso::IsoV]; ++IPointV)
 		{
 			double Length = 0;
-			FPoint LastPoint = MiddlePoints.Points3D[IPointV];
+			FVector LastPoint = MiddlePoints.Points3D[IPointV];
 			for (int32 IPointU = StartIndexUp; IPointU < CuttingCount[EIso::IsoU]; ++IPointU)
 			{
 				NumUV = IPointV * CuttingCount[EIso::IsoU] + IPointU;
@@ -217,7 +217,7 @@ bool FGridBase::ScaleGrid()
 	return true;
 }
 
-void FGridBase::TransformPoints(EGridSpace DestinationSpace, const TArray<FPoint2D>& InPointsToScale, TArray<FPoint2D>& OutTransformedPoints) const
+void FGridBase::TransformPoints(EGridSpace DestinationSpace, const TArray<FVector2d>& InPointsToScale, TArray<FVector2d>& OutTransformedPoints) const
 {
 	const FCoordinateGrid& CoordinateGrid = GetCoordinateGrid();
 
@@ -227,7 +227,7 @@ void FGridBase::TransformPoints(EGridSpace DestinationSpace, const TArray<FPoint
 	int32 IndexV = 0;
 	for (int32 Index = 0; Index < InPointsToScale.Num(); ++Index)
 	{
-		const FPoint2D& Point = InPointsToScale[Index];
+		const FVector2d& Point = InPointsToScale[Index];
 
 		ArrayUtils::FindCoordinateIndex(CoordinateGrid[EIso::IsoU], Point.X, IndexU);
 		ArrayUtils::FindCoordinateIndex(CoordinateGrid[EIso::IsoV], Point.Y, IndexV);
diff --git a/Engine/Source/Runtime/Datasmith/CADKernel/Base/Private/Mesh/Structure/GridBase.h b/Engine/Source/Runtime/Datasmith/CADKernel/Base/Private/Mesh/Structure/GridBase.h
--- a/Engine/Source/Runtime/Datasmith/CADKernel/Base/Private/Mesh/Structure/GridBase.h
+++ b/Engine/Source/Runtime/Datasmith/CADKernel/Base/Private/Mesh/Structure/GridBase.h
@@ -88,12 +88,12 @@ protected:
 	/**
 	 * 2D Coordinate of grid nodes in each space
 	 */
-	TArray<FPoint2D> Points2D[EGridSpace::EndGridSpace];
+	TArray<FVector2d> Points2D[EGridSpace::EndGridSpace];
 
 	/**
 	 * 3D Coordinate of inner nodes
 	 */
-	TArray<FPoint> Points3D;
+	TArray<FVector> Points3D;
 
 	/**
 	 * Surface Normal at each inner nodes
@@ -117,7 +117,7 @@ public:
 	 * Convert an array of points of "DefaultParametric" space into a scaled parametric space
 	 * @see FThinZone2DFinder::BuildBoundarySegments()
 	 */
-	void TransformPoints(EGridSpace DestinationSpace, const TArray<FPoint2D>& InPointsToScale, TArray<FPoint2D>& OutScaledPoints) const;
+	void TransformPoints(EGridSpace DestinationSpace, const TArray<FVector2d>& InPointsToScale, TArray<FVector2d>& OutScaledPoints) const;
 
 	/**
 	 * Builds the scaled parametric spaces
@@ -142,14 +142,14 @@ protected:
 	 * Convert Coordinate of "DefaultParametric" space into a scaled parametric space
 	 * @see ScaleLoops
 	 */
-	void ComputeNewCoordinate(const TArray<FPoint2D>& NewGrid, int32 IndexU, int32 IndexV, const FPoint2D& InPoint, FPoint2D& OutNewScaledPoint) const
+	void ComputeNewCoordinate(const TArray<FVector2d>& NewGrid, int32 IndexU, int32 IndexV, const FVector2d& InPoint, FVector2d& OutNewScaledPoint) const
 	{
 		const FCoordinateGrid& CoordinateGrid = GetCoordinateGrid();
 
-		const FPoint2D& PointU0V0 = NewGrid[(IndexV + 0) * CuttingCount[EIso::IsoU] + (IndexU + 0)];
-		const FPoint2D& PointU1V0 = NewGrid[(IndexV + 0) * CuttingCount[EIso::IsoU] + (IndexU + 1)];
-		const FPoint2D& PointU0V1 = NewGrid[(IndexV + 1) * CuttingCount[EIso::IsoU] + (IndexU + 0)];
-		const FPoint2D& PointU1V1 = NewGrid[(IndexV + 1) * CuttingCount[EIso::IsoU] + (IndexU + 1)];
+		const FVector2d& PointU0V0 = NewGrid[(IndexV + 0) * CuttingCount[EIso::IsoU] + (IndexU + 0)];
+		const FVector2d& PointU1V0 = NewGrid[(IndexV + 0) * CuttingCount[EIso::IsoU] + (IndexU + 1)];
+		const FVector2d& PointU0V1 = NewGrid[(IndexV + 1) * CuttingCount[EIso::IsoU] + (IndexU + 0)];
+		const FVector2d& PointU1V1 = NewGrid[(IndexV + 1) * CuttingCount[EIso::IsoU] + (IndexU + 1)];
 
 		const double U1MinusU0 = CoordinateGrid[EIso::IsoU][IndexU + 1] - CoordinateGrid[EIso::IsoU][IndexU];
 		const double U0MinusU  = CoordinateGrid[EIso::IsoU][IndexU] - InPoint.X;
@@ -166,27 +166,27 @@ protected:
 public:
 
 	/**
-	 * @return the FPoint2D (parametric coordinates) of the point at the Index of the grid in the defined grid space
+	 * @return the FVector2d (parametric coordinates) of the point at the Index of the grid in the defined grid space
 	 * @see EGridSpace
 	 * @see Points2D
 	 */
-	const FPoint2D& GetInner2DPoint(EGridSpace Space, int32 Index) const
+	const FVector2d& GetInner2DPoint(EGridSpace Space, int32 Index) const
 	{
 		return Points2D[(int32)Space][Index];
 	}
 
 	/**
-	 * @return the FPoint2D (parametric coordinates) of the point at the Index of the grid in the defined grid space @see EGridSpace
+	 * @return the FVector2d (parametric coordinates) of the point at the Index of the grid in the defined grid space @see EGridSpace
 	 */
-	const FPoint2D& GetInner2DPoint(EGridSpace Space, int32 IndexU, int32 IndexV) const
+	const FVector2d& GetInner2DPoint(EGridSpace Space, int32 IndexU, int32 IndexV) const
 	{
 		return Points2D[(int32)Space][GobalIndex(IndexU, IndexV)];
 	}
 
 	/**
-	 * @return the FPoint (3D coordinates) of the point at the Index of the grid
+	 * @return the FVector (3D coordinates) of the point at the Index of the grid
 	 */
-	const FPoint& GetInner3DPoint(int32 Index) const
+	const FVector& GetInner3DPoint(int32 Index) const
 	{
 		return Points3D[Index];
 	}
@@ -204,7 +204,7 @@ public:
 	/**
 	 * @return the array of 3d points of the grid
 	 */
-	TArray<FPoint>& GetInner3DPoints()
+	TArray<FVector>& GetInner3DPoints()
 	{
 		return Points3D;
 	}
@@ -212,7 +212,7 @@ public:
 	/**
 	 * @return the array of 3d points of the grid
 	 */
-	const TArray<FPoint>& GetInner3DPoints() const
+	const TArray<FVector>& GetInner3DPoints() const
 	{
 		return Points3D;
 	}
@@ -220,7 +220,7 @@ public:
 	/**
 	 * @return the array of 2d points of the grid in the defined space
 	 */
-	const TArray<FPoint2D>& GetInner2DPoints(EGridSpace Space) const
+	const TArray<FVector2d>& GetInner2DPoints(EGridSpace Space) const
 	{
 		return Points2D[(int32)Space];
 	}
