Replaced FPoint and FPoint2D by FVector and  FVector2d and remove unused FFPoint and FPointH classes - Final

diff --git a/Engine/Source/Runtime/Datasmith/CADKernel/Base/Private/Mesh/Structure/LoopCleaner.cpp b/Engine/Source/Runtime/Datasmith/CADKernel/Base/Private/Mesh/Structure/LoopCleaner.cpp
--- a/Engine/Source/Runtime/Datasmith/CADKernel/Base/Private/Mesh/Structure/LoopCleaner.cpp
+++ b/Engine/Source/Runtime/Datasmith/CADKernel/Base/Private/Mesh/Structure/LoopCleaner.cpp
@@ -401,8 +401,8 @@ bool FLoopCleaner::TryToRemoveIntersectionOfTwoConsecutiveIntersectingSegments(c
 	// if the segment and IntersectingSegment are parallel, segment are moved inside
 	if (SegmentSlope < 0.01)
 	{
-		double StartPointSquareDistance = SquareDistanceOfPointToSegment(Segment2D.Point0, IntersectingSegment2D.Point0, IntersectingSegment2D.Point1);
-		double EndPointSquareDistance = SquareDistanceOfPointToSegment(Segment2D.Point1, IntersectingSegment2D.Point0, IntersectingSegment2D.Point1);
+		double StartPointSquareDistance = SquareDistanceOfPointToSegment<FVector2d>(Segment2D.Point0, IntersectingSegment2D.Point0, IntersectingSegment2D.Point1);
+		double EndPointSquareDistance = SquareDistanceOfPointToSegment<FVector2d>(Segment2D.Point1, IntersectingSegment2D.Point0, IntersectingSegment2D.Point1);
 		if (StartPointSquareDistance < Tolerances.SquareGeometricTolerance2 && EndPointSquareDistance < Tolerances.SquareGeometricTolerance2)
 		{
 			OffsetSegment(Segment, Segment2D, IntersectingSegment2D);
@@ -461,11 +461,11 @@ bool FLoopCleaner::TryToRemoveIntersectionOfTwoConsecutiveIntersectingSegments(c
 
 	FSegment2D  PreviousSegment2D(Node->Get2DPoint(EGridSpace::UniformScaled, Grid), PreviousNode->Get2DPoint(EGridSpace::UniformScaled, Grid));
 
-	FPoint2D Intersection1 = FindIntersectionOfSegments2D(PreviousSegment2D, IntersectingSegment2D);
-	FPoint2D Intersection2 = FindIntersectionOfSegments2D(NextSegment2D, IntersectingSegment2D);
+	FVector2d Intersection1 = FindIntersectionOfSegments2D(PreviousSegment2D, IntersectingSegment2D);
+	FVector2d Intersection2 = FindIntersectionOfSegments2D(NextSegment2D, IntersectingSegment2D);
 
 	double Coordinate;
-	FPoint2D ProjectedPoint = ProjectPointOnSegment(Node->Get2DPoint(EGridSpace::UniformScaled, Grid), Intersection1, Intersection2, Coordinate);
+	FVector2d ProjectedPoint = ProjectPointOnSegment<FVector2d>(Node->Get2DPoint(EGridSpace::UniformScaled, Grid), Intersection1, Intersection2, Coordinate);
 
 #ifdef DEBUG_TWO_CONSECUTIVE_INTERSECTING
 	if (bDisplay)
@@ -475,13 +475,13 @@ bool FLoopCleaner::TryToRemoveIntersectionOfTwoConsecutiveIntersectingSegments(c
 	}
 #endif
 
-	FPoint2D SegmentTangent = IntersectingSegment2D.Point1 - IntersectingSegment2D.Point0;
+	FVector2d SegmentTangent = IntersectingSegment2D.Point1 - IntersectingSegment2D.Point0;
 	SegmentTangent.Normalize();
-	FPoint2D MoveDirection = FPoint2D(-SegmentTangent.Y, SegmentTangent.X);
+	FVector2d MoveDirection = FVector2d(-SegmentTangent.Y, SegmentTangent.X);
 
 	MoveDirection *= Tolerances.GeometricTolerance;
 
-	FPoint2D NewCoordinate = ProjectedPoint + MoveDirection;
+	FVector2d NewCoordinate = ProjectedPoint + MoveDirection;
 
 #ifdef DEBUG_TWO_CONSECUTIVE_INTERSECTING
 	if (bDisplay)
@@ -491,8 +491,8 @@ bool FLoopCleaner::TryToRemoveIntersectionOfTwoConsecutiveIntersectingSegments(c
 	}
 #endif
 
-	if ((FPoint2D::DistSquared(PreviousNode->Get2DPoint(EGridSpace::UniformScaled, Grid), NewCoordinate) < Tolerances.SquareGeometricTolerance2) ||
-		(FPoint2D::DistSquared(NextNode->Get2DPoint(EGridSpace::UniformScaled, Grid), NewCoordinate) < Tolerances.SquareGeometricTolerance2))
+	if ((FVector2d::DistSquared(PreviousNode->Get2DPoint(EGridSpace::UniformScaled, Grid), NewCoordinate) < Tolerances.SquareGeometricTolerance2) ||
+		(FVector2d::DistSquared(NextNode->Get2DPoint(EGridSpace::UniformScaled, Grid), NewCoordinate) < Tolerances.SquareGeometricTolerance2))
 	{
 		if(RemoveNodeOfLoop(*Node))
 		{
@@ -525,13 +525,13 @@ bool FLoopCleaner::TryToRemoveIntersectionOfTwoConsecutiveIntersectingSegments(c
 
 void FLoopCleaner::OffsetSegment(FIsoSegment& Segment, FSegment2D& Segment2D, FSegment2D& IntersectingSegment2D)
 {
-	FPoint2D SegmentTangent = IntersectingSegment2D.Point1 - IntersectingSegment2D.Point0;
+	FVector2d SegmentTangent = IntersectingSegment2D.Point1 - IntersectingSegment2D.Point0;
 	SegmentTangent.Normalize();
-	FPoint2D MoveDirection = FPoint2D(-SegmentTangent.Y, SegmentTangent.X);
+	FVector2d MoveDirection = FVector2d(-SegmentTangent.Y, SegmentTangent.X);
 	MoveDirection *= Tolerances.GeometricTolerance;
 
-	FPoint2D NewPoint0 = Segment2D.Point0 + MoveDirection;
-	FPoint2D NewPoint1 = Segment2D.Point1 + MoveDirection;
+	FVector2d NewPoint0 = Segment2D.Point0 + MoveDirection;
+	FVector2d NewPoint1 = Segment2D.Point1 + MoveDirection;
 
 	Segment.GetFirstNode().Set2DPoint(EGridSpace::UniformScaled, Grid, NewPoint0);
 	Segment.GetSecondNode().Set2DPoint(EGridSpace::UniformScaled, Grid, NewPoint1);
@@ -539,12 +539,12 @@ void FLoopCleaner::OffsetSegment(FIsoSegment& Segment, FSegment2D& Segment2D, FS
 
 void FLoopCleaner::OffsetNode(FLoopNode& Node, FSegment2D& IntersectingSegment2D)
 {
-	FPoint2D SegmentTangent = IntersectingSegment2D.Point1 - IntersectingSegment2D.Point0;
+	FVector2d SegmentTangent = IntersectingSegment2D.Point1 - IntersectingSegment2D.Point0;
 	SegmentTangent.Normalize();
-	FPoint2D MoveDirection = FPoint2D(-SegmentTangent.Y, SegmentTangent.X);
+	FVector2d MoveDirection = FVector2d(-SegmentTangent.Y, SegmentTangent.X);
 	MoveDirection *= Tolerances.GeometricTolerance;
 
-	FPoint2D NewPoint = Node.Get2DPoint(EGridSpace::UniformScaled, Grid) + MoveDirection;
+	FVector2d NewPoint = Node.Get2DPoint(EGridSpace::UniformScaled, Grid) + MoveDirection;
 	Node.Set2DPoint(EGridSpace::UniformScaled, Grid, NewPoint);
 }
 
@@ -556,16 +556,16 @@ void FLoopCleaner::RemoveIntersectionByMovingOutsideSegmentNodeInside(const FIso
 	Nodes[1][0] = &Segment.GetFirstNode();
 	Nodes[1][1] = &Segment.GetSecondNode();
 
-	FPoint2D IntersectingPoints[2];
+	FVector2d IntersectingPoints[2];
 	IntersectingPoints[0] = IntersectingSegment.GetFirstNode().Get2DPoint(EGridSpace::UniformScaled, Grid);
 	IntersectingPoints[1] = IntersectingSegment.GetSecondNode().Get2DPoint(EGridSpace::UniformScaled, Grid);
 
-	FPoint2D SegmentPoints[2];
+	FVector2d SegmentPoints[2];
 	SegmentPoints[0] = Segment.GetFirstNode().Get2DPoint(EGridSpace::UniformScaled, Grid);
 	SegmentPoints[1] = Segment.GetSecondNode().Get2DPoint(EGridSpace::UniformScaled, Grid);
 
-	double SquareLenghtIntersectingSegment = FPoint2D::DistSquared(IntersectingPoints[0], IntersectingPoints[1]);
-	double SquareLenghtSegment = FPoint2D::DistSquared(SegmentPoints[0], SegmentPoints[1]);
+	double SquareLenghtIntersectingSegment = FVector2d::DistSquared(IntersectingPoints[0], IntersectingPoints[1]);
+	double SquareLenghtSegment = FVector2d::DistSquared(SegmentPoints[0], SegmentPoints[1]);
 
 	if (SquareLenghtSegment > 10 * SquareLenghtIntersectingSegment)
 	{
@@ -573,7 +573,7 @@ void FLoopCleaner::RemoveIntersectionByMovingOutsideSegmentNodeInside(const FIso
 	}
 
 	bool bFirstNodeIsOutside = false;
-	FPoint2D PointToMove = Segment.GetSecondNode().Get2DPoint(EGridSpace::UniformScaled, Grid);
+	FVector2d PointToMove = Segment.GetSecondNode().Get2DPoint(EGridSpace::UniformScaled, Grid);
 
 	// Is Second node, the outside node ?
 	{
@@ -599,21 +599,21 @@ void FLoopCleaner::RemoveIntersectionByMovingOutsideSegmentNodeInside(const FIso
 #endif
 
 	double Coordinate;
-	FPoint2D ProjectedPoint = ProjectPointOnSegment(PointToMove, IntersectingPoints[0], IntersectingPoints[1], Coordinate);
+	FVector2d ProjectedPoint = ProjectPointOnSegment<FVector2d>(PointToMove, IntersectingPoints[0], IntersectingPoints[1], Coordinate);
 
-	FPoint2D MoveDirection = IntersectingPoints[1] - IntersectingPoints[0];
+	FVector2d MoveDirection = IntersectingPoints[1] - IntersectingPoints[0];
 	MoveDirection.Normalize();
-	MoveDirection = FPoint2D(-MoveDirection.Y, MoveDirection.X);
+	MoveDirection = FVector2d(-MoveDirection.Y, MoveDirection.X);
 	MoveDirection *= bIsSameInnerLoop ? -Tolerances.GeometricTolerance : Tolerances.GeometricTolerance;
 
-	FPoint2D NewCoordinate = ProjectedPoint + MoveDirection;
+	FVector2d NewCoordinate = ProjectedPoint + MoveDirection;
 
 	FLoopNode& Node = (FLoopNode&) const_cast<FIsoNode&>(!bFirstNodeIsOutside ? Segment.GetSecondNode() : Segment.GetFirstNode());
 	FLoopNode& PreviousNode = Node.GetPreviousNode();
 	FLoopNode& NextNode = Node.GetNextNode();
 
-	if ((FPoint2D::DistSquared(PreviousNode.Get2DPoint(EGridSpace::UniformScaled, Grid), NewCoordinate) < Tolerances.SquareGeometricTolerance2) ||
-		(FPoint2D::DistSquared(NextNode.Get2DPoint(EGridSpace::UniformScaled, Grid), NewCoordinate) < Tolerances.SquareGeometricTolerance2))
+	if ((FVector2d::DistSquared(PreviousNode.Get2DPoint(EGridSpace::UniformScaled, Grid), NewCoordinate) < Tolerances.SquareGeometricTolerance2) ||
+		(FVector2d::DistSquared(NextNode.Get2DPoint(EGridSpace::UniformScaled, Grid), NewCoordinate) < Tolerances.SquareGeometricTolerance2))
 	{
 		RemoveNodeOfLoop(Node);
 		return;
@@ -793,14 +793,14 @@ bool FLoopCleaner::RemoveOutgoingLoop(const TPair<double, double>& Intersection,
 	//	return false;
 	//}
 
-	TFunction<FPoint2D(double, FLoopNode*, FLoopNode*)> IntersectingPoint = [&](double Coordinate, FLoopNode* Start, FLoopNode* End) -> FPoint2D
+	TFunction<FVector2d(double, FLoopNode*, FLoopNode*)> IntersectingPoint = [&](double Coordinate, FLoopNode* Start, FLoopNode* End) -> FVector2d
 	{
-		const FPoint2D& StartPoint = Start->Get2DPoint(EGridSpace::UniformScaled, Grid);
-		const FPoint2D& EndPoint = End->Get2DPoint(EGridSpace::UniformScaled, Grid);
+		const FVector2d& StartPoint = Start->Get2DPoint(EGridSpace::UniformScaled, Grid);
+		const FVector2d& EndPoint = End->Get2DPoint(EGridSpace::UniformScaled, Grid);
 
 		int32 StartIndex = (int32)Coordinate;
 		Coordinate -= StartIndex;
-		FPoint2D Intersection = PointOnSegment(StartPoint, EndPoint, Coordinate);
+		FVector2d Intersection = PointOnSegment<FVector2d>(StartPoint, EndPoint, Coordinate);
 
 		//{
 		//	F3DDebugSession _(TEXT("Intersected Point"));
@@ -820,12 +820,12 @@ bool FLoopCleaner::RemoveOutgoingLoop(const TPair<double, double>& Intersection,
 	}
 
 	FLoopNode* StartNode = GetPrevious(TmpNode);
-	FPoint2D FirstIntersection = IntersectingPoint(Intersection.Value, StartNode, TmpNode);
+	FVector2d FirstIntersection = IntersectingPoint(Intersection.Value, StartNode, TmpNode);
 
 	TmpNode = GetNext(EndNode);
-	FPoint2D SecondIntersection = IntersectingPoint(NextIntersection.Value, EndNode, TmpNode);
+	FVector2d SecondIntersection = IntersectingPoint(NextIntersection.Value, EndNode, TmpNode);
 
-	FPoint2D MiddlePoint = (FirstIntersection + SecondIntersection) * 0.5;
+	FVector2d MiddlePoint = (FirstIntersection + SecondIntersection) * 0.5;
 
 	TmpNode = GetNext(StartNode);
 	while (TmpNode && (TmpNode != EndNode) && !TmpNode->IsDelete())
@@ -837,7 +837,7 @@ bool FLoopCleaner::RemoveOutgoingLoop(const TPair<double, double>& Intersection,
 		TmpNode = GetNext(StartNode);
 	}
 
-	FPoint2D MoveDirection = SecondIntersection - FirstIntersection;
+	FVector2d MoveDirection = SecondIntersection - FirstIntersection;
 	double Length = MoveDirection.Length();
 	if (FMath::IsNearlyZero(Length))
 	{
@@ -851,7 +851,7 @@ bool FLoopCleaner::RemoveOutgoingLoop(const TPair<double, double>& Intersection,
 	}
 
 	MoveDirection /= Length;
-	MoveDirection = FPoint2D(-MoveDirection.Y, MoveDirection.X);
+	MoveDirection = FVector2d(-MoveDirection.Y, MoveDirection.X);
 	MoveDirection *= Tolerances.GeometricTolerance;
 	MiddlePoint += MoveDirection;
 
@@ -867,22 +867,22 @@ bool FLoopCleaner::RemoveIntersectionsOfSubLoop(int32 IntersectionIndex, int32 I
 {
 	using namespace LoopCleanerImpl;
 
-	TFunction<void(FLoopNode*, const FPoint2D&, FPoint2D&)> MoveNodeToProjection = [&](FLoopNode* NodeToProject, const FPoint2D& PointToProject, FPoint2D& ProjectedPoint)
+	TFunction<void(FLoopNode*, const FVector2d&, FVector2d&)> MoveNodeToProjection = [&](FLoopNode* NodeToProject, const FVector2d& PointToProject, FVector2d& ProjectedPoint)
 	{
-		FPoint2D MoveDirection = ProjectedPoint - PointToProject;
+		FVector2d MoveDirection = ProjectedPoint - PointToProject;
 		MoveDirection.Normalize();
 		MoveDirection *= Tolerances.GeometricTolerance;
 		ProjectedPoint += MoveDirection;
 		NodeToProject->Set2DPoint(EGridSpace::UniformScaled, Grid, ProjectedPoint);
 	};
 
-	TFunction<void(FLoopNode*, const FPoint2D&, const FPoint2D&)> ProjectNodeOnSegment = [&](FLoopNode* NodeToProject, const FPoint2D& Point0, const FPoint2D& Point1)
+	TFunction<void(FLoopNode*, const FVector2d&, const FVector2d&)> ProjectNodeOnSegment = [&](FLoopNode* NodeToProject, const FVector2d& Point0, const FVector2d& Point1)
 	{
-		const FPoint2D& PointToProject = NodeToProject->Get2DPoint(EGridSpace::UniformScaled, Grid);
+		const FVector2d& PointToProject = NodeToProject->Get2DPoint(EGridSpace::UniformScaled, Grid);
 
-		FPoint2D ProjectedPoint;
+		FVector2d ProjectedPoint;
 		double Coordinate;
-		ProjectedPoint = ProjectPointOnSegment(PointToProject, Point0, Point1, Coordinate);
+		ProjectedPoint = ProjectPointOnSegment<FVector2d>(PointToProject, Point0, Point1, Coordinate);
 
 		MoveNodeToProjection(NodeToProject, PointToProject, ProjectedPoint);
 	};
@@ -900,8 +900,8 @@ bool FLoopCleaner::RemoveIntersectionsOfSubLoop(int32 IntersectionIndex, int32 I
 		StopNode = GetNext(StopNode);
 		FLoopNode* StartSegment = GetPrevious(EndSegment);
 
-		const FPoint2D& EndPoint = EndSegment->Get2DPoint(EGridSpace::UniformScaled, Grid);
-		const FPoint2D& StartPoint = StartSegment->Get2DPoint(EGridSpace::UniformScaled, Grid);
+		const FVector2d& EndPoint = EndSegment->Get2DPoint(EGridSpace::UniformScaled, Grid);
+		const FVector2d& StartPoint = StartSegment->Get2DPoint(EGridSpace::UniformScaled, Grid);
 
 		for (; Node != StopNode; Node = GetNext(Node))
 		{
@@ -929,17 +929,17 @@ bool FLoopCleaner::RemoveIntersectionsOfSubLoop(int32 IntersectionIndex, int32 I
 	TFunction<void(FLoopNode*, FLoopNode*)> MoveNodeBehindOther = [&](FLoopNode* NodeToMove, FLoopNode* Node1Side1)
 	{
 		FLoopNode* Node0Side1 = GetPrevious(Node1Side1);
-		const FPoint2D& Point0 = Node0Side1->Get2DPoint(EGridSpace::UniformScaled, Grid);
-		const FPoint2D& Point1 = Node1Side1->Get2DPoint(EGridSpace::UniformScaled, Grid);
+		const FVector2d& Point0 = Node0Side1->Get2DPoint(EGridSpace::UniformScaled, Grid);
+		const FVector2d& Point1 = Node1Side1->Get2DPoint(EGridSpace::UniformScaled, Grid);
 
-		const FPoint2D& PointToMove = NodeToMove->Get2DPoint(EGridSpace::UniformScaled, Grid);
+		const FVector2d& PointToMove = NodeToMove->Get2DPoint(EGridSpace::UniformScaled, Grid);
 
-		FPoint2D MoveDirection = Point1 - Point0;
+		FVector2d MoveDirection = Point1 - Point0;
 		MoveDirection.Normalize();
-		MoveDirection = FPoint2D(-MoveDirection.Y, MoveDirection.X);
+		MoveDirection = FVector2d(-MoveDirection.Y, MoveDirection.X);
 		MoveDirection *= Tolerances.GeometricTolerance;
 
-		FPoint2D NewCoordinate = Point1 + MoveDirection;
+		FVector2d NewCoordinate = Point1 + MoveDirection;
 		NodeToMove->Set2DPoint(EGridSpace::UniformScaled, Grid, NewCoordinate);
 	};
 
@@ -1142,8 +1142,8 @@ bool FLoopCleaner::RemovePickOrCoincidenceBetween(FLoopNode* StartNode, FLoopNod
 			return false;
 		}
 
-		const FPoint2D& PreviousPoint = NodeToBeProcessed->GetPreviousNode().Get2DPoint(EGridSpace::UniformScaled, Grid);
-		const FPoint2D& PointToBeProcessed = NodeToBeProcessed->Get2DPoint(EGridSpace::UniformScaled, Grid);
+		const FVector2d& PreviousPoint = NodeToBeProcessed->GetPreviousNode().Get2DPoint(EGridSpace::UniformScaled, Grid);
+		const FVector2d& PointToBeProcessed = NodeToBeProcessed->Get2DPoint(EGridSpace::UniformScaled, Grid);
 		if (CheckAndRemovePick(PreviousPoint, NodeToBeProcessed->Get2DPoint(EGridSpace::UniformScaled, Grid), NodeToBeProcessed->GetNextNode().Get2DPoint(EGridSpace::UniformScaled, Grid), *NodeToBeProcessed) ||
 			CheckAndRemoveCoincidence(PreviousPoint, PointToBeProcessed, *NodeToBeProcessed))
 		{
@@ -1254,7 +1254,7 @@ bool FLoopCleaner::MoveIntersectingSectionBehindOppositeSection(LoopCleanerImpl:
 	int32 OppositeSectionCount = 1;
 	for (FLoopNode* SegmentNode = FirstNodeOppositeSection; SegmentNode != LastNodeOppositeSection; SegmentNode = GetNext(SegmentNode), ++OppositeSectionCount);
 
-	TArray<const FPoint2D*> OppositeSectionPoint;
+	TArray<const FVector2d*> OppositeSectionPoint;
 	OppositeSectionPoint.Reserve(OppositeSectionCount);
 
 	for (FLoopNode* SegmentNode = FirstNodeOppositeSection; SegmentNode != LastNodeOppositeSection; SegmentNode = GetNext(SegmentNode))
@@ -1288,13 +1288,13 @@ bool FLoopCleaner::MoveIntersectingSectionBehindOppositeSection(LoopCleanerImpl:
 		FLoopNode* NodeToBeProcessed = Node;
 		Node = GetNext(Node);
 
-		FPoint2D  CandiatePosition;
+		FVector2d  CandiatePosition;
 		double MinSquareDistance = HUGE_VALUE;
-		FPoint2D PointToProject = NodeToBeProcessed->Get2DPoint(EGridSpace::UniformScaled, Grid);
+		FVector2d PointToProject = NodeToBeProcessed->Get2DPoint(EGridSpace::UniformScaled, Grid);
 		for (int32 Index = 1; Index < OppositeSectionCount; ++Index)
 		{
-			FPoint2D ProjectedPoint = ProjectPointOnSegment(PointToProject, *OppositeSectionPoint[Index - 1], *OppositeSectionPoint[Index], Coordinate);
-			double SquareDistance = FPoint2D::DistSquared(PointToProject, ProjectedPoint);
+			FVector2d ProjectedPoint = ProjectPointOnSegment<FVector2d>(PointToProject, *OppositeSectionPoint[Index - 1], *OppositeSectionPoint[Index], Coordinate);
+			double SquareDistance = FVector2d::DistSquared(PointToProject, ProjectedPoint);
 			if (SquareDistance < MinSquareDistance)
 			{
 				MinSquareDistance = SquareDistance;
@@ -1322,12 +1322,12 @@ void FLoopCleaner::MoveNodeBehindSegment(const FIsoSegment& IntersectingSegment,
 	Nodes[0] = &IntersectingSegment.GetFirstNode();
 	Nodes[1] = &IntersectingSegment.GetSecondNode();
 
-	FPoint2D IntersectingPoints[2];
+	FVector2d IntersectingPoints[2];
 	IntersectingPoints[0] = IntersectingSegment.GetFirstNode().Get2DPoint(EGridSpace::UniformScaled, Grid);
 	IntersectingPoints[1] = IntersectingSegment.GetSecondNode().Get2DPoint(EGridSpace::UniformScaled, Grid);
 
 	bool bEndNodeIsOutside = true;
-	FPoint2D PointToMove = NodeToMove.Get2DPoint(EGridSpace::UniformScaled, Grid);
+	FVector2d PointToMove = NodeToMove.Get2DPoint(EGridSpace::UniformScaled, Grid);
 
 #ifdef DEBUG_CLOSED_OUSIDE_POINT
 	if (bDisplay)
@@ -1338,7 +1338,7 @@ void FLoopCleaner::MoveNodeBehindSegment(const FIsoSegment& IntersectingSegment,
 #endif
 
 	double Coordinate;
-	FPoint2D ProjectedPoint = ProjectPointOnSegment(PointToMove, IntersectingPoints[0], IntersectingPoints[1], Coordinate);
+	FVector2d ProjectedPoint = ProjectPointOnSegment<FVector2d>(PointToMove, IntersectingPoints[0], IntersectingPoints[1], Coordinate);
 	MoveNode(NodeToMove, ProjectedPoint);
 }
 
@@ -1396,8 +1396,8 @@ bool FLoopCleaner::DisconnectCoincidentNodes(const LoopCleanerImpl::FPinchInters
 	// #cadkernel_check: Why is this constant equal to 30????
 	constexpr double ScaleFactor = 30.;
 	const double Scale = Tolerances.GeometricTolerance * ScaleFactor;
-	FPoint2D P1 = SlopeToVector(MediumSlope1).GetSafeNormal() * Scale;
-	FPoint2D P0 = SlopeToVector(MediumSlope0).GetSafeNormal() * Scale;
+	FVector2d P1 = SlopeToVector(MediumSlope1).GetSafeNormal() * Scale;
+	FVector2d P0 = SlopeToVector(MediumSlope0).GetSafeNormal() * Scale;
 
 	P0 += *Context.Points[0][1];
 	P1 += *Context.Points[1][1];
@@ -1513,8 +1513,8 @@ bool FLoopCleaner::SwapNodes(const TPair<double, double>& Intersection)
 	FLoopNode* Pick = GetNext(Node0);
 	FLoopNode* Node1 = GetNext(Pick);
 
-	const FPoint2D Point0Copy = Node0->Get2DPoint(EGridSpace::UniformScaled, Grid);
-	const FPoint2D& Point1 = Node1->Get2DPoint(EGridSpace::UniformScaled, Grid);
+	const FVector2d Point0Copy = Node0->Get2DPoint(EGridSpace::UniformScaled, Grid);
+	const FVector2d& Point1 = Node1->Get2DPoint(EGridSpace::UniformScaled, Grid);
 
 	Node0->Set2DPoint(EGridSpace::UniformScaled, Grid, Point1);
 	Node1->Set2DPoint(EGridSpace::UniformScaled, Grid, Point0Copy);
@@ -1542,10 +1542,10 @@ bool FLoopCleaner::TryToSwapSegmentsOrRemoveLoop(const TPair<double, double>& In
 		return false;
 	}
 
-	const FPoint2D& Segment0_Point0 = Segment0_Node0->Get2DPoint(EGridSpace::UniformScaled, Grid);
-	const FPoint2D& Segment0_Point1 = Segment0_Node1->Get2DPoint(EGridSpace::UniformScaled, Grid);
-	const FPoint2D& Segment1_Point0 = Segment1_Node0->Get2DPoint(EGridSpace::UniformScaled, Grid);
-	const FPoint2D& Segment1_Point1 = Segment1_Node1->Get2DPoint(EGridSpace::UniformScaled, Grid);
+	const FVector2d& Segment0_Point0 = Segment0_Node0->Get2DPoint(EGridSpace::UniformScaled, Grid);
+	const FVector2d& Segment0_Point1 = Segment0_Node1->Get2DPoint(EGridSpace::UniformScaled, Grid);
+	const FVector2d& Segment1_Point0 = Segment1_Node0->Get2DPoint(EGridSpace::UniformScaled, Grid);
+	const FVector2d& Segment1_Point1 = Segment1_Node1->Get2DPoint(EGridSpace::UniformScaled, Grid);
 
 #ifdef DEBUG_SWAP_SEGMENTS_OR_REMOVE
 	if (bDisplay)
@@ -1681,21 +1681,21 @@ bool FLoopCleaner::RemoveSubLoop(FLoopNode* StartNode, FLoopNode* EndNode)
 	return true;
 }
 
-void FLoopCleaner::MoveNode(FLoopNode& NodeToMove, FPoint2D& NewPostion)
+void FLoopCleaner::MoveNode(FLoopNode& NodeToMove, FVector2d& NewPostion)
 {
-	const FPoint2D& PointToMove = NodeToMove.Get2DPoint(EGridSpace::UniformScaled, Grid);
+	const FVector2d& PointToMove = NodeToMove.Get2DPoint(EGridSpace::UniformScaled, Grid);
 
-	FPoint2D ProjectedDirection = NewPostion - PointToMove;
+	FVector2d ProjectedDirection = NewPostion - PointToMove;
 	ProjectedDirection.Normalize();
 	ProjectedDirection *= Tolerances.GeometricTolerance;
 
-	FPoint2D NewCoordinate = NewPostion + ProjectedDirection;
+	FVector2d NewCoordinate = NewPostion + ProjectedDirection;
 
 	FLoopNode& PreviousNode = NodeToMove.GetPreviousNode();
 	FLoopNode& NextNode = NodeToMove.GetNextNode();
 
-	if ((FPoint2D::DistSquared(PreviousNode.Get2DPoint(EGridSpace::UniformScaled, Grid), NewCoordinate) < Tolerances.SquareGeometricTolerance2) ||
-		(FPoint2D::DistSquared(NextNode.Get2DPoint(EGridSpace::UniformScaled, Grid), NewCoordinate) < Tolerances.SquareGeometricTolerance2))
+	if ((FVector2d::DistSquared(PreviousNode.Get2DPoint(EGridSpace::UniformScaled, Grid), NewCoordinate) < Tolerances.SquareGeometricTolerance2) ||
+		(FVector2d::DistSquared(NextNode.Get2DPoint(EGridSpace::UniformScaled, Grid), NewCoordinate) < Tolerances.SquareGeometricTolerance2))
 	{
 		RemoveNodeOfLoop(NodeToMove);
 	}
@@ -1773,10 +1773,10 @@ void FLoopCleaner::FindLoopIntersections()
 				FSegment2D IntersectedSegmentPoints(GetFirst(IntersectedSegment)->Get2DPoint(EGridSpace::UniformScaled, Grid), GetSecond(IntersectedSegment)->Get2DPoint(EGridSpace::UniformScaled, Grid));
 
 				double SegmentIntersectionCoordinate;
-				FPoint2D IntersectionPoint = FindIntersectionOfSegments2D(SegmentPoints, IntersectedSegmentPoints, SegmentIntersectionCoordinate);
+				FVector2d IntersectionPoint = FindIntersectionOfSegments2D(SegmentPoints, IntersectedSegmentPoints, SegmentIntersectionCoordinate);
 
 				const double IntersectedSegmentPointsSquareLength = IntersectedSegmentPoints.SquaredLength();
-				const double SquareDistanceToIntersectionPoint = FPoint2D::DistSquared(IntersectedSegmentPoints[0], IntersectionPoint);
+				const double SquareDistanceToIntersectionPoint = FVector2d::DistSquared(IntersectedSegmentPoints[0], IntersectionPoint);
 				if (SquareDistanceToIntersectionPoint > IntersectedSegmentPointsSquareLength)
 				{
 					// the IntersectionPoint is not inside IntersectedSegmentPoints
@@ -1907,10 +1907,10 @@ bool FLoopCleaner::RemoveOuterNode(const TPair<double, double>& Intersection)
 	FLoopNode* Node1 = GetNext(Node0);
 	FLoopNode* NextNode = GetNext(Node1);
 
-	const FPoint2D* PreviousPoint = &PreviousNode->Get2DPoint(EGridSpace::UniformScaled, Grid);
-	const FPoint2D* Point0 = &Node0->Get2DPoint(EGridSpace::UniformScaled, Grid);
-	const FPoint2D* Point1 = &Node1->Get2DPoint(EGridSpace::UniformScaled, Grid);
-	const FPoint2D* NextPoint = &NextNode->Get2DPoint(EGridSpace::UniformScaled, Grid);
+	const FVector2d* PreviousPoint = &PreviousNode->Get2DPoint(EGridSpace::UniformScaled, Grid);
+	const FVector2d* Point0 = &Node0->Get2DPoint(EGridSpace::UniformScaled, Grid);
+	const FVector2d* Point1 = &Node1->Get2DPoint(EGridSpace::UniformScaled, Grid);
+	const FVector2d* NextPoint = &NextNode->Get2DPoint(EGridSpace::UniformScaled, Grid);
 
 	double Slope0 = ComputeUnorientedSlope(*Point0, *PreviousPoint, *Point1);
 	double Slope1 = ComputeUnorientedSlope(*Point1, *PreviousPoint, *NextPoint);
@@ -1966,10 +1966,10 @@ bool FLoopCleaner::FindAndRemoveCoincidence(FLoopNode*& StartNode)
 #endif
 
 	FLoopNode* PreviousNode = StartNode;
-	const FPoint2D* PreviousPoint = &StartNode->Get2DPoint(EGridSpace::UniformScaled, Grid);
+	const FVector2d* PreviousPoint = &StartNode->Get2DPoint(EGridSpace::UniformScaled, Grid);
 	for (FLoopNode* Node = GetNext(StartNode); ; Node = GetNext(Node))
 	{
-		const FPoint2D& Point = Node->Get2DPoint(EGridSpace::UniformScaled, Grid);
+		const FVector2d& Point = Node->Get2DPoint(EGridSpace::UniformScaled, Grid);
 
 		if (CheckAndRemoveCoincidence(*PreviousPoint, Point, *PreviousNode))
 		{
@@ -2031,10 +2031,10 @@ bool FLoopCleaner::RemovePickRecursively(FLoopNode* Node0, FLoopNode* Node1)
 		return false;
 	}
 
-	const FPoint2D* PreviousPoint = &PreviousNode->Get2DPoint(EGridSpace::UniformScaled, Grid);
-	const FPoint2D* Point0 = &Node0->Get2DPoint(EGridSpace::UniformScaled, Grid);
-	const FPoint2D* Point1 = &Node1->Get2DPoint(EGridSpace::UniformScaled, Grid);
-	const FPoint2D* NextPoint = &NextNode->Get2DPoint(EGridSpace::UniformScaled, Grid);
+	const FVector2d* PreviousPoint = &PreviousNode->Get2DPoint(EGridSpace::UniformScaled, Grid);
+	const FVector2d* Point0 = &Node0->Get2DPoint(EGridSpace::UniformScaled, Grid);
+	const FVector2d* Point1 = &Node1->Get2DPoint(EGridSpace::UniformScaled, Grid);
+	const FVector2d* NextPoint = &NextNode->Get2DPoint(EGridSpace::UniformScaled, Grid);
 
 	while (true)
 	{
@@ -2183,8 +2183,8 @@ void FLoopCleaner::FindBestLoopExtremity()
 		FLoopNode& PreviousNode = Node->GetPreviousNode();
 		FLoopNode& NextNode = Node->GetNextNode();
 
-		if (FPoint2D::DistSquared(Node->Get2DPoint(EGridSpace::UniformScaled, Grid), PreviousNode.Get2DPoint(EGridSpace::UniformScaled, Grid)) < DOUBLE_SMALL_NUMBER ||
-			FPoint2D::DistSquared(Node->Get2DPoint(EGridSpace::UniformScaled, Grid), NextNode.Get2DPoint(EGridSpace::UniformScaled, Grid)) < DOUBLE_SMALL_NUMBER)
+		if (FVector2d::DistSquared(Node->Get2DPoint(EGridSpace::UniformScaled, Grid), PreviousNode.Get2DPoint(EGridSpace::UniformScaled, Grid)) < DOUBLE_SMALL_NUMBER ||
+			FVector2d::DistSquared(Node->Get2DPoint(EGridSpace::UniformScaled, Grid), NextNode.Get2DPoint(EGridSpace::UniformScaled, Grid)) < DOUBLE_SMALL_NUMBER)
 		{
 			return;
 		}
@@ -2244,7 +2244,7 @@ void FLoopCleaner::FindBestLoopExtremity()
 			LoopIndex = Node.GetLoopIndex();
 		}
 		Index++;
-		const FPoint2D& Point = Node.Get2DPoint(EGridSpace::UniformScaled, Grid);
+		const FVector2d& Point = Node.Get2DPoint(EGridSpace::UniformScaled, Grid);
 
 		if (Point.X > UMax)
 		{
@@ -2319,9 +2319,9 @@ EOrientation FLoopCleaner::GetLoopOrientation(const FLoopNode* StartNode)
 		const FLoopNode& PreviousNode = Node->GetPreviousNode();
 		const FLoopNode& NextNode = Node->GetNextNode();
 
-		const FPoint2D Point = Node->Get2DPoint(EGridSpace::UniformScaled, Grid);
-		const FPoint2D PreviousPoint = PreviousNode.Get2DPoint(EGridSpace::UniformScaled, Grid);
-		const FPoint2D NextPoint = NextNode.Get2DPoint(EGridSpace::UniformScaled, Grid);
+		const FVector2d Point = Node->Get2DPoint(EGridSpace::UniformScaled, Grid);
+		const FVector2d PreviousPoint = PreviousNode.Get2DPoint(EGridSpace::UniformScaled, Grid);
+		const FVector2d NextPoint = NextNode.Get2DPoint(EGridSpace::UniformScaled, Grid);
 
 		double Slope = ComputePositiveSlope(Point, PreviousPoint, NextPoint);
 
@@ -2379,7 +2379,7 @@ EOrientation FLoopCleaner::GetLoopOrientation(const FLoopNode* StartNode)
 	const FLoopNode* Node = StartNode;
 	do
 	{
-		const FPoint2D& Point = Node->Get2DPoint(EGridSpace::UniformScaled, Grid);
+		const FVector2d& Point = Node->Get2DPoint(EGridSpace::UniformScaled, Grid);
 
 		if (Point.X > UMax)
 		{
diff --git a/Engine/Source/Runtime/Datasmith/CADKernel/Base/Private/Mesh/Structure/LoopCleaner.h b/Engine/Source/Runtime/Datasmith/CADKernel/Base/Private/Mesh/Structure/LoopCleaner.h
--- a/Engine/Source/Runtime/Datasmith/CADKernel/Base/Private/Mesh/Structure/LoopCleaner.h
+++ b/Engine/Source/Runtime/Datasmith/CADKernel/Base/Private/Mesh/Structure/LoopCleaner.h
@@ -21,7 +21,7 @@ class FMeshingTolerances;
 
 namespace LoopCleanerImpl
 {
-typedef TFunction<double(const FPoint2D&, const FPoint2D&, double)> GetSlopeMethod;
+typedef TFunction<double(const FVector2d&, const FVector2d&, double)> GetSlopeMethod;
 typedef TFunction<FLoopNode* (FLoopNode*)> GetNextNodeMethod;
 typedef TFunction<const FLoopNode* (const FLoopNode*)> GetNextConstNodeMethod;
 typedef TFunction<const FLoopNode* (const FIsoSegment*)> GetSegmentToNodeMethod;
@@ -31,7 +31,7 @@ struct FPinchIntersectionContext
 {
 	const TPair<double, double>& Intersection;
 	FLoopNode* Nodes[2][3];
-	TArray<const FPoint2D*> Points[2];
+	TArray<const FVector2d*> Points[2];
 
 	FPinchIntersectionContext(const TPair<double, double>& InIntersection)
 		: Intersection(InIntersection)
@@ -174,14 +174,14 @@ private:
 	/**
 	 * @return true if the node has been deleted
 	 */
-	bool CheckAndRemovePick(const FPoint2D& PreviousPoint, const FPoint2D& NodeToRemovePoint, const FPoint2D& NextPoint, FLoopNode& NodeToRemove)
+	bool CheckAndRemovePick(const FVector2d& PreviousPoint, const FVector2d& NodeToRemovePoint, const FVector2d& NextPoint, FLoopNode& NodeToRemove)
 	{
 		double Slope = ComputeUnorientedSlope(NodeToRemovePoint, PreviousPoint, NextPoint);
 		bool bRemoveNode = false;
 		if (Slope < 0.1)
 		{
-			double SquareDistance1 = SquareDistanceOfPointToSegment(PreviousPoint, NodeToRemovePoint, NextPoint);
-			double SquareDistance2 = SquareDistanceOfPointToSegment(NextPoint, NodeToRemovePoint, PreviousPoint);
+			double SquareDistance1 = SquareDistanceOfPointToSegment<FVector2d>(PreviousPoint, NodeToRemovePoint, NextPoint);
+			double SquareDistance2 = SquareDistanceOfPointToSegment<FVector2d>(NextPoint, NodeToRemovePoint, PreviousPoint);
 			double MinSquareDistance = FMath::Min(SquareDistance1, SquareDistance2);
 			if (MinSquareDistance < Tolerances.SquareGeometricTolerance2)
 			{
@@ -199,9 +199,9 @@ private:
 	/**
 	 * @return true if the node has been deleted
 	 */
-	bool CheckAndRemoveCoincidence(const FPoint2D& Point0, const FPoint2D& Point1, FLoopNode& NodeToRemove)
+	bool CheckAndRemoveCoincidence(const FVector2d& Point0, const FVector2d& Point1, FLoopNode& NodeToRemove)
 	{
-		double SquareDistance = FPoint2D::DistSquared(Point0, Point1);
+		double SquareDistance = FVector2d::DistSquared(Point0, Point1);
 		if (SquareDistance < Tolerances.SquareGeometricTolerance2)
 		{
 			return RemoveNodeOfLoop(NodeToRemove);
@@ -249,7 +249,7 @@ private:
 
 	bool MoveIntersectingSectionBehindOppositeSection(LoopCleanerImpl::FLoopSection IntersectingSection, LoopCleanerImpl::FLoopSection OppositeSection);
 	void MoveNodeBehindSegment(const FIsoSegment& IntersectingSegment, FLoopNode& NodeToMove);
-	void MoveNode(FLoopNode& NodeToMove, FPoint2D& NewPosition);
+	void MoveNode(FLoopNode& NodeToMove, FVector2d& NewPosition);
 
 	void FixLoopOrientation();
 	EOrientation GetLoopOrientation(const FLoopNode* StartNode);
diff --git a/Engine/Source/Runtime/Datasmith/CADKernel/Base/Private/Mesh/Structure/Mesh.cpp b/Engine/Source/Runtime/Datasmith/CADKernel/Base/Private/Mesh/Structure/Mesh.cpp
--- a/Engine/Source/Runtime/Datasmith/CADKernel/Base/Private/Mesh/Structure/Mesh.cpp
+++ b/Engine/Source/Runtime/Datasmith/CADKernel/Base/Private/Mesh/Structure/Mesh.cpp
@@ -14,10 +14,10 @@ namespace UE::CADKernel
 TArray<double> FEdgeMesh::GetElementLengths() const
 {
 	const FTopologicalEdge& Edge = static_cast<const FTopologicalEdge&>(TopologicalEntity);
-	const TArray<FPoint>& MeshInnerNodes = GetNodeCoordinates();
+	const TArray<FVector>& MeshInnerNodes = GetNodeCoordinates();
 
-	const FPoint& StartNode = Edge.GetStartVertex()->GetCoordinates();
-	const FPoint& EndNode = Edge.GetEndVertex()->GetCoordinates();
+	const FVector& StartNode = Edge.GetStartVertex()->GetCoordinates();
+	const FVector& EndNode = Edge.GetEndVertex()->GetCoordinates();
 	return PolylineTools::ComputePolylineSegmentLengths(StartNode, MeshInnerNodes, EndNode);
 }
 
diff --git a/Engine/Source/Runtime/Datasmith/CADKernel/Base/Private/Mesh/Structure/ModelMesh.cpp b/Engine/Source/Runtime/Datasmith/CADKernel/Base/Private/Mesh/Structure/ModelMesh.cpp
--- a/Engine/Source/Runtime/Datasmith/CADKernel/Base/Private/Mesh/Structure/ModelMesh.cpp
+++ b/Engine/Source/Runtime/Datasmith/CADKernel/Base/Private/Mesh/Structure/ModelMesh.cpp
@@ -57,11 +57,11 @@ const FVertexMesh* FModelMesh::GetMeshOfVertexNodeId(const int32 Ident) const
 	return nullptr;
 }
 
-void FModelMesh::GetNodeCoordinates(TArray<FPoint>& NodeCoordinates) const
+void FModelMesh::GetNodeCoordinates(TArray<FVector>& NodeCoordinates) const
 {
 	NodeCoordinates.Reserve(LastIdUsed + 1);
 
-	for (const TArray<FPoint>* PointArray : GlobalPointCloud)
+	for (const TArray<FVector>* PointArray : GlobalPointCloud)
 	{
 		NodeCoordinates.Insert(*PointArray, NodeCoordinates.Num());
 	}
@@ -71,9 +71,9 @@ void FModelMesh::GetNodeCoordinates(TArray<FVector3f>& NodeCoordinates) const
 {
 	NodeCoordinates.Reserve(LastIdUsed);
 
-	for (const TArray<FPoint>* PointArray : GlobalPointCloud)
+	for (const TArray<FVector>* PointArray : GlobalPointCloud)
 	{
-		for (const FPoint& Point : *PointArray)
+		for (const FVector& Point : *PointArray)
 		{
 			NodeCoordinates.Emplace(Point.X, Point.Y, Point.Z);
 		}
diff --git a/Engine/Source/Runtime/Datasmith/CADKernel/Base/Private/Mesh/Structure/ModelMeshAnalyzer.cpp b/Engine/Source/Runtime/Datasmith/CADKernel/Base/Private/Mesh/Structure/ModelMeshAnalyzer.cpp
--- a/Engine/Source/Runtime/Datasmith/CADKernel/Base/Private/Mesh/Structure/ModelMeshAnalyzer.cpp
+++ b/Engine/Source/Runtime/Datasmith/CADKernel/Base/Private/Mesh/Structure/ModelMeshAnalyzer.cpp
@@ -194,8 +194,8 @@ double FModelMeshAnalyzer::ComputeMaxAngle() const
 			continue;
 		}
 
-		FPoint Normal0 = Edge.Triangles[0]->ComputeNormal(NodeCoordinates);
-		FPoint Normal1 = Edge.Triangles[1]->ComputeNormal(NodeCoordinates);
+		FVector Normal0 = Edge.Triangles[0]->ComputeNormal(NodeCoordinates);
+		FVector Normal1 = Edge.Triangles[1]->ComputeNormal(NodeCoordinates);
 
 		double Angle = FVectorUtil::ComputeAngle(Normal0, Normal1);
 		if (Angle > MaxAngle)
@@ -247,7 +247,7 @@ void FModelMeshAnalyzer::Display()
 
 void FModelMeshAnalyzer::DisplayTriangle(const Analyzer::FTriangle& Triangle)
 {
-	TArray<FPoint> Vertices;
+	TArray<FVector> Vertices;
 	Vertices.Reserve(3);
 	Vertices.Add(NodeCoordinates[Triangle.Vertices[0]]);
 	Vertices.Add(NodeCoordinates[Triangle.Vertices[1]]);
@@ -260,7 +260,7 @@ void FModelMeshAnalyzer::DisplayEdge(const Analyzer::FEdge& Edge, EVisuProperty
 	DrawSegment( NodeCoordinates[Edge.VertexIndices[0]], NodeCoordinates[Edge.VertexIndices[1]], Property);
 }
 
-FPoint Analyzer::FTriangle::ComputeNormal(const TArray<FPoint>& NodeCoordinates) const
+FVector Analyzer::FTriangle::ComputeNormal(const TArray<FVector>& NodeCoordinates) const
 {
 	return UE::CADKernel::FTriangle(NodeCoordinates[Vertices[0]], NodeCoordinates[Vertices[1]], NodeCoordinates[Vertices[2]]).ComputeNormal();
 }
diff --git a/Engine/Source/Runtime/Datasmith/CADKernel/Base/Private/Mesh/Structure/ModelMeshAnalyzer.h b/Engine/Source/Runtime/Datasmith/CADKernel/Base/Private/Mesh/Structure/ModelMeshAnalyzer.h
--- a/Engine/Source/Runtime/Datasmith/CADKernel/Base/Private/Mesh/Structure/ModelMeshAnalyzer.h
+++ b/Engine/Source/Runtime/Datasmith/CADKernel/Base/Private/Mesh/Structure/ModelMeshAnalyzer.h
@@ -43,7 +43,7 @@ struct FTriangle
 		}
 	}
 
-	FPoint ComputeNormal(const TArray<FPoint>& NodeCoordinates) const;
+	FVector ComputeNormal(const TArray<FVector>& NodeCoordinates) const;
 };
 
 struct FEdge : public FHaveStates
@@ -80,7 +80,7 @@ class CADKERNEL_API FModelMeshAnalyzer
 {
 	const FModelMesh& ModelMesh;
 
-	TArray<FPoint> NodeCoordinates;
+	TArray<FVector> NodeCoordinates;
 	TArray<Analyzer::FEdge> Edges;
 	TArray<Analyzer::FTriangle> Triangles;
 	TArray<TArray<Analyzer::FEdge*>> VertexToEdges;  // Vertex to border or non manifold edges 
diff --git a/Engine/Source/Runtime/Datasmith/CADKernel/Base/Private/Mesh/Structure/ThinZone2DFinder.cpp b/Engine/Source/Runtime/Datasmith/CADKernel/Base/Private/Mesh/Structure/ThinZone2DFinder.cpp
--- a/Engine/Source/Runtime/Datasmith/CADKernel/Base/Private/Mesh/Structure/ThinZone2DFinder.cpp
+++ b/Engine/Source/Runtime/Datasmith/CADKernel/Base/Private/Mesh/Structure/ThinZone2DFinder.cpp
@@ -66,9 +66,9 @@ void FThinZone2DFinder::FindCloseSegments()
 			SegmentLoop = nullptr;
 		}
 
-		const FPoint2D SegmentMiddle = Segment->GetCenter();
+		const FVector2d SegmentMiddle = Segment->GetCenter();
 
-		FPoint2D ClosePoint;
+		FVector2d ClosePoint;
 		FEdgeSegment* CloseSegment = nullptr;
 
 		//     DistanceMin between Candidate and Segment with Candidate is before 
@@ -139,11 +139,11 @@ void FThinZone2DFinder::FindCloseSegments()
 				break;
 			}
 
-			const FPoint2D& FirstPointCandidate = Candidate->GetExtemity(ELimit::Start);
+			const FVector2d& FirstPointCandidate = Candidate->GetExtemity(ELimit::Start);
 
 			// if the distance of FirstPointCandidate with the middle of current segment is biggest than MaxSpace, then the projection of the point cannot be smaller than the tolerance, 
 			{
-				const double SquareDistance = FPoint2D::DistSquared(SegmentMiddle, FirstPointCandidate);
+				const double SquareDistance = FVector2d::DistSquared(SegmentMiddle, FirstPointCandidate);
 				if (SquareDistance > MaxSquareSpace)
 				{
 					continue;
@@ -159,9 +159,9 @@ void FThinZone2DFinder::FindCloseSegments()
 				}
 			}
 
-			const FPoint2D& SecondPointCandidate = Candidate->GetExtemity(ELimit::End);
+			const FVector2d& SecondPointCandidate = Candidate->GetExtemity(ELimit::End);
 			double Coordinate;
-			const FPoint2D Projection = ProjectPointOnSegment(SegmentMiddle, FirstPointCandidate, SecondPointCandidate, Coordinate, true);
+			const FVector2d Projection = ProjectPointOnSegment<FVector2d>(SegmentMiddle, FirstPointCandidate, SecondPointCandidate, Coordinate, true);
 
 #ifdef DEBUG_FIND_CLOSE_SEGMENTS_
 			if (SegmentIndex == SegmentID)
@@ -171,7 +171,7 @@ void FThinZone2DFinder::FindCloseSegments()
 			}
 #endif
 
-			const double SquareDistance = FPoint2D::DistSquared(SegmentMiddle, Projection);
+			const double SquareDistance = FVector2d::DistSquared(SegmentMiddle, Projection);
 			if (SquareDistance > SquareFinderTolerance)
 			{
 				continue;
@@ -235,9 +235,9 @@ void FThinZone2DFinder::FindCloseSegments(const TArray<FEdgeSegment*>& Segments,
 	double MinSquareThickness;
 	TFunction<FEdgeSegment* (const FEdgeSegment*, bool)> FindCloseSegment = [&OppositeSides, &MinSquareThickness](const FEdgeSegment* Segment, const bool bForce)
 	{
-		const FPoint2D SegmentMiddle = Segment->GetCenter();
+		const FVector2d SegmentMiddle = Segment->GetCenter();
 
-		FPoint2D ClosePoint;
+		FVector2d ClosePoint;
 		FEdgeSegment* CloseSegment = nullptr;
 
 		MinSquareThickness = HUGE_VALUE;
@@ -256,13 +256,13 @@ void FThinZone2DFinder::FindCloseSegments(const TArray<FEdgeSegment*>& Segments,
 					}
 				}
 
-				const FPoint2D& FirstPointCandidate = Candidate->GetExtemity(ELimit::Start);
-				const FPoint2D& SecondPointCandidate = Candidate->GetExtemity(ELimit::End);
+				const FVector2d& FirstPointCandidate = Candidate->GetExtemity(ELimit::Start);
+				const FVector2d& SecondPointCandidate = Candidate->GetExtemity(ELimit::End);
 
 				double Coordinate;
-				const FPoint2D Projection = ProjectPointOnSegment(SegmentMiddle, FirstPointCandidate, SecondPointCandidate, Coordinate, true);
+				const FVector2d Projection = ProjectPointOnSegment<FVector2d>(SegmentMiddle, FirstPointCandidate, SecondPointCandidate, Coordinate, true);
 
-				const double SquareDistance = FPoint2D::DistSquared(SegmentMiddle, Projection);
+				const double SquareDistance = FVector2d::DistSquared(SegmentMiddle, Projection);
 				if (MinSquareThickness > SquareDistance)
 				{
 					MinSquareThickness = SquareDistance;
@@ -437,20 +437,20 @@ bool FThinZone2DFinder::CheckIfCloseSideOfThinSideIsNotDegenerated(TArray<FEdgeS
 #endif
 
 	const FEdgeSegment* CloseSegments[3] = { nullptr, nullptr, nullptr };
-	TFunction<FPoint2D(const FPoint2D&)> GetOppositeSegment = [&CloseSegments](const FPoint2D& Candidate) -> FPoint2D
+	TFunction<FVector2d(const FVector2d&)> GetOppositeSegment = [&CloseSegments](const FVector2d& Candidate) -> FVector2d
 	{
-		FPoint2D OutOppositePoint;
+		FVector2d OutOppositePoint = FVector2d::ZeroVector;
 		double MinSquareThickness = HUGE_VALUE;
 		for (int32 Index = 0; CloseSegments[Index] != nullptr; ++Index)
 		{
 			const FEdgeSegment* CloseSegment = CloseSegments[Index];
-			const FPoint2D& FirstPointCandidate = CloseSegment->GetExtemity(ELimit::Start);
-			const FPoint2D& SecondPointCandidate = CloseSegment->GetExtemity(ELimit::End);
+			const FVector2d& FirstPointCandidate = CloseSegment->GetExtemity(ELimit::Start);
+			const FVector2d& SecondPointCandidate = CloseSegment->GetExtemity(ELimit::End);
 
 			double Coordinate;
-			const FPoint2D Projection = ProjectPointOnSegment(Candidate, FirstPointCandidate, SecondPointCandidate, Coordinate, true);
+			const FVector2d Projection = ProjectPointOnSegment<FVector2d>(Candidate, FirstPointCandidate, SecondPointCandidate, Coordinate, true);
 
-			const double SquareDistance = FPoint2D::DistSquared(Candidate, Projection);
+			const double SquareDistance = FVector2d::DistSquared(Candidate, Projection);
 			if (MinSquareThickness > SquareDistance)
 			{
 				MinSquareThickness = SquareDistance;
@@ -548,8 +548,8 @@ bool FThinZone2DFinder::CheckIfCloseSideOfThinSideIsNotDegenerated(TArray<FEdgeS
 		return false;
 	}
 
-	const FPoint2D& StartPoint = Side[0]->GetExtemity(Start);
-	FPoint2D FirstOppositePoint = GetOppositeSegment(StartPoint);
+	const FVector2d& StartPoint = Side[0]->GetExtemity(Start);
+	FVector2d FirstOppositePoint = GetOppositeSegment(StartPoint);
 
 #ifdef DEBUG_CHECK_THIN_SIDE_
 	{
@@ -566,9 +566,9 @@ bool FThinZone2DFinder::CheckIfCloseSideOfThinSideIsNotDegenerated(TArray<FEdgeS
 		{
 			continue;
 		}
-		const FPoint2D& NextPoint = Segment->GetExtemity(End);
-		FPoint2D OppositePoint = GetOppositeSegment(NextPoint);
-		const double SquarewDistance = FPoint2D::DistSquared(OppositePoint, FirstOppositePoint);
+		const FVector2d& NextPoint = Segment->GetExtemity(End);
+		FVector2d OppositePoint = GetOppositeSegment(NextPoint);
+		const double SquarewDistance = FVector2d::DistSquared(OppositePoint, FirstOppositePoint);
 
 #ifdef DEBUG_CHECK_THIN_SIDE_
 		{
@@ -1399,9 +1399,9 @@ bool FThinZone2DFinder::BuildLoopSegments()
 			const double SegmentLength = FMath::Min(FMath::Max(EdgeLength / 1000., WishedSegmentLength), EdgeLength);
 			Edge.Entity->Sample(SegmentLength, Coordinates);
 
-			TArray<FPoint2D> ScaledPoints;
+			TArray<FVector2d> ScaledPoints;
 			{
-				TArray<FPoint2D> Points;
+				TArray<FVector2d> Points;
 				Edge.Entity->Approximate2DPoints(Coordinates, Points);
 				Grid.TransformPoints(EGridSpace::UniformScaled, Points, ScaledPoints);
 			}
@@ -1409,7 +1409,7 @@ bool FThinZone2DFinder::BuildLoopSegments()
 			// Remove duplicated points (vs GeometricTolerance) of the end
 			{
 				int32 EndPointIndex = ScaledPoints.Num() - 1;
-				while (EndPointIndex > 0 && FPoint2D::Distance(ScaledPoints[EndPointIndex], ScaledPoints[EndPointIndex - 1]) < GeometricTolerance)
+				while (EndPointIndex > 0 && FVector2d::Distance(ScaledPoints[EndPointIndex], ScaledPoints[EndPointIndex - 1]) < GeometricTolerance)
 				{
 					EndPointIndex--;
 					ScaledPoints.RemoveAt(EndPointIndex);
@@ -1426,7 +1426,7 @@ bool FThinZone2DFinder::BuildLoopSegments()
 			{
 				for (int32 PointIndex = 1; PointIndex < ScaledPoints.Num(); PointIndex++)
 				{
-					while (PointIndex < ScaledPoints.Num() && FPoint2D::Distance(ScaledPoints[PointIndex - 1], ScaledPoints[PointIndex]) < GeometricTolerance)
+					while (PointIndex < ScaledPoints.Num() && FVector2d::Distance(ScaledPoints[PointIndex - 1], ScaledPoints[PointIndex]) < GeometricTolerance)
 					{
 						ScaledPoints.RemoveAt(PointIndex);
 						Coordinates.RemoveAt(PointIndex);
@@ -1799,7 +1799,7 @@ void FThinZoneSide::GetExistingMeshNodes(const FTopologicalFace& Face, FModelMes
 			}
 
 			const double MeshingTolerance = bWithTolerance ? ComputeTolerance(Index) : 0;
-			const FPoint2D MeshNode2D = EdgeSegment.ComputeEdgePoint(MeshNodeCoordinate);
+			const FVector2d MeshNode2D = EdgeSegment.ComputeEdgePoint(MeshNodeCoordinate);
 			AddMeshNode(NodeIndices[Index], MeshNode2D, MeshingTolerance, EdgeSegment, OppositeNodeIndices[Index]);
 		}
 	};
@@ -1817,7 +1817,7 @@ void FThinZoneSide::GetExistingMeshNodes(const FTopologicalFace& Face, FModelMes
 
 				if (Edge->GetLinkActiveEdge()->IsMeshed())
 				{
-					const TArray<FPoint>& EdgeMeshNodes = Edge->GetMesh()->GetNodeCoordinates();
+					const TArray<FVector>& EdgeMeshNodes = Edge->GetMesh()->GetNodeCoordinates();
 					MeshVertexCount += EdgeMeshNodes.Num();
 				}
 				else if (Edge->IsPreMeshed())
diff --git a/Engine/Source/Runtime/Datasmith/CADKernel/Base/Private/Topo/FaceAnalyzer.cpp b/Engine/Source/Runtime/Datasmith/CADKernel/Base/Private/Topo/FaceAnalyzer.cpp
--- a/Engine/Source/Runtime/Datasmith/CADKernel/Base/Private/Topo/FaceAnalyzer.cpp
+++ b/Engine/Source/Runtime/Datasmith/CADKernel/Base/Private/Topo/FaceAnalyzer.cpp
@@ -54,7 +54,7 @@ void FFaceAnalyzer::FindClosedSegments(Topo::FThinFaceContext& Context)
 			continue;
 		}
 
-		const FPoint& Point = Segment->GetMiddle();
+		const FVector& Point = Segment->GetMiddle();
 
 		Topo::FEdgeSegment* ClosedSegment = nullptr;
 
@@ -68,7 +68,7 @@ void FFaceAnalyzer::FindClosedSegments(Topo::FThinFaceContext& Context)
 				continue;
 			}
 
-			const FPoint& CenterPointCandidate = Candidate->GetMiddle();
+			const FVector& CenterPointCandidate = Candidate->GetMiddle();
 			const double SquareDistanceToCenter = FVector::DistSquared(CenterPointCandidate, Point);
 
 			// if the distance of Point with the middle of candidate segment is biggest than MaxSpace, then the projection of the point cannot be smaller than the tolerance, 
@@ -77,11 +77,11 @@ void FFaceAnalyzer::FindClosedSegments(Topo::FThinFaceContext& Context)
 				continue;
 			}
 
-			const FPoint& FirstPointCandidate = Candidate->GetExtemity(ELimit::Start);
-			const FPoint& SecondPointCandidate = Candidate->GetExtemity(ELimit::End);
+			const FVector& FirstPointCandidate = Candidate->GetExtemity(ELimit::Start);
+			const FVector& SecondPointCandidate = Candidate->GetExtemity(ELimit::End);
 
 			double Coordinate;
-			FPoint Projection = ProjectPointOnSegment(Point, FirstPointCandidate, SecondPointCandidate, Coordinate, true);
+			FVector Projection = ProjectPointOnSegment<FVector>(Point, FirstPointCandidate, SecondPointCandidate, Coordinate, true);
 
 			double SquareDistance = FVector::DistSquared(Point, Projection);
 			if (SquareDistance > MinSquareDistance)
@@ -215,7 +215,7 @@ void FFaceAnalyzer::BuildLoopSegments(Topo::FThinFaceContext& Context)
 		}
 
 		TArray<double> Coordinates;
-		TArray<FPoint> Points;
+		TArray<FVector> Points;
 		Edge->GetCurve()->GetDiscretizationPoints(Edge->GetBoundary(), Coordinates, Points);
 
 		if (OrientedEdge.Direction == EOrientation::Back)
diff --git a/Engine/Source/Runtime/Datasmith/CADKernel/Base/Private/Topo/FaceAnalyzer.h b/Engine/Source/Runtime/Datasmith/CADKernel/Base/Private/Topo/FaceAnalyzer.h
--- a/Engine/Source/Runtime/Datasmith/CADKernel/Base/Private/Topo/FaceAnalyzer.h
+++ b/Engine/Source/Runtime/Datasmith/CADKernel/Base/Private/Topo/FaceAnalyzer.h
@@ -34,9 +34,9 @@ class FEdgeSegment : public FHaveStates
 private:
 	const FTopologicalEdge* Edge = nullptr;
 	double Coordinates[2] = {0. ,0.};
-	FPoint Points[2];
-	FPoint Middle;
-	FPoint Vector;
+	FVector Points[2];
+	FVector Middle;
+	FVector Vector;
 
 	FEdgeSegment* ClosedSegment = nullptr;
 
@@ -57,7 +57,7 @@ public:
 
 	virtual ~FEdgeSegment() = default;
 
-	void SetBoundarySegment(const FTopologicalEdge* InEdge, double InStartU, double InEndU, const FPoint& InStartPoint, const FPoint& InEndPoint)
+	void SetBoundarySegment(const FTopologicalEdge* InEdge, double InStartU, double InEndU, const FVector& InStartPoint, const FVector& InEndPoint)
 	{
 		Edge = InEdge;
 		Coordinates[ELimit::Start] = InStartU;
@@ -79,7 +79,7 @@ public:
 #endif
 	};
 
-	const FPoint GetVector() const
+	const FVector GetVector() const
 	{
 		return Vector;
 	}
@@ -101,12 +101,12 @@ public:
 		return Length;
 	}
 
-	FPoint GetMiddle() const
+	FVector GetMiddle() const
 	{
 		return Middle;
 	}
 
-	constexpr const FPoint& GetExtemity(const ELimit Limit) const
+	constexpr const FVector& GetExtemity(const ELimit Limit) const
 	{
 		return Points[Limit];
 	}
@@ -140,9 +140,9 @@ public:
 		return SquareDistanceToClosedSegment;
 	}
 
-	FPoint ProjectPoint(const FPoint& PointToProject, double& SegmentU) const
+	FVector ProjectPoint(const FVector& PointToProject, double& SegmentU) const
 	{
-		return ProjectPointOnSegment(PointToProject, Points[ELimit::Start], Points[ELimit::End], SegmentU, true);
+		return ProjectPointOnSegment<FVector>(PointToProject, Points[ELimit::Start], Points[ELimit::End], SegmentU, true);
 	}
 
 	bool IsThinZone()
diff --git a/Engine/Source/Runtime/Datasmith/CADKernel/Base/Private/Topo/TopologicalEdge.cpp b/Engine/Source/Runtime/Datasmith/CADKernel/Base/Private/Topo/TopologicalEdge.cpp
--- a/Engine/Source/Runtime/Datasmith/CADKernel/Base/Private/Topo/TopologicalEdge.cpp
+++ b/Engine/Source/Runtime/Datasmith/CADKernel/Base/Private/Topo/TopologicalEdge.cpp
@@ -55,7 +55,7 @@ FTopologicalEdge::FTopologicalEdge(const TSharedRef<FRestrictionCurve>& InCurve,
 	EndVertex = FTopologicalVertex::Make(Points[1].Point);
 }
 
-FTopologicalEdge::FTopologicalEdge(const TSharedRef<FSurface>& InSurface, const FPoint2D& InCoordinateVertex1, const TSharedRef<FTopologicalVertex>& InVertex1, const FPoint2D& InCoordinateVertex2, const TSharedRef<FTopologicalVertex>& InVertex2)
+FTopologicalEdge::FTopologicalEdge(const TSharedRef<FSurface>& InSurface, const FVector2d& InCoordinateVertex1, const TSharedRef<FTopologicalVertex>& InVertex1, const FVector2d& InCoordinateVertex2, const TSharedRef<FTopologicalVertex>& InVertex2)
 	: StartVertex(InVertex1)
 	, EndVertex(InVertex2)
 	, Length3D(-1)
@@ -78,11 +78,11 @@ void FTopologicalEdge::LinkVertex()
 bool FTopologicalEdge::CheckVertices()
 {
 	TArray<double> Coordinates = { Boundary.Min, Boundary.Max };
-	TArray<FPoint> Points;
+	TArray<FVector> Points;
 	Curve->Approximate3DPoints(Coordinates, Points);
 
 	double ToleranceGeo = GetTolerance3D();
-	TFunction<bool(TSharedPtr<FTopologicalVertex>, FPoint)> CheckExtremityGap = [&](TSharedPtr<FTopologicalVertex> Vertex, FPoint Point)
+	TFunction<bool(TSharedPtr<FTopologicalVertex>, FVector)> CheckExtremityGap = [&](TSharedPtr<FTopologicalVertex> Vertex, FVector Point)
 	{
 		double GapToVertex = FVector::Distance(Vertex->GetCoordinates(), Point);
 		return (GapToVertex < ToleranceGeo);
@@ -140,7 +140,7 @@ TSharedPtr<FTopologicalEdge> FTopologicalEdge::Make(const TSharedRef<FRestrictio
 	return Make(InCurve, InCurve->GetBoundary());
 }
 
-TSharedPtr<FTopologicalEdge> FTopologicalEdge::Make(const TSharedRef<FSurface>& InSurface, const FPoint2D& InCoordinateVertex1, const TSharedRef<FTopologicalVertex>& InVertex1, const FPoint2D& InCoordinateVertex2, const TSharedRef<FTopologicalVertex>& InVertex2)
+TSharedPtr<FTopologicalEdge> FTopologicalEdge::Make(const TSharedRef<FSurface>& InSurface, const FVector2d& InCoordinateVertex1, const TSharedRef<FTopologicalVertex>& InVertex1, const FVector2d& InCoordinateVertex2, const TSharedRef<FTopologicalVertex>& InVertex2)
 {
 	TSharedRef<FTopologicalEdge> Edge = FEntity::MakeShared<FTopologicalEdge>(InSurface, InCoordinateVertex1, InVertex1, InCoordinateVertex2, InVertex2);
 	return ReturnIfValid(Edge, /*bCheckVertices*/ false);
@@ -198,10 +198,10 @@ bool FTopologicalEdge::IsTangentAtExtremitiesWith(const FTopologicalEdge& Edge)
 {
 	TFunction<bool(EOrientation)> IsTangentAtExtremities = [&](EOrientation Orientation) ->  bool
 	{
-		FPoint EdgeStartTangent;
-		FPoint EdgeEndTangent;
-		FPoint StartTangent;
-		FPoint EndTangent;
+		FVector EdgeStartTangent;
+		FVector EdgeEndTangent;
+		FVector StartTangent;
+		FVector EndTangent;
 		GetTangentsAtExtremities(StartTangent, EndTangent, Orientation);
 		Edge.GetTangentsAtExtremities(EdgeStartTangent, EdgeEndTangent, EOrientation::Front);
 
@@ -278,10 +278,10 @@ void FTopologicalEdge::LinkIfCoincident(FTopologicalEdge& Twin, double EdgeLengt
 		return;
 	}
 
-	const FPoint& Edge1Vertex1 = GetStartVertex()->GetBarycenter();
-	const FPoint& Edge1Vertex2 = GetEndVertex()->GetBarycenter();
-	const FPoint& Edge2Vertex1 = Twin.GetStartVertex()->GetBarycenter();
-	const FPoint& Edge2Vertex2 = Twin.GetEndVertex()->GetBarycenter();
+	const FVector& Edge1Vertex1 = GetStartVertex()->GetBarycenter();
+	const FVector& Edge1Vertex2 = GetEndVertex()->GetBarycenter();
+	const FVector& Edge2Vertex1 = Twin.GetStartVertex()->GetBarycenter();
+	const FVector& Edge2Vertex2 = Twin.GetEndVertex()->GetBarycenter();
 
 	// Define the orientation
 	const double SquareDistanceE1V1_E2V1 = GetStartVertex()->IsLinkedTo(Twin.GetStartVertex()) ? 0. : FVector::DistSquared(Edge1Vertex1, Edge2Vertex1);
@@ -495,14 +495,14 @@ void FTopologicalEdge::SetEndVertex(const double NewCoordinate)
 	EndVertex->SetCoordinates(OutPoint.Point);
 }
 
-void FTopologicalEdge::SetStartVertex(const double NewCoordinate, const FPoint& NewPoint3D)
+void FTopologicalEdge::SetStartVertex(const double NewCoordinate, const FVector& NewPoint3D)
 {
 	ensureCADKernel(Curve->GetUMin() < NewCoordinate);
 	Boundary.SetMin(NewCoordinate);
 	StartVertex->SetCoordinates(NewPoint3D);
 }
 
-void FTopologicalEdge::SetEndVertex(const double NewCoordinate, const FPoint& NewPoint3D)
+void FTopologicalEdge::SetEndVertex(const double NewCoordinate, const FVector& NewPoint3D)
 {
 	ensureCADKernel(Curve->GetUMax() > NewCoordinate);
 	Boundary.SetMax(NewCoordinate);
@@ -524,7 +524,7 @@ double FTopologicalEdge::Length() const
 	return Length3D;
 }
 
-void FTopologicalEdge::GetTangentsAtExtremities(FPoint& StartTangent, FPoint& EndTangent, bool bForward) const
+void FTopologicalEdge::GetTangentsAtExtremities(FVector& StartTangent, FVector& EndTangent, bool bForward) const
 {
 	ensureCADKernel(Curve->Polyline.Size());
 
@@ -532,7 +532,7 @@ void FTopologicalEdge::GetTangentsAtExtremities(FPoint& StartTangent, FPoint& En
 	int32 StartIndex = Finder.Find(Boundary.Min);
 	int32 EndIndex = Finder.Find(Boundary.Max);
 
-	const TArray<FPoint>& Polyline3D = Curve->Polyline.GetPoints();
+	const TArray<FVector>& Polyline3D = Curve->Polyline.GetPoints();
 	if (bForward)
 	{
 		StartTangent = Polyline3D[StartIndex + 1] - Polyline3D[StartIndex];
@@ -593,8 +593,8 @@ double FTopologicalEdge::TransformTwinEdgeCoordinateToLocalCoordinate(const FTop
 		return Start + (TwinStart - InTwinCoordinate) * Distance / TwinDistance;
 	}
 
-	FPoint PointOnEdge = TwinEdge.GetCurve()->Approximate3DPoint(InTwinCoordinate);
-	FPoint ProjectedPoint;
+	FVector PointOnEdge = TwinEdge.GetCurve()->Approximate3DPoint(InTwinCoordinate);
+	FVector ProjectedPoint;
 	return Curve->GetCoordinateOfProjectedPoint(Boundary, PointOnEdge, ProjectedPoint);
 }
 
@@ -660,7 +660,7 @@ void FTopologicalEdge::TransformTwinEdgeCoordinatesToLocalCoordinates(const FTop
 		return;
 	}
 
-	TArray<FPoint> Cutting3DPoints;
+	TArray<FVector> Cutting3DPoints;
 	TwinEdge.ApproximatePoints(InTwinCoordinates, Cutting3DPoints);
 	ProjectTwinEdgePoints(Cutting3DPoints, IsSameDirection(TwinEdge), OutLocalCoordinates);
 }
@@ -730,7 +730,7 @@ void FTopologicalEdge::SortImposedCuttingPoints()
 
 
 
-void FTopologicalEdge::ProjectTwinEdgePointsOn2DCurve(const TSharedRef<FTopologicalEdge>& InTwinEdge, const TArray<double>& InTwinEdgePointCoords, TArray<FPoint2D>& OutPoints2D)
+void FTopologicalEdge::ProjectTwinEdgePointsOn2DCurve(const TSharedRef<FTopologicalEdge>& InTwinEdge, const TArray<double>& InTwinEdgePointCoords, TArray<FVector2d>& OutPoints2D)
 {
 	if (&InTwinEdge.Get() == this)
 	{
@@ -738,7 +738,7 @@ void FTopologicalEdge::ProjectTwinEdgePointsOn2DCurve(const TSharedRef<FTopologi
 	}
 	else
 	{
-		TArray<FPoint> Points3D;
+		TArray<FVector> Points3D;
 		InTwinEdge->ApproximatePoints(InTwinEdgePointCoords, Points3D);
 
 		TArray<double> Coordinates;
@@ -766,7 +766,7 @@ void FTopologicalEdge::GenerateMeshElements(FModelMesh& MeshModel)
 	CuttingPointCoordinates.RemoveAt(0);
 	CuttingPointCoordinates.Pop();
 
-	TArray<FPoint>& Coordinates = EdgeMesh.GetNodeCoordinates();
+	TArray<FVector>& Coordinates = EdgeMesh.GetNodeCoordinates();
 	ApproximatePoints(CuttingPointCoordinates, Coordinates);
 
 	if (!bSameDirection)
@@ -808,12 +808,12 @@ bool FTopologicalEdge::IsSameDirection(const FTopologicalEdge& Edge) const
 			return true;
 		}
 
-		FPoint EdgeStartTangent;
-		FPoint EdgeEndTangent;
+		FVector EdgeStartTangent;
+		FVector EdgeEndTangent;
 		Edge.GetTangentsAtExtremities(EdgeStartTangent, EdgeEndTangent, true);
 
-		FPoint StartTangent;
-		FPoint EndTangent;
+		FVector StartTangent;
+		FVector EndTangent;
 		GetTangentsAtExtremities(StartTangent, EndTangent, true);
 
 		double StartAngle = FVectorUtil::ComputeCosinus(StartTangent, EdgeStartTangent);
@@ -947,7 +947,7 @@ TArray<double> FTopologicalEdge::GetPreElementLengths() const
 	const TArray<FCuttingPoint>& CuttingPoints = GetCuttingPoints();
 	const int32 CuttingPointCount = CuttingPoints.Num();
 
-	TArray<FPoint> InnerNodes;
+	TArray<FVector> InnerNodes;
 	if (CuttingPointCount > 2)
 	{
 		TArray<double> Coordinates;
@@ -959,8 +959,8 @@ TArray<double> FTopologicalEdge::GetPreElementLengths() const
 		ApproximatePoints(Coordinates, InnerNodes);
 	}
 
-	const FPoint& StartNode = GetStartVertex()->GetCoordinates();
-	const FPoint& EndNode = GetEndVertex()->GetCoordinates();
+	const FVector& StartNode = GetStartVertex()->GetCoordinates();
+	const FVector& EndNode = GetEndVertex()->GetCoordinates();
 	return PolylineTools::ComputePolylineSegmentLengths(StartNode, InnerNodes, EndNode);
 }
 
@@ -1076,21 +1076,21 @@ TSharedPtr<FTopologicalEdge> FTopologicalEdge::CreateEdgeByMergingEdges(const do
 				{
 					TargetCoordinate = Edges[Index].Direction == EOrientation::Front ? Edges[Index].Entity->GetBoundary().GetMin() : Edges[Index].Entity->GetBoundary().GetMax();
 				}
-				FPoint2D Target = Edges[Index].Entity->Approximate2DPoint(TargetCoordinate);
+				FVector2d Target = Edges[Index].Entity->Approximate2DPoint(TargetCoordinate);
 
 				int32 NeigborIndex = PreviousLength > NextLength ? Index - 1 : Index + 1;
 				if (NeigborIndex < 0 || NeigborIndex >= NurbsCurves.Num())
 				{
 					return TSharedPtr<FTopologicalEdge>();
 				}
-				NurbsCurves[NeigborIndex]->ExtendTo(Target);
+				NurbsCurves[NeigborIndex]->ExtendTo(FVector(Target, 0.));
 			}
 		}
 	}
 
 	TArray<double> NewNodalVector;
 	TArray<double> NewWeights;
-	TArray<FPoint> NewPoles;
+	TArray<FVector> NewPoles;
 	NurbsMaxDegree++;
 	NewNodalVector.Reserve(PoleCount + NurbsMaxDegree);
 	NewPoles.Reserve(PoleCount + NurbsMaxDegree);
@@ -1218,7 +1218,7 @@ void FTopologicalEdge::ReplaceEdgeVertex(bool bIsStartVertex, TSharedRef<FTopolo
 	OldVertex = NewVertex;
 }
 
-bool FTopologicalEdge::ExtendTo(bool bIsStartExtremity, const FPoint2D& NewExtremityCoordinate, TSharedRef<FTopologicalVertex>& NewVertex)
+bool FTopologicalEdge::ExtendTo(bool bIsStartExtremity, const FVector2d& NewExtremityCoordinate, TSharedRef<FTopologicalVertex>& NewVertex)
 {
 	if (bIsStartExtremity ? FMath::IsNearlyEqual(Boundary.Min, Curve->GetBoundary().Min) : FMath::IsNearlyEqual(Boundary.Max, Curve->GetBoundary().Max))
 	{
@@ -1226,7 +1226,7 @@ bool FTopologicalEdge::ExtendTo(bool bIsStartExtremity, const FPoint2D& NewExtre
 	}
 	else
 	{
-		FPoint ProjectedPoint;
+		FVector ProjectedPoint;
 		double UProjectedPoint = ProjectPoint(NewVertex->GetCoordinates(), ProjectedPoint);
 		if (FVector::Distance(ProjectedPoint, NewVertex->GetCoordinates()) > GetTolerance3D())
 		{
@@ -1251,8 +1251,8 @@ bool FTopologicalEdge::ExtendTo(bool bIsStartExtremity, const FPoint2D& NewExtre
 
 void FTopologicalEdge::ComputeEdge2DProperties(FEdge2DProperties& EdgeCharacteristics)
 {
-	const TArray<FPoint2D>& Polyline2D = Curve->Polyline.Get2DPoints();
-	const TArray<FPoint>& Polyline3D = Curve->Polyline.GetPoints();
+	const TArray<FVector2d>& Polyline2D = Curve->Polyline.Get2DPoints();
+	const TArray<FVector>& Polyline3D = Curve->Polyline.GetPoints();
 	const TArray<double>& Parameters = Curve->Polyline.GetCoordinates();
 
 	FDichotomyFinder Finder(Curve->Polyline.GetCoordinates());
@@ -1270,7 +1270,7 @@ void FTopologicalEdge::ComputeEdge2DProperties(FEdge2DProperties& EdgeCharacteri
 	}
 }
 
-FPoint FTopologicalEdge::GetTangentAt(const FTopologicalVertex& InVertex)
+FVector FTopologicalEdge::GetTangentAt(const FTopologicalVertex& InVertex)
 {
 	if (&InVertex == StartVertex.Get())
 	{
@@ -1278,7 +1278,7 @@ FPoint FTopologicalEdge::GetTangentAt(const FTopologicalVertex& InVertex)
 	}
 	else if (&InVertex == EndVertex.Get())
 	{
-		FPoint Tangent = Curve->GetTangentAt(Boundary.GetMax());
+		FVector Tangent = Curve->GetTangentAt(Boundary.GetMax());
 		Tangent *= -1.;
 		return Tangent;
 	}
@@ -1288,18 +1288,18 @@ FPoint FTopologicalEdge::GetTangentAt(const FTopologicalVertex& InVertex)
 	}
 	else if (InVertex.GetLink() == EndVertex->GetLink())
 	{
-		FPoint Tangent = Curve->GetTangentAt(Boundary.GetMax());
+		FVector Tangent = Curve->GetTangentAt(Boundary.GetMax());
 		Tangent *= -1.;
 		return Tangent;
 	}
 	else
 	{
 		ensureCADKernel(false);
-		return FPoint::ZeroVector;
+		return FVector::ZeroVector;
 	}
 }
 
-FPoint2D FTopologicalEdge::GetTangent2DAt(const FTopologicalVertex& InVertex)
+FVector2d FTopologicalEdge::GetTangent2DAt(const FTopologicalVertex& InVertex)
 {
 	if (&InVertex == StartVertex.Get())
 	{
@@ -1307,7 +1307,7 @@ FPoint2D FTopologicalEdge::GetTangent2DAt(const FTopologicalVertex& InVertex)
 	}
 	else if (&InVertex == EndVertex.Get())
 	{
-		FPoint Tangent = Curve->GetTangent2DAt(Boundary.GetMax());
+		FVector2d Tangent = Curve->GetTangent2DAt(Boundary.GetMax());
 		Tangent *= -1.;
 		return Tangent;
 	}
@@ -1317,7 +1317,7 @@ FPoint2D FTopologicalEdge::GetTangent2DAt(const FTopologicalVertex& InVertex)
 	}
 	else if (InVertex.GetLink() == EndVertex->GetLink())
 	{
-		FPoint2D Tangent = Curve->GetTangent2DAt(Boundary.GetMax());
+		FVector2d Tangent = Curve->GetTangent2DAt(Boundary.GetMax());
 		Tangent *= -1.;
 		return Tangent;
 	}
@@ -1349,7 +1349,7 @@ void FTopologicalEdge::SpawnIdent(FDatabase& Database)
 	}
 }
 
-FTopologicalVertex* FTopologicalEdge::SplitAt(double SplittingCoordinate, const FPoint& NewVertexCoordinate, bool bKeepStartVertexConnectivity, TSharedPtr<FTopologicalEdge>& NewEdge)
+FTopologicalVertex* FTopologicalEdge::SplitAt(double SplittingCoordinate, const FVector& NewVertexCoordinate, bool bKeepStartVertexConnectivity, TSharedPtr<FTopologicalEdge>& NewEdge)
 {
 	if (GetTwinEntityCount() > 1)
 	{
@@ -1430,7 +1430,7 @@ bool FTopologicalEdge::IsSharpEdge() const
 	return false;
 }
 
-void FTopologicalEdge::Offset2D(const FPoint2D& OffsetDirection)
+void FTopologicalEdge::Offset2D(const FVector2d& OffsetDirection)
 {
 	Curve->Offset2D(OffsetDirection);
 }
diff --git a/Engine/Source/Runtime/Datasmith/CADKernel/Base/Private/Topo/TopologicalFace.cpp b/Engine/Source/Runtime/Datasmith/CADKernel/Base/Private/Topo/TopologicalFace.cpp
--- a/Engine/Source/Runtime/Datasmith/CADKernel/Base/Private/Topo/TopologicalFace.cpp
+++ b/Engine/Source/Runtime/Datasmith/CADKernel/Base/Private/Topo/TopologicalFace.cpp
@@ -26,12 +26,12 @@ namespace UE::CADKernel
 void FTopologicalFace::ComputeBoundary() const
 {
 	Boundary->Init();
-	TArray<TArray<FPoint2D>> TmpLoops;
+	TArray<TArray<FVector2d>> TmpLoops;
 	Get2DLoopSampling(TmpLoops);
 
-	for (const TArray<FPoint2D>& Loop : TmpLoops)
+	for (const TArray<FVector2d>& Loop : TmpLoops)
 	{
-		for (const FPoint2D& Point : Loop)
+		for (const FVector2d& Point : Loop)
 		{
 			Boundary->ExtendTo(Point);
 		}
@@ -52,7 +52,7 @@ void FTopologicalFace::UpdateBBox(int32 IsoCount, const double ApproximationFact
 {
 	const double SAG = GetCarrierSurface()->Get3DTolerance() * ApproximationFactor;
 
-	TArray<TArray<FPoint2D>> BoundaryApproximation;
+	TArray<TArray<FVector2d>> BoundaryApproximation;
 	Get2DLoopSampling(BoundaryApproximation);
 
 	FPolyline3D Polyline;
@@ -92,7 +92,7 @@ void FTopologicalFace::UpdateBBox(int32 IsoCount, const double ApproximationFact
 
 			if (IntersectionCount % 2 != 0)
 			{
-				TArray<FPoint> SubPolyline;
+				TArray<FVector> SubPolyline;
 				FLinearBoundary IntersectionBoundary(Intersections[IntersectionCount - 1], CurveBounds.GetMax());
 
 				if (Polyline.Size() > 1)
@@ -113,7 +113,7 @@ void FTopologicalFace::UpdateBBox(int32 IsoCount, const double ApproximationFact
 			{
 				for (int32 ISection = 0; ISection < IntersectionCount; ISection += 2)
 				{
-					TArray<FPoint> SubPolyline;
+					TArray<FVector> SubPolyline;
 					FLinearBoundary IntersectionBoundary(Intersections[ISection], Intersections[ISection + 1]);
 
 					Polyline.UpdateSubPolylineBBox(IntersectionBoundary, IsoBBox);
@@ -163,7 +163,7 @@ void FTopologicalFace::ApplyNaturalLoops(const FSurfacicBoundary& Boundaries)
 
 	TArray<TSharedPtr<FTopologicalEdge>> Edges;
 	Edges.Reserve(4);
-	TFunction<void(const FPoint&, const FPoint&)> BuildEdge = [&](const FPoint& StartPoint, const FPoint& EndPoint)
+	TFunction<void(const FVector&, const FVector&)> BuildEdge = [&](const FVector& StartPoint, const FVector& EndPoint)
 	{
 		double Tolerance3D = CarrierSurface->Get3DTolerance();
 		TSharedRef<FCurve> Curve2D = FEntity::MakeShared<FSegmentCurve>(StartPoint, EndPoint, 2);
@@ -176,8 +176,8 @@ void FTopologicalFace::ApplyNaturalLoops(const FSurfacicBoundary& Boundaries)
 		Edges.Add(Edge);
 	};
 
-	FPoint StartPoint;
-	FPoint EndPoint;
+	FVector StartPoint;
+	FVector EndPoint;
 
 	// Build 4 bounding edges of the surface
 	StartPoint.Set(Boundaries[EIso::IsoU].Min, Boundaries[EIso::IsoV].Min, 0.);
@@ -427,13 +427,13 @@ void FTopologicalFace::GetEdgeIndex(const FTopologicalEdge& Edge, int32& OutBoun
 	OutBoundaryIndex = INDEX_NONE;
 }
 
-const void FTopologicalFace::Get2DLoopSampling(TArray<TArray<FPoint2D>>& LoopSamplings) const
+const void FTopologicalFace::Get2DLoopSampling(TArray<TArray<FVector2d>>& LoopSamplings) const
 {
 	LoopSamplings.Empty(GetLoops().Num());
 
 	for (const TSharedPtr<FTopologicalLoop>& Loop : GetLoops())
 	{
-		TArray<FPoint2D>& LoopSampling2D = LoopSamplings.Emplace_GetRef();
+		TArray<FVector2d>& LoopSampling2D = LoopSamplings.Emplace_GetRef();
 		Loop->Get2DSampling(LoopSampling2D);
 	}
 }
@@ -517,13 +517,13 @@ void FTopologicalFace::ApplyCriteria(const TArray<TSharedPtr<FCriterion>>& Crite
 	{
 		for (int32 IndexU = 0; IndexU < CrossingCoordinates[EIso::IsoU].Num() - 1; ++IndexU)
 		{
-			const FPoint& Point_U0_V0 = Grid.GetPoint(IndexU, IndexV);
-			const FPoint& Point_U1_V1 = Grid.GetPoint(IndexU + 1, IndexV + 1);
-			const FPoint& Point_Um_V0 = Grid.GetIntermediateU(IndexU, IndexV);
-			const FPoint& Point_Um_V1 = Grid.GetIntermediateU(IndexU, IndexV + 1);
-			const FPoint& Point_U0_Vm = Grid.GetIntermediateV(IndexU, IndexV);
-			const FPoint& Point_U1_Vm = Grid.GetIntermediateV(IndexU + 1, IndexV);
-			const FPoint& Point_Um_Vm = Grid.GetIntermediateUV(IndexU, IndexV);
+			const FVector& Point_U0_V0 = Grid.GetPoint(IndexU, IndexV);
+			const FVector& Point_U1_V1 = Grid.GetPoint(IndexU + 1, IndexV + 1);
+			const FVector& Point_Um_V0 = Grid.GetIntermediateU(IndexU, IndexV);
+			const FVector& Point_Um_V1 = Grid.GetIntermediateU(IndexU, IndexV + 1);
+			const FVector& Point_U0_Vm = Grid.GetIntermediateV(IndexU, IndexV);
+			const FVector& Point_U1_Vm = Grid.GetIntermediateV(IndexU + 1, IndexV);
+			const FVector& Point_Um_Vm = Grid.GetIntermediateUV(IndexU, IndexV);
 
 			// Evaluate Sag
 			double LengthU;
diff --git a/Engine/Source/Runtime/Datasmith/CADKernel/Base/Private/Topo/TopologicalFaceUtilities.cpp b/Engine/Source/Runtime/Datasmith/CADKernel/Base/Private/Topo/TopologicalFaceUtilities.cpp
--- a/Engine/Source/Runtime/Datasmith/CADKernel/Base/Private/Topo/TopologicalFaceUtilities.cpp
+++ b/Engine/Source/Runtime/Datasmith/CADKernel/Base/Private/Topo/TopologicalFaceUtilities.cpp
@@ -58,7 +58,7 @@ namespace UE::CADKernel::TopologicalFaceUtilities
 
 		const FLinearBoundary& EdgeBounds = Edge.GetBoundary();
 		const TArray<double>& PolyCoords = Edge.GetCurve()->GetPolyline().GetCoordinates();
-		const TArray<FPoint2D>& PolyPoints = Edge.GetCurve()->GetPolyline().Get2DPoints();
+		const TArray<FVector2d>& PolyPoints = Edge.GetCurve()->GetPolyline().Get2DPoints();
 		const FLinearBoundary PolyBounds(PolyCoords[0] - UE_DOUBLE_SMALL_NUMBER, PolyCoords.Last() + UE_DOUBLE_SMALL_NUMBER);
 
 		ensureCADKernel(PolyPoints.Num() > 1);
@@ -70,7 +70,7 @@ namespace UE::CADKernel::TopologicalFaceUtilities
 		if (FMath::IsNearlyEqual(EdgeBounds.Min, PolyCoords[0]) && FMath::IsNearlyEqual(EdgeBounds.Max, PolyCoords.Last()))
 		{
 			// The entire polyline is used for the edge, just copy it in
-			for (const FPoint2D& Point : PolyPoints)
+			for (const FVector2d& Point : PolyPoints)
 			{
 				PointsOut.Emplace(Point.X, Point.Y);
 			}
@@ -82,7 +82,7 @@ namespace UE::CADKernel::TopologicalFaceUtilities
 					if (Value < PolyBounds.Min)
 					{
 						const double Ratio = (Value - PolyCoords[0]) / (PolyCoords[1] - PolyCoords[0]);
-						const FPoint2D Point = PolyPoints[0] + (PolyPoints[1] - PolyPoints[0]) * Ratio;
+						const FVector2d Point = PolyPoints[0] + (PolyPoints[1] - PolyPoints[0]) * Ratio;
 
 						PointsOut.Emplace(Point.X, Point.Y);
 					}
@@ -90,7 +90,7 @@ namespace UE::CADKernel::TopologicalFaceUtilities
 					{
 						const int32 LastIndex = PolyCoords.Num() - 1;
 						const double Ratio = (Value - PolyCoords[LastIndex]) / (PolyCoords[LastIndex] - PolyCoords[LastIndex - 1]);
-						const FPoint2D Point = PolyPoints[LastIndex] + (PolyPoints[LastIndex] - PolyPoints[LastIndex - 1]) * Ratio;
+						const FVector2d Point = PolyPoints[LastIndex] + (PolyPoints[LastIndex] - PolyPoints[LastIndex - 1]) * Ratio;
 
 						PointsOut.Emplace(Point.X, Point.Y);
 					}
@@ -100,13 +100,13 @@ namespace UE::CADKernel::TopologicalFaceUtilities
 
 						if (FMath::IsNearlyEqual(Value, PolyCoords[Index]))
 						{
-							const FPoint2D& Point = PolyPoints[Index];
+							const FVector2d& Point = PolyPoints[Index];
 							PointsOut.Emplace(Point.X, Point.Y);
 						}
 						else
 						{
 							const double Ratio = (Value - PolyCoords[Index - 1]) / (PolyCoords[Index] - PolyCoords[Index - 1]);
-							const FPoint2D Point = PolyPoints[Index - 1] + (PolyPoints[Index] - PolyPoints[Index - 1]) * Ratio;
+							const FVector2d Point = PolyPoints[Index - 1] + (PolyPoints[Index] - PolyPoints[Index - 1]) * Ratio;
 
 							PointsOut.Emplace(Point.X, Point.Y);
 						}
@@ -207,7 +207,7 @@ namespace UE::CADKernel::TopologicalFaceUtilities
 
 		const FLinearBoundary& EdgeBounds = Edge.GetBoundary();
 		const TArray<double>& PolyCoords = Edge.GetCurve()->GetPolyline().GetCoordinates();
-		const TArray<FPoint>& PolyPoints = Edge.GetCurve()->GetPolyline().GetPoints();
+		const TArray<FVector>& PolyPoints = Edge.GetCurve()->GetPolyline().GetPoints();
 		const FLinearBoundary PolyBounds(PolyCoords[0] - UE_DOUBLE_SMALL_NUMBER, PolyCoords.Last() + UE_DOUBLE_SMALL_NUMBER);
 
 		ensureCADKernel(PolyPoints.Num() > 1);
@@ -218,7 +218,7 @@ namespace UE::CADKernel::TopologicalFaceUtilities
 		if (FMath::IsNearlyEqual(EdgeBounds.Min, PolyCoords[0]) && FMath::IsNearlyEqual(EdgeBounds.Max, PolyCoords.Last()))
 		{
 			// The entire polyline is used for the edge, just copy it in
-			for (const FPoint& Point : PolyPoints)
+			for (const FVector& Point : PolyPoints)
 			{
 				PointsOut.Emplace(Point.X, Point.Y, Point.Z);
 			}
@@ -230,7 +230,7 @@ namespace UE::CADKernel::TopologicalFaceUtilities
 					if (Value < PolyBounds.Min)
 					{
 						const double Ratio = (Value - PolyCoords[0]) / (PolyCoords[1] - PolyCoords[0]);
-						const FPoint Point = PolyPoints[0] + (PolyPoints[1] - PolyPoints[0]) * Ratio;
+						const FVector Point = PolyPoints[0] + (PolyPoints[1] - PolyPoints[0]) * Ratio;
 
 						PointsOut.Emplace(Point.X, Point.Y, Point.Z);
 					}
@@ -238,7 +238,7 @@ namespace UE::CADKernel::TopologicalFaceUtilities
 					{
 						const int32 LastIndex = PolyCoords.Num() - 1;
 						const double Ratio = (Value - PolyCoords[LastIndex]) / (PolyCoords[LastIndex] - PolyCoords[LastIndex - 1]);
-						const FPoint Point = PolyPoints[LastIndex] + (PolyPoints[LastIndex] - PolyPoints[LastIndex - 1]) * Ratio;
+						const FVector Point = PolyPoints[LastIndex] + (PolyPoints[LastIndex] - PolyPoints[LastIndex - 1]) * Ratio;
 
 						PointsOut.Emplace(Point.X, Point.Y, Point.Z);
 					}
@@ -248,13 +248,13 @@ namespace UE::CADKernel::TopologicalFaceUtilities
 
 						if (FMath::IsNearlyEqual(Value, PolyCoords[Index]))
 						{
-							const FPoint& Point = PolyPoints[Index];
+							const FVector& Point = PolyPoints[Index];
 							PointsOut.Emplace(Point.X, Point.Y, Point.Z);
 						}
 						else
 						{
 							const double Ratio = (Value - PolyCoords[Index - 1]) / (PolyCoords[Index] - PolyCoords[Index - 1]);
-							const FPoint Point = PolyPoints[Index - 1] + (PolyPoints[Index] - PolyPoints[Index - 1]) * Ratio;
+							const FVector Point = PolyPoints[Index - 1] + (PolyPoints[Index] - PolyPoints[Index - 1]) * Ratio;
 
 							PointsOut.Emplace(Point.X, Point.Y, Point.Z);
 						}
diff --git a/Engine/Source/Runtime/Datasmith/CADKernel/Base/Private/Topo/TopologicalLoop.cpp b/Engine/Source/Runtime/Datasmith/CADKernel/Base/Private/Topo/TopologicalLoop.cpp
--- a/Engine/Source/Runtime/Datasmith/CADKernel/Base/Private/Topo/TopologicalLoop.cpp
+++ b/Engine/Source/Runtime/Datasmith/CADKernel/Base/Private/Topo/TopologicalLoop.cpp
@@ -30,18 +30,18 @@ TSharedPtr<FTopologicalLoop> FTopologicalLoop::Make(const TArray<TSharedPtr<FTop
 		const FTopologicalEdge& Edge0 = *InEdges[0];
 		FTopologicalEdge& Edge1 = *InEdges[1];
 
-		const FPoint2D StartPoint = Edge0.Approximate2DPoint(Edge0.GetBoundary().GetMin());
-		const FPoint2D EndPoint = Edge0.Approximate2DPoint(Edge0.GetBoundary().GetMax());
+		const FVector2d StartPoint = Edge0.Approximate2DPoint(Edge0.GetBoundary().GetMin());
+		const FVector2d EndPoint = Edge0.Approximate2DPoint(Edge0.GetBoundary().GetMax());
 		if (FMath::IsNearlyEqual(FMath::Abs(StartPoint.Y), DOUBLE_HALF_PI) && FMath::IsNearlyEqual(FMath::Abs(EndPoint.Y), DOUBLE_HALF_PI))
 		{
 			const double TolU = Surface.GetIsoTolerance(EIso::IsoU);
 			const bool Edge0IsIso = Edge0.GetCurve()->IsIso(EIso::IsoU, TolU);
 			const bool Edge1IsIso = Edge1.GetCurve()->IsIso(EIso::IsoU, TolU);
-			const FPoint2D Edge1EndPoint = Edge1.Approximate2DPoint(InEdgeDirections[0]== InEdgeDirections[1] ? Edge1.GetBoundary().GetMax() : Edge1.GetBoundary().GetMin());
-			double Distance = FPoint2D::DistSquared(Edge1EndPoint, StartPoint);
+			const FVector2d Edge1EndPoint = Edge1.Approximate2DPoint(InEdgeDirections[0]== InEdgeDirections[1] ? Edge1.GetBoundary().GetMax() : Edge1.GetBoundary().GetMin());
+			double Distance = FVector2d::DistSquared(Edge1EndPoint, StartPoint);
 			if (Distance < FMath::Square(TolU))
 			{
-				const FPoint2D Offset(TWO_PI, 0.);
+				const FVector2d Offset(TWO_PI, 0.);
 				Edge1.Offset2D(Offset);
 			}
 		}
@@ -60,7 +60,7 @@ TSharedPtr<FTopologicalLoop> FTopologicalLoop::Make(const TArray<TSharedPtr<FTop
 		OrientedEdge.Entity->SetLoop(Loop);
 	}
 
-	TArray<FPoint2D> LoopSampling;
+	TArray<FVector2d> LoopSampling;
 	Loop.Get2DSampling(LoopSampling);
 	FAABB2D LoopBoundary;
 	LoopBoundary += LoopSampling;
@@ -146,7 +146,7 @@ EOrientation FTopologicalLoop::GetDirection(TSharedPtr<FTopologicalEdge>& InEdge
 }
 
 
-void FTopologicalLoop::Get2DSampling(TArray<FPoint2D>& LoopSampling) const
+void FTopologicalLoop::Get2DSampling(TArray<FVector2d>& LoopSampling) const
 {
 	if (Edges.IsEmpty())
 	{
@@ -169,7 +169,7 @@ void FTopologicalLoop::Get2DSampling(TArray<FPoint2D>& LoopSampling) const
 	LoopSampling.Emplace(LoopSampling[0]);
 }
 
-bool FTopologicalLoop::Get2DSamplingWithoutDegeneratedEdges(TArray<FPoint2D>& LoopSampling) const
+bool FTopologicalLoop::Get2DSamplingWithoutDegeneratedEdges(TArray<FVector2d>& LoopSampling) const
 {
 	if (Edges.IsEmpty())
 	{
@@ -259,7 +259,7 @@ bool FTopologicalLoop::Orient()
 
 	ensureCADKernel(Edges.Num() > 0);
 
-	TArray<FPoint2D> LoopSampling;
+	TArray<FVector2d> LoopSampling;
 
 	if (!Get2DSamplingWithoutDegeneratedEdges(LoopSampling))
 	{
@@ -352,8 +352,8 @@ bool FTopologicalLoop::Orient()
 
 		if(Slope > 7.9 || Slope < 0.1)
 		{
-			double SquareLengthBefore = FPoint2D::DistSquared(LoopSampling[Index], LoopSampling[PreviousIndex]);
-			double SquareLengthAfter = FPoint2D::DistSquared(LoopSampling[Index], LoopSampling[NextIndex]);
+			double SquareLengthBefore = FVector2d::DistSquared(LoopSampling[Index], LoopSampling[PreviousIndex]);
+			double SquareLengthAfter = FVector2d::DistSquared(LoopSampling[Index], LoopSampling[NextIndex]);
 			if(SquareLengthBefore < SquareLengthAfter)
 			{
 				Index = PreviousIndex;
@@ -706,21 +706,21 @@ void FTopologicalLoop::CheckEdgesOrientation()
 	{
 		int32 PIndex = 0;
 		int32 EIndex = 0;
-		double SmallDistance = FPoint2D::Distance(PreviousExtremities[0].Point2D, EdgeExtremities[0].Point2D);
-		double Distance = FPoint2D::Distance(PreviousExtremities[0].Point2D, EdgeExtremities[1].Point2D);
+		double SmallDistance = FVector2d::Distance(PreviousExtremities[0].Point2D, EdgeExtremities[0].Point2D);
+		double Distance = FVector2d::Distance(PreviousExtremities[0].Point2D, EdgeExtremities[1].Point2D);
 		if (Distance < SmallDistance)
 		{
 			SmallDistance = Distance;
 			EIndex = 1;
 		}
-		Distance = FPoint2D::Distance(PreviousExtremities[1].Point2D, EdgeExtremities[0].Point2D);
+		Distance = FVector2d::Distance(PreviousExtremities[1].Point2D, EdgeExtremities[0].Point2D);
 		if (Distance < SmallDistance)
 		{
 			SmallDistance = Distance;
 			PIndex = 1;
 			EIndex = 0;
 		}
-		Distance = FPoint2D::Distance(PreviousExtremities[1].Point2D, EdgeExtremities[1].Point2D);
+		Distance = FVector2d::Distance(PreviousExtremities[1].Point2D, EdgeExtremities[1].Point2D);
 		if (Distance < SmallDistance)
 		{
 			SmallDistance = Distance;
@@ -767,21 +767,21 @@ void FTopologicalLoop::CheckLoopWithTwoEdgesOrientation()
 
 	int32 IndexEdge0 = 0;
 	int32 IndexEdge1 = 0;
-	double SmallDistance = FPoint2D::Distance(Edge0Extremities[0].Point2D, Edge1Extremities[0].Point2D);
-	double Distance = FPoint2D::Distance(Edge0Extremities[0].Point2D, Edge1Extremities[1].Point2D);
+	double SmallDistance = FVector2d::Distance(Edge0Extremities[0].Point2D, Edge1Extremities[0].Point2D);
+	double Distance = FVector2d::Distance(Edge0Extremities[0].Point2D, Edge1Extremities[1].Point2D);
 	if (Distance < SmallDistance)
 	{
 		SmallDistance = Distance;
 		IndexEdge1 = 1;
 	}
-	Distance = FPoint2D::Distance(Edge0Extremities[1].Point2D, Edge1Extremities[0].Point2D);
+	Distance = FVector2d::Distance(Edge0Extremities[1].Point2D, Edge1Extremities[0].Point2D);
 	if (Distance < SmallDistance)
 	{
 		SmallDistance = Distance;
 		IndexEdge0 = 1;
 		IndexEdge1 = 0;
 	}
-	Distance = FPoint2D::Distance(Edge0Extremities[1].Point2D, Edge1Extremities[1].Point2D);
+	Distance = FVector2d::Distance(Edge0Extremities[1].Point2D, Edge1Extremities[1].Point2D);
 	if (Distance < SmallDistance)
 	{
 		SmallDistance = Distance;
@@ -827,14 +827,14 @@ void FTopologicalLoop::RemoveDegeneratedEdges()
 
 		TSharedRef<FTopologicalVertex> DegeneratedEdgeVertex = (bPrevious == (DegeneratedOrientedEdge.Direction == EOrientation::Front)) ? DegeneratedEdge->GetStartVertex() : DegeneratedEdge->GetEndVertex();
 		TSharedRef<FTopologicalVertex> OtherDegeneratedEdgeVertex = (bPrevious == (DegeneratedOrientedEdge.Direction == EOrientation::Front)) ? DegeneratedEdge->GetEndVertex() : DegeneratedEdge->GetStartVertex();
-		FPoint2D DegeneratedEdgeTangent = DegeneratedEdge->GetTangent2DAt(*DegeneratedEdgeVertex);
+		FVector2d DegeneratedEdgeTangent = DegeneratedEdge->GetTangent2DAt(*DegeneratedEdgeVertex);
 
 		TSharedRef<FTopologicalVertex> NearEdgeVertex = (bPrevious == (NearOrientedEdge.Direction == EOrientation::Front)) ? NearOrientedEdge.Entity->GetEndVertex() : NearOrientedEdge.Entity->GetStartVertex();
-		FPoint2D NearEdgeTangent = NearOrientedEdge.Entity->GetTangent2DAt(*NearEdgeVertex);
+		FVector2d NearEdgeTangent = NearOrientedEdge.Entity->GetTangent2DAt(*NearEdgeVertex);
 
-		FPoint2D& NearEdgeExtremity = (bPrevious == (NearOrientedEdge.Direction == EOrientation::Front)) ? NearEdgeExtremities[1].Point2D : NearEdgeExtremities[0].Point2D;
-		FPoint2D& DegeneratedEdgeExtremity = (bPrevious == (DegeneratedOrientedEdge.Direction == EOrientation::Front)) ? DegeneratedEdgeExtremities[0].Point2D : DegeneratedEdgeExtremities[1].Point2D;
-		FPoint2D& OtherDegeneratedEdgeExtremity = (bPrevious == (DegeneratedOrientedEdge.Direction == EOrientation::Front)) ? DegeneratedEdgeExtremities[1].Point2D : DegeneratedEdgeExtremities[0].Point2D;
+		FVector2d& NearEdgeExtremity = (bPrevious == (NearOrientedEdge.Direction == EOrientation::Front)) ? NearEdgeExtremities[1].Point2D : NearEdgeExtremities[0].Point2D;
+		FVector2d& DegeneratedEdgeExtremity = (bPrevious == (DegeneratedOrientedEdge.Direction == EOrientation::Front)) ? DegeneratedEdgeExtremities[0].Point2D : DegeneratedEdgeExtremities[1].Point2D;
+		FVector2d& OtherDegeneratedEdgeExtremity = (bPrevious == (DegeneratedOrientedEdge.Direction == EOrientation::Front)) ? DegeneratedEdgeExtremities[1].Point2D : DegeneratedEdgeExtremities[0].Point2D;
 
 		double Slope = ComputeUnorientedSlope(FVector2d::ZeroVector, DegeneratedEdgeTangent, NearEdgeTangent);
 		constexpr double FlatSlope = 3.9; // ~175 deg. 
@@ -909,14 +909,14 @@ void FTopologicalLoop::EnsureLogicalClosing(const double Tolerance3D)
 		TSharedRef<FTopologicalVertex> PreviousEdgeEndVertex = PreviousEdge.Direction == EOrientation::Front ? PreviousEdge.Entity->GetEndVertex() : PreviousEdge.Entity->GetStartVertex();
 		TSharedRef<FTopologicalVertex> EdgeStartVertex = OrientedEdge.Direction == EOrientation::Front ? OrientedEdge.Entity->GetStartVertex() : OrientedEdge.Entity->GetEndVertex();
 
-		const FPoint2D Gap2D = EdgeExtremities[ExtremityIndex].Point2D - PreviousExtremity.Point2D;
+		const FVector2d Gap2D = EdgeExtremities[ExtremityIndex].Point2D - PreviousExtremity.Point2D;
 		if (SquareGap3D > SquareTolerance3D)
 		{
 			FMessage::Printf(Log, TEXT("Loop %d Gap 3D : %f\n"), Id, sqrt(SquareGap3D));
-			const FPoint PreviousTangent = PreviousEdge.Entity->GetTangentAt(*PreviousEdgeEndVertex);
-			const FPoint EdgeTangent = OrientedEdge.Entity->GetTangentAt(*EdgeStartVertex);
+			const FVector PreviousTangent = PreviousEdge.Entity->GetTangentAt(*PreviousEdgeEndVertex);
+			const FVector EdgeTangent = OrientedEdge.Entity->GetTangentAt(*EdgeStartVertex);
 
-			FPoint Gap = EdgeExtremities[ExtremityIndex].Point - PreviousExtremity.Point;
+			FVector Gap = EdgeExtremities[ExtremityIndex].Point - PreviousExtremity.Point;
 
 			double CosAngle = FVectorUtil::ComputeCosinus(Gap, EdgeTangent);
 			constexpr double CosTangentLimit = 0.9; // ~25 deg : 25 deg is not big angle between the extremity curve tangent and the missing segment to close the gap.  
@@ -970,8 +970,8 @@ void FTopologicalLoop::EnsureLogicalClosing(const double Tolerance3D)
 			// if the gap is not so large and the extremity of on edge is tangent to the gap segment then we extend the edge to close the gap
 			FMessage::Printf(Log, TEXT("Loop %d Gap 2D : [%f, %f] vs Tol2D [%f, %f]\n"), Id, FMath::Abs(Gap2D.X), FMath::Abs(Gap2D.Y), Tolerance2D.X, Tolerance2D.Y);
 
-			const FPoint2D PreviousTangent = PreviousEdge.Entity->GetTangent2DAt(*PreviousEdgeEndVertex);
-			const FPoint2D EdgeTangent = OrientedEdge.Entity->GetTangent2DAt(*EdgeStartVertex);
+			const FVector2d PreviousTangent = PreviousEdge.Entity->GetTangent2DAt(*PreviousEdgeEndVertex);
+			const FVector2d EdgeTangent = OrientedEdge.Entity->GetTangent2DAt(*EdgeStartVertex);
 
 			double CosAngle = FVectorUtil::ComputeCosinus(Gap2D, EdgeTangent);
 			constexpr double CosFlatTangent = 0.98;  // ~10 deg: In 2D, the distortion due to the parametric space, degenerated area and else imposes to be more careful before extending a curve.
@@ -1043,14 +1043,14 @@ void FTopologicalLoop::FindBreaks(TArray<TSharedPtr<FTopologicalVertex>>& OutBre
 	OutBreakValues.Empty(EdgeNum);
 	OutStartSideIndex.Empty(EdgeNum);
 
-	FPoint StartTangentEdge;
-	FPoint EndTangentPreviousEdge;
+	FVector StartTangentEdge;
+	FVector EndTangentPreviousEdge;
 	Edges[EdgeNum - 1].Entity->GetTangentsAtExtremities(StartTangentEdge, EndTangentPreviousEdge, Edges[EdgeNum - 1].Direction == EOrientation::Front);
 	bool bPreviousIsSurface = (Edges[EdgeNum - 1].Entity->GetTwinEntityCount() > 1);
 
 	for (int32 Index = 0; Index < EdgeNum; Index++)
 	{
-		FPoint EndTangentEdge;
+		FVector EndTangentEdge;
 		Edges[Index].Entity->GetTangentsAtExtremities(StartTangentEdge, EndTangentEdge, Edges[Index].Direction == EOrientation::Front);
 		bool bIsSurface = (Edges[Index].Entity->GetTwinEntityCount() > 1);
 
@@ -1061,7 +1061,7 @@ void FTopologicalLoop::FindBreaks(TArray<TSharedPtr<FTopologicalVertex>>& OutBre
 
 #ifdef FIND_BREAKS
 			{
-				FPoint& Start = BoundaryEdges[Index]->GetStartVertex(BoundaryEdgeDirections[Index])->GetCoordinates();
+				FVector& Start = BoundaryEdges[Index]->GetStartVertex(BoundaryEdgeDirections[Index])->GetCoordinates();
 				Open3DDebugSession(TEXT("Cos Angle " + Utils::ToString(CosAngle));
 				DisplayPoint(Start, (CosAngle > MinCosAngleOfBreak) ? EVisuProperty::RedPoint : EVisuProperty::BluePoint);
 				DisplaySegment(Start, Start + StartTangentEdge);
@@ -1085,14 +1085,14 @@ void FTopologicalLoop::FindBreaks(TArray<TSharedPtr<FTopologicalVertex>>& OutBre
 
 namespace TopologicalLoopImpl
 {
-void FindLoopIntersectionsWithIso(const EIso Iso, const double IsoParameter, const TArray<FPoint2D>& Loop, TArray<double>& OutIntersections)
+void FindLoopIntersectionsWithIso(const EIso Iso, const double IsoParameter, const TArray<FVector2d>& Loop, TArray<double>& OutIntersections)
 {
 	OutIntersections.Empty(8);
 
 	int32 UIndex = Iso == EIso::IsoU ? 0 : 1;
 	int32 VIndex = Iso == EIso::IsoU ? 1 : 0;
 
-	TFunction<void(const FPoint2D&, const FPoint2D&)> ComputeIntersection = [&](const FPoint2D& Point1, const FPoint2D& Point2)
+	TFunction<void(const FVector2d&, const FVector2d&)> ComputeIntersection = [&](const FVector2d& Point1, const FVector2d& Point2)
 	{
 		if (IsoParameter > Point1[UIndex] && IsoParameter <= Point2[UIndex])
 		{
@@ -1101,8 +1101,8 @@ void FindLoopIntersectionsWithIso(const EIso Iso, const double IsoParameter, con
 		}
 	};
 
-	const FPoint2D* Point1 = &Loop.Last();
-	for (const FPoint2D& Point2 : Loop)
+	const FVector2d* Point1 = &Loop.Last();
+	for (const FVector2d& Point2 : Loop)
 	{
 		if (!FMath::IsNearlyEqual((*Point1)[UIndex], Point2[UIndex]))
 		{
@@ -1124,10 +1124,10 @@ void FindLoopIntersectionsWithIso(const EIso Iso, const double IsoParameter, con
 
 bool FTopologicalLoop::IsInside(const FTopologicalLoop& OtherLoop) const
 {
-	TArray<FPoint2D> Sampling2D;
+	TArray<FVector2d> Sampling2D;
 	Get2DSampling(Sampling2D);
 
-	TArray<FPoint2D> OtherSampling2D;
+	TArray<FVector2d> OtherSampling2D;
 	OtherLoop.Get2DSampling(OtherSampling2D);
 
 #ifdef DEBUG_IS_INSIDE
@@ -1149,7 +1149,7 @@ bool FTopologicalLoop::IsInside(const FTopologicalLoop& OtherLoop) const
 	int32 InsidePoint = 0;
 	int32 OutsidePoint = 0;
 
-	TFunction<void(const FPoint2D&, EIso)> CountLeftIntersection = [&](const FPoint2D& Point, EIso Iso)
+	TFunction<void(const FVector2d&, EIso)> CountLeftIntersection = [&](const FVector2d& Point, EIso Iso)
 	{
 		TArray<double> Intersections;
 		TopologicalLoopImpl::FindLoopIntersectionsWithIso(Iso, Point[Iso], Sampling2D, Intersections);
@@ -1184,7 +1184,7 @@ bool FTopologicalLoop::IsInside(const FTopologicalLoop& OtherLoop) const
 	const int32 Increment = FMath::Max(OtherSampling2DCount / Step, 1);
 	for (int32 Index = 0; Index < OtherSampling2DCount; Index += Increment)
 	{
-		const FPoint2D& TestPoint = OtherSampling2D[Index];
+		const FVector2d& TestPoint = OtherSampling2D[Index];
 		CountLeftIntersection(TestPoint, EIso::IsoU);
 		CountLeftIntersection(TestPoint, EIso::IsoV);
 	}
diff --git a/Engine/Source/Runtime/Datasmith/CADKernel/Base/Private/Topo/TopologicalVertex.cpp b/Engine/Source/Runtime/Datasmith/CADKernel/Base/Private/Topo/TopologicalVertex.cpp
--- a/Engine/Source/Runtime/Datasmith/CADKernel/Base/Private/Topo/TopologicalVertex.cpp
+++ b/Engine/Source/Runtime/Datasmith/CADKernel/Base/Private/Topo/TopologicalVertex.cpp
@@ -96,7 +96,7 @@ void FTopologicalVertex::Link(FTopologicalVertex& Twin)
 		}
 	}
 
-	FPoint Barycenter = GetBarycenter() * (double)GetTwinEntityCount()
+	FVector Barycenter = GetBarycenter() * (double)GetTwinEntityCount()
 		+ Twin.GetBarycenter() * (double)Twin.GetTwinEntityCount();
 
 	MakeLink(Twin);
@@ -137,7 +137,7 @@ void FTopologicalVertex::UnlinkTo(FTopologicalVertex& OtherVertex)
 
 void FVertexLink::ComputeBarycenter()
 {
-	Barycenter = FPoint::ZeroVector;
+	Barycenter = FVector::ZeroVector;
 	for (const FTopologicalVertex* Vertex : TwinEntities)
 	{
 		Barycenter += Vertex->Coordinates;
diff --git a/Engine/Source/Runtime/Datasmith/CADKernel/Base/Private/Topo/Topomaker.cpp b/Engine/Source/Runtime/Datasmith/CADKernel/Base/Private/Topo/Topomaker.cpp
--- a/Engine/Source/Runtime/Datasmith/CADKernel/Base/Private/Topo/Topomaker.cpp
+++ b/Engine/Source/Runtime/Datasmith/CADKernel/Base/Private/Topo/Topomaker.cpp
@@ -93,7 +93,7 @@ void MergeCoincidentVertices(TArray<FTopologicalVertex*>& VerticesToMerge, doubl
 		Vertex->SetMarker1();
 
 		double VertexWeigth = VerticesWeight[SortedVertexIndices[IndexI]];
-		FPoint Barycenter = Vertex->GetBarycenter();
+		FVector Barycenter = Vertex->GetBarycenter();
 
 		for (int32 IndexJ = IndexI + 1; IndexJ < VertexNum; ++IndexJ)
 		{
@@ -167,7 +167,7 @@ FTopologicalVertex* SplitAndLink(FTopologicalVertex& StartVertex, FTopologicalEd
 		return &VertexToLink->GetLinkActiveEntity().Get();
 	}
 
-	FPoint ProjectedPoint;
+	FVector ProjectedPoint;
 	double UProjectedPoint = EdgeToSplit.ProjectPoint(VertexToLink->GetBarycenter(), ProjectedPoint);
 
 	double SquareDistanceToProjectedPoint = FVector::DistSquared(ProjectedPoint, VertexToLink->GetBarycenter());
@@ -335,8 +335,8 @@ FTopologicalVertex* StitchParallelEdgesFrom(FTopologicalVertex* Vertex, double S
 				}
 			}
 
-			FPoint StartTangentEdge = Edge->GetTangentAt(*Vertex);
-			FPoint StartTangentOtherEdge = SecondEdge->GetTangentAt(*Vertex);
+			FVector StartTangentEdge = Edge->GetTangentAt(*Vertex);
+			FVector StartTangentOtherEdge = SecondEdge->GetTangentAt(*Vertex);
 
 			double CosAngle = FVectorUtil::ComputeCosinus(StartTangentEdge, StartTangentOtherEdge);
 			if (CosAngle < UE_DOUBLE_HALF_SQRT_3) // cos(30 deg)
@@ -828,8 +828,8 @@ void FTopomaker::MergeUnconnectedSuccessiveEdges()
 					if (ConnectedEdges.Num() == 2)
 					{
 						// check if the edges are tangents
-						FPoint StartTangentEdge = ConnectedEdges[0]->GetTangentAt(*EndVertex);
-						FPoint StartTangentOtherEdge = ConnectedEdges[1]->GetTangentAt(*EndVertex);
+						FVector StartTangentEdge = ConnectedEdges[0]->GetTangentAt(*EndVertex);
+						FVector StartTangentOtherEdge = ConnectedEdges[1]->GetTangentAt(*EndVertex);
 
 						double CosAngle = FVectorUtil::ComputeCosinus(StartTangentEdge, StartTangentOtherEdge);
 						if (CosAngle < -UE_DOUBLE_HALF_SQRT_3) // Cos(30 deg)
