Replaced FPoint and FPoint2D by FVector and  FVector2d and remove unused FFPoint and FPointH classes - Final

diff --git a/Engine/Source/Runtime/Datasmith/CADKernel/Base/Private/UI/Display.cpp b/Engine/Source/Runtime/Datasmith/CADKernel/Base/Private/UI/Display.cpp
--- a/Engine/Source/Runtime/Datasmith/CADKernel/Base/Private/UI/Display.cpp
+++ b/Engine/Source/Runtime/Datasmith/CADKernel/Base/Private/UI/Display.cpp
@@ -78,7 +78,7 @@ void FlushVisu()
 #endif
 }
 
-void DrawElement(int32 Dimension, TArray<FPoint>& Points, EVisuProperty Property)
+void DrawElement(int32 Dimension, TArray<FVector>& Points, EVisuProperty Property)
 {
 #ifdef CADKERNEL_DEBUG
 	FSystem::Get().GetVisu()->DrawElement(Dimension, Points, Property);
@@ -102,8 +102,8 @@ void Draw(const FLinearBoundary& Boundary, const FRestrictionCurve& Curve, EVisu
 
 	bool bShowOrientation = FSystem::Get().GetVisu()->GetParameters()->bDisplayCADOrient;
 
-	TArray<FPoint> Polyline;
-	Curve.GetDiscretizationPoints<FPoint>(Boundary, EOrientation::Front, Polyline);
+	TArray<FVector> Polyline;
+	Curve.GetDiscretizationPoints<FVector>(Boundary, EOrientation::Front, Polyline);
 	Draw(Polyline, Property);
 
 	if (bShowOrientation)
@@ -215,31 +215,31 @@ void Draw(const FCurve& Curve, EVisuProperty Property)
 #endif
 }
 
-void DrawQuadripode(double Height, double Base, FPoint& Center, FPoint& InDirection, EVisuProperty Property)
+void DrawQuadripode(double Height, double Base, FVector& Center, FVector& InDirection, EVisuProperty Property)
 {
 #ifdef CADKERNEL_DEBUG
-	FPoint Direction = InDirection;
+	FVector Direction = InDirection;
 	Direction.Normalize();
 
-	FPoint Normal(-Direction[1], Direction[0], 0.0);
+	FVector Normal(-Direction[1], Direction[0], 0.0);
 
 	double UVNorm = Normal.Length();
 	Normal /= UVNorm;
 
-	FPoint BiNormal(-Direction[2] * Normal[1], Direction[2] * Normal[0], UVNorm);
+	FVector BiNormal(-Direction[2] * Normal[1], Direction[2] * Normal[0], UVNorm);
 
-	FPoint Point0 = Center;
-	FPoint PointBase = Point0 - Direction * Height;
+	FVector Point0 = Center;
+	FVector PointBase = Point0 - Direction * Height;
 
 	Normal *= Base;
 	BiNormal *= Base;
 
-	FPoint Point1 = PointBase + Normal;
-	FPoint Point2 = PointBase + BiNormal;
-	FPoint Point3 = PointBase - Normal;
-	FPoint Point4 = PointBase - BiNormal;
+	FVector Point1 = PointBase + Normal;
+	FVector Point2 = PointBase + BiNormal;
+	FVector Point3 = PointBase - Normal;
+	FVector Point4 = PointBase - BiNormal;
 
-	TArray<FPoint> Polygone;
+	TArray<FVector> Polygone;
 	Polygone.Reserve(7);
 
 	Polygone.Add(Point1);
@@ -375,7 +375,7 @@ void Display(const FPlane& Plane, FIdent Ident)
 	VAxis.Normalize();
 	VAxis *= 10;
 
-	TArray<FPoint> Points;
+	TArray<FVector> Points;
 
 	const FVector Point = Plane.GetOrigin();
 	
@@ -401,7 +401,7 @@ void Display(const FPlane& Plane, FIdent Ident)
 #endif
 }
 
-void DisplayEdgeCriteriaGrid(int32 EdgeId, const TArray<FPoint>& Points3D)
+void DisplayEdgeCriteriaGrid(int32 EdgeId, const TArray<FVector>& Points3D)
 {
 #ifdef CADKERNEL_DEBUG
 	FString Name = FString::Printf(TEXT("Edge Grid %d"), EdgeId);
@@ -426,7 +426,7 @@ void DisplayEdgeCriteriaGrid(int32 EdgeId, const TArray<FPoint>& Points3D)
 void DisplayAABB(const FAABB& Aabb, FIdent Ident)
 {
 #ifdef CADKERNEL_DEBUG
-	TFunction<void(int32, FPoint&, FPoint&)> GetWire = [&](int32 Segment, FPoint& A, FPoint& B)
+	TFunction<void(int32, FVector&, FVector&)> GetWire = [&](int32 Segment, FVector& A, FVector& B)
 	{
 		int32 Corner1;
 		int32 Corner2;
@@ -466,7 +466,7 @@ void DisplayAABB(const FAABB& Aabb, FIdent Ident)
 	F3DDebugSegment G(Ident);
 	for (int32 SegmentIndex = 0; SegmentIndex < 12; SegmentIndex++)
 	{
-		TArray<FPoint> Points;
+		TArray<FVector> Points;
 		Points.SetNum(2);
 		GetWire(SegmentIndex, Points[0], Points[1]);
 		Draw(Points);
@@ -477,12 +477,12 @@ void DisplayAABB(const FAABB& Aabb, FIdent Ident)
 void DisplayAABB2D(const FAABB2D& aabb, FIdent Ident)
 {
 #ifdef CADKERNEL_DEBUG
-	FPoint A(aabb.GetMin().X, aabb.GetMin().Y, 0.);
-	FPoint B(aabb.GetMin().X, aabb.GetMax().Y, 0.);
-	FPoint C(aabb.GetMax().X, aabb.GetMax().Y, 0.);
-	FPoint D(aabb.GetMax().X, aabb.GetMin().Y, 0.);
+	FVector A(aabb.GetMin().X, aabb.GetMin().Y, 0.);
+	FVector B(aabb.GetMin().X, aabb.GetMax().Y, 0.);
+	FVector C(aabb.GetMax().X, aabb.GetMax().Y, 0.);
+	FVector D(aabb.GetMax().X, aabb.GetMin().Y, 0.);
 
-	TArray<FPoint> Points;
+	TArray<FVector> Points;
 	Points.Add(A);
 	Points.Add(B);
 	Points.Add(C);
@@ -612,15 +612,15 @@ void Display2D(const FSurface& Surface)
 
 	for (int32 iIso = 0; iIso <= (IsoUCount + 1); iIso++)
 	{
-		FPoint2D StartPoint(Bounds[EIso::IsoU].Min + (iIso)*StepU, Bounds[EIso::IsoV].Min);
-		FPoint2D EndPoint(Bounds[EIso::IsoU].Min + (iIso)*StepU, Bounds[EIso::IsoV].Max);
+		FVector2d StartPoint(Bounds[EIso::IsoU].Min + (iIso)*StepU, Bounds[EIso::IsoV].Min);
+		FVector2d EndPoint(Bounds[EIso::IsoU].Min + (iIso)*StepU, Bounds[EIso::IsoV].Max);
 		DrawSegment(StartPoint, EndPoint, EVisuProperty::Iso);
 	}
 
 	for (int32 iIso = 0; iIso <= (IsoVCount + 1); iIso++)
 	{
-		FPoint2D StartPoint(Bounds[EIso::IsoU].Min, Bounds[EIso::IsoV].Min + (iIso)*StepV);
-		FPoint2D EndPoint(Bounds[EIso::IsoU].Max, Bounds[EIso::IsoV].Min + (iIso)*StepV);
+		FVector2d StartPoint(Bounds[EIso::IsoU].Min, Bounds[EIso::IsoV].Min + (iIso)*StepV);
+		FVector2d EndPoint(Bounds[EIso::IsoU].Max, Bounds[EIso::IsoV].Min + (iIso)*StepV);
 		DrawSegment(StartPoint, EndPoint, EVisuProperty::Iso);
 	}
 #endif
@@ -647,16 +647,16 @@ void DisplayIsoCurve(const FSurface& Surface, double Coordinate, EIso IsoType)
 void DisplayControlPolygon(const FCurve& Curve)
 {
 #ifdef CADKERNEL_DEBUG
-	TFunction<void(const TArray<FPoint>&)> DisplayHull = [](const TArray<FPoint>& Poles)
+	TFunction<void(const TArray<FVector>&)> DisplayHull = [](const TArray<FVector>& Poles)
 	{
-		for (const FPoint& Pole : Poles)
+		for (const FVector& Pole : Poles)
 		{
 			DisplayPoint(Pole, EVisuProperty::GreenPoint);
 		}
 
 		for (int32 Index = 1; Index < Poles.Num(); Index++)
 		{
-			FPoint Segment = Poles[Index] - Poles[Index - 1];
+			FVector Segment = Poles[Index] - Poles[Index - 1];
 			DisplaySegment(Poles[Index - 1] + Segment * 0.1, Poles[Index] - Segment * 0.1, 0, EVisuProperty::GreenCurve);
 		}
 	};
@@ -673,7 +673,7 @@ void DisplayControlPolygon(const FCurve& Curve)
 	case ECurve::Nurbs:
 	{
 		const FNURBSCurve& Nurbs = (const FNURBSCurve&)Curve;
-		const TArray<FPoint>& Poles = Nurbs.GetPoles();
+		const TArray<FVector>& Poles = Nurbs.GetPoles();
 		DisplayHull(Poles);
 		return;
 	}
@@ -705,7 +705,7 @@ void DisplayControlPolygon(const FSurface& Surface)
 #ifdef CADKERNEL_DEBUG
 	bool bShowOrientation = FSystem::Get().GetVisu()->GetParameters()->bDisplayCADOrient;
 
-	TFunction<void(const TArray<FPoint>&, int32, int32)> DisplayHull = [&bShowOrientation](const TArray<FPoint>& Poles, int32 PoleUNum, int32 PoleVNum)
+	TFunction<void(const TArray<FVector>&, int32, int32)> DisplayHull = [&bShowOrientation](const TArray<FVector>& Poles, int32 PoleUNum, int32 PoleVNum)
 	{
 		for (int32 Index = 0; Index < Poles.Num(); Index++)
 		{
@@ -717,7 +717,7 @@ void DisplayControlPolygon(const FSurface& Surface)
 			Index = Index + 1;
 			for (int32 IndexU = 1; IndexU < PoleUNum; IndexU++, Index++)
 			{
-				FPoint Segment = Poles[Index] - Poles[Index - 1];
+				FVector Segment = Poles[Index] - Poles[Index - 1];
 				DisplaySegment(Poles[Index - 1] + Segment * 0.1, Poles[Index] - Segment * 0.1, 0, EVisuProperty::YellowCurve);
 				if (bShowOrientation)
 				{
@@ -730,7 +730,7 @@ void DisplayControlPolygon(const FSurface& Surface)
 		{
 			for (int32 IndexV = 1, Index = PoleUNum + IndexU; IndexV < PoleVNum; IndexV++, Index += PoleUNum)
 			{
-				FPoint Segment = Poles[Index] - Poles[Index - PoleUNum];
+				FVector Segment = Poles[Index] - Poles[Index - PoleUNum];
 				DisplaySegment(Poles[Index - PoleUNum] + Segment * 0.1, Poles[Index] - Segment * 0.1, 0, EVisuProperty::GreenCurve);
 				if (bShowOrientation)
 				{
@@ -812,10 +812,10 @@ void Draw(const FTopologicalFace& Face)
 void Draw2D(const FTopologicalFace& Face)
 {
 #ifdef CADKERNEL_DEBUG
-	TArray<TArray<FPoint2D>> BoundaryApproximation;
+	TArray<TArray<FVector2d>> BoundaryApproximation;
 	Face.Get2DLoopSampling(BoundaryApproximation);
 
-	for (const TArray<FPoint2D>& Boudary : BoundaryApproximation)
+	for (const TArray<FVector2d>& Boudary : BoundaryApproximation)
 	{
 		Draw(Boudary, EVisuProperty::BlueCurve);
 	}
@@ -838,8 +838,8 @@ void Draw2D(const FTopologicalFace& Face)
 			TArray<double> Intersections;
 			FindLoopIntersectionsWithIso(IsoType, Coordinate, BoundaryApproximation, Intersections);
 
-			FPoint2D Start;
-			FPoint2D End;
+			FVector2d Start;
+			FVector2d End;
 
 			Start[IsoType] = Coordinate;
 			End[IsoType] = Coordinate;
@@ -872,7 +872,7 @@ void Draw2D(const FTopologicalFace& Face)
 void DrawIsoCurves(const FTopologicalFace& Face)
 {
 #ifdef CADKERNEL_DEBUG
-	TArray<TArray<FPoint2D>> BoundaryApproximation;
+	TArray<TArray<FVector2d>> BoundaryApproximation;
 	Face.Get2DLoopSampling(BoundaryApproximation);
 
 	const double VisuSag = FSystem::Get().GetVisu()->GetParameters()->ChordError;
@@ -911,7 +911,7 @@ void DrawIsoCurves(const FTopologicalFace& Face)
 
 			if (IntersectionCount % 2 != 0)
 			{
-				TArray<FPoint> SubPolyline;
+				TArray<FVector> SubPolyline;
 				FLinearBoundary Boundary(Intersections[IntersectionCount - 1], CurveBounds.GetMax());
 				Polyline.GetSubPolyline(Boundary, EOrientation::Front, SubPolyline);
 				Draw(SubPolyline, EVisuProperty::YellowCurve);
@@ -921,7 +921,7 @@ void DrawIsoCurves(const FTopologicalFace& Face)
 
 			for (int32 ISection = 0; ISection < IntersectionCount; ISection += 2)
 			{
-				TArray<FPoint> SubPolyline;
+				TArray<FVector> SubPolyline;
 				FLinearBoundary Boundary(Intersections[ISection], Intersections[ISection + 1]);
 				Polyline.GetSubPolyline(Boundary, EOrientation::Front, SubPolyline);
 				Draw(SubPolyline, EVisuProperty::Iso);
@@ -949,7 +949,7 @@ void Display2D(const FTopologicalEdge& Edge, EVisuProperty Property)
 {
 #ifdef CADKERNEL_DEBUG
 	F3DDebugSegment GraphicSegment(Edge.GetId());
-	TArray<FPoint2D> Polyline;
+	TArray<FVector2d> Polyline;
 	Edge.GetCurve()->GetDiscretizationPoints(Edge.GetBoundary(), EOrientation::Front, Polyline);
 	DisplayPolyline(Polyline, Property);
 #endif
@@ -959,7 +959,7 @@ void Display2DWithScale(const FTopologicalEdge& Edge, EVisuProperty Property)
 {
 #ifdef CADKERNEL_DEBUG
 	F3DDebugSegment GraphicSegment(Edge.GetId());
-	TArray<FPoint2D> Polyline;
+	TArray<FVector2d> Polyline;
 	Edge.GetCurve()->GetDiscretizationPoints(Edge.GetBoundary(), EOrientation::Front, Polyline);
 	DisplayPolylineWithScale(Polyline, Property);
 #endif
@@ -1127,19 +1127,19 @@ void DisplayMesh(const FShell& Shell)
 void DisplayMesh(const FFaceMesh& Mesh)
 {
 #ifdef CADKERNEL_DEBUG
-	TMap<int32, const FPoint*> NodeIdToCoordinates;
+	TMap<int32, const FVector*> NodeIdToCoordinates;
 	Mesh.GetNodeIdToCoordinates(NodeIdToCoordinates);
 
-	TArray<FPoint> Points;
+	TArray<FVector> Points;
 	Points.SetNum(3);
 	double EdgeMeanLength = 0;
 	const TArray<int32>& TriangleIndices = Mesh.TrianglesVerticesIndex;
 	const TArray<int32>& VertexIndices = Mesh.VerticesGlobalIndex;
 	for (int32 Index = 0; Index < TriangleIndices.Num();)
 	{
-		const FPoint** P0 = NodeIdToCoordinates.Find(VertexIndices[TriangleIndices[Index++]]);
-		const FPoint** P1 = NodeIdToCoordinates.Find(VertexIndices[TriangleIndices[Index++]]);
-		const FPoint** P2 = NodeIdToCoordinates.Find(VertexIndices[TriangleIndices[Index++]]);
+		const FVector** P0 = NodeIdToCoordinates.Find(VertexIndices[TriangleIndices[Index++]]);
+		const FVector** P1 = NodeIdToCoordinates.Find(VertexIndices[TriangleIndices[Index++]]);
+		const FVector** P2 = NodeIdToCoordinates.Find(VertexIndices[TriangleIndices[Index++]]);
 
 		if (!P0 || !P1 || !P2)
 		{
@@ -1161,7 +1161,7 @@ void DisplayMesh(const FFaceMesh& Mesh)
 
 	for (const int32& Index : VertexIndices)
 	{
-		const FPoint** PointPtr = NodeIdToCoordinates.Find(Index);
+		const FVector** PointPtr = NodeIdToCoordinates.Find(Index);
 		if (!PointPtr)
 		{
 			continue;
@@ -1177,7 +1177,7 @@ void DisplayMesh(const FFaceMesh& Mesh)
 		const TArray<FVector3f>& Normals = Mesh.Normals;
 		for (int32 Index = 0; Index < VertexIndices.Num(); ++Index)
 		{
-			const FPoint** PointPtr = NodeIdToCoordinates.Find(VertexIndices[Index]);
+			const FVector** PointPtr = NodeIdToCoordinates.Find(VertexIndices[Index]);
 			if (!PointPtr)
 			{
 				continue;
@@ -1199,7 +1199,7 @@ void DisplayMesh(const FEdgeMesh& Mesh)
 	const FModelMesh& MeshModel = Mesh.GetMeshModel();
 
 	const TArray<int32>& NodeIds = Mesh.EdgeVerticesIndex;
-	const TArray<FPoint>& NodeCoordinates = Mesh.GetNodeCoordinates();
+	const TArray<FVector>& NodeCoordinates = Mesh.GetNodeCoordinates();
 
 	int32 StartNodeId = NodeIds[0];
 	int32 LastNodeId = NodeIds.Last();
@@ -1265,7 +1265,7 @@ void Display(const FModelMesh& MeshModel)
 #endif
 }
 
-void DisplaySegment(const FPoint& Point1, const FPoint& Point2, FIdent Ident, EVisuProperty Property)
+void DisplaySegment(const FVector& Point1, const FVector& Point2, FIdent Ident, EVisuProperty Property)
 {
 #ifdef CADKERNEL_DEBUG
 	F3DDebugSegment G(Ident);
@@ -1273,7 +1273,7 @@ void DisplaySegment(const FPoint& Point1, const FPoint& Point2, FIdent Ident, EV
 #endif
 }
 
-void DisplaySegment(const FPoint2D& Point1, const FPoint2D& Point2, FIdent Ident, EVisuProperty Property, bool bWithOrientation)
+void DisplaySegment(const FVector2d& Point1, const FVector2d& Point2, FIdent Ident, EVisuProperty Property, bool bWithOrientation)
 {
 #ifdef CADKERNEL_DEBUG
 	F3DDebugSegment G(Ident);
diff --git a/Engine/Source/Runtime/Datasmith/CADKernel/Base/Public/Geo/Curves/Curve.h b/Engine/Source/Runtime/Datasmith/CADKernel/Base/Public/Geo/Curves/Curve.h
--- a/Engine/Source/Runtime/Datasmith/CADKernel/Base/Public/Geo/Curves/Curve.h
+++ b/Engine/Source/Runtime/Datasmith/CADKernel/Base/Public/Geo/Curves/Curve.h
@@ -45,10 +45,10 @@ protected:
 public:
 
 	static TSharedPtr<FCurve> MakeNurbsCurve(FNurbsCurveData& InNurbsData);
-	static TSharedPtr<FCurve> MakeBezierCurve(const TArray<FPoint>& InPoles);
-	static TSharedPtr<FCurve> MakeSplineCurve(const TArray<FPoint>& InPoles);
-	static TSharedPtr<FCurve> MakeSplineCurve(const TArray<FPoint>& InPoles, const TArray<FPoint>& Tangents);
-	static TSharedPtr<FCurve> MakeSplineCurve(const TArray<FPoint>& InPoles, const TArray<FPoint>& ArriveTangents, const TArray<FPoint>& LeaveTangents);
+	static TSharedPtr<FCurve> MakeBezierCurve(const TArray<FVector>& InPoles);
+	static TSharedPtr<FCurve> MakeSplineCurve(const TArray<FVector>& InPoles);
+	static TSharedPtr<FCurve> MakeSplineCurve(const TArray<FVector>& InPoles, const TArray<FVector>& Tangents);
+	static TSharedPtr<FCurve> MakeSplineCurve(const TArray<FVector>& InPoles, const TArray<FVector>& ArriveTangents, const TArray<FVector>& LeaveTangents);
 
 	// TODO
 	//static TSharedPtr<FCurve> MakeConeCurve(const double InToleranceGeometric, const FMatrixH& InMatrix, double InStartRadius, double InConeAngle, const FSurfacicBoundary& InBoundary);
@@ -107,7 +107,7 @@ public:
 	}
 
 	virtual TSharedPtr<FEntityGeom> ApplyMatrix(const FMatrixH& InMatrix) const override = 0;
-	virtual void Offset(const FPoint& OffsetDirection) = 0;
+	virtual void Offset(const FVector& OffsetDirection) = 0;
 
 	virtual double GetLength(double Tolerance) const;
 
@@ -124,7 +124,7 @@ public:
 	 * Evaluate exact 3D point of the curve at the input Coordinate
 	 * The function can only be used with 3D curve (Dimension == 3)
 	 */
-	virtual FPoint EvaluatePoint(double Coordinate) const
+	virtual FVector EvaluatePoint(double Coordinate) const
 	{
 		FCurvePoint OutPoint;
 		EvaluatePoint(Coordinate, OutPoint);
@@ -144,7 +144,7 @@ public:
 	 * Evaluate exact 3D point of the curve at the input Coordinate
 	 * The function can only be used with 3D curve (Dimension == 3)
 	 */
-	virtual FPoint2D Evaluate2DPoint(double Coordinate) const
+	virtual FVector2d Evaluate2DPoint(double Coordinate) const
 	{
 		FCurvePoint2D OutPoint;
 		Evaluate2DPoint(Coordinate, OutPoint);
@@ -155,7 +155,7 @@ public:
 	 * Evaluate exact 2D point of the curve at the input Coordinate
 	 * The function can only be used with 2D curve (Dimension == 2)
 	 */
-	virtual void Evaluate2DPoint(double Coordinate, FPoint2D& OutPoint) const
+	virtual void Evaluate2DPoint(double Coordinate, FVector2d& OutPoint) const
 	{
 		ensureCADKernel(Dimension == 2);
 		FCurvePoint2D CurvePoint;
@@ -173,13 +173,13 @@ public:
 	 * Evaluate exact 3D points of the curve at the input Coordinates
 	 * The function can only be used with 3D curve (Dimension == 3)
 	 */
-	virtual void EvaluatePoints(const TArray<double>& Coordinates, TArray<FPoint>& OutPoints) const;
+	virtual void EvaluatePoints(const TArray<double>& Coordinates, TArray<FVector>& OutPoints) const;
 
 	/**
 	 * Evaluate exact 2D points of the curve at the input Coordinates
 	 * The function can only be used with 2D curve (Dimension == 2)
 	 */
-	virtual void Evaluate2DPoints(const TArray<double>& Coordinates, TArray<FPoint2D>& OutPoints) const;
+	virtual void Evaluate2DPoints(const TArray<double>& Coordinates, TArray<FVector2d>& OutPoints) const;
 
 	/**
 	 * Evaluate exact 2D points of the curve at the input Coordinates
@@ -219,7 +219,7 @@ public:
 	/**
 	 * Linear deformation of the curve along the axis [Start point, End point] so that the nearest extremity is at the desired position and the other is not modified
 	 */
-	virtual void ExtendTo(const FPoint& DesiredPosition)
+	virtual void ExtendTo(const FVector& DesiredPosition)
 	{
 		ensureCADKernel(Dimension == 3);
 
@@ -230,10 +230,10 @@ public:
 	/**
 	 * Linear deformation of the curve along the axis [Start point, End point] so that the nearest extremity is at the desired position and the other is not modified
 	 */
-	virtual void ExtendTo(const FPoint2D& DesiredPosition)
+	virtual void ExtendTo(const FVector2d& DesiredPosition)
 	{
 		ensureCADKernel(Dimension == 2);
-		FPoint Point = DesiredPosition;
+		FVector Point(DesiredPosition, 0.);
 		ExtendTo(Point);
 	}
 
diff --git a/Engine/Source/Runtime/Datasmith/CADKernel/Base/Public/Geo/Curves/NURBSCurveData.h b/Engine/Source/Runtime/Datasmith/CADKernel/Base/Public/Geo/Curves/NURBSCurveData.h
--- a/Engine/Source/Runtime/Datasmith/CADKernel/Base/Public/Geo/Curves/NURBSCurveData.h
+++ b/Engine/Source/Runtime/Datasmith/CADKernel/Base/Public/Geo/Curves/NURBSCurveData.h
@@ -15,7 +15,7 @@ struct FNurbsCurveData
 	TArray<double> NodalVector;
 
 	TArray<double> Weights;
-	TArray<FPoint> Poles;
+	TArray<FVector> Poles;
 };
 } // namespace UE::CADKernel
 
diff --git a/Engine/Source/Runtime/Datasmith/CADKernel/Base/Public/Geo/Curves/RestrictionCurve.h b/Engine/Source/Runtime/Datasmith/CADKernel/Base/Public/Geo/Curves/RestrictionCurve.h
--- a/Engine/Source/Runtime/Datasmith/CADKernel/Base/Public/Geo/Curves/RestrictionCurve.h
+++ b/Engine/Source/Runtime/Datasmith/CADKernel/Base/Public/Geo/Curves/RestrictionCurve.h
@@ -94,7 +94,7 @@ public:
 	 * Fast computation of the point in the parametric space of the carrier surface.
 	 * This approximation is based on FSurfacicPolyline::Polyline
 	 */
-	FPoint2D Approximate2DPoint(double InCoordinate) const
+	FVector2d Approximate2DPoint(double InCoordinate) const
 	{
 		return Polyline.Approximate2DPoint(InCoordinate);
 	}
@@ -103,17 +103,17 @@ public:
 	 * Fast computation of the point in the parametric space of the carrier surface.
 	 * This approximation is based on FSurfacicPolyline::Polyline
 	 */
-	FPoint Approximate3DPoint(double InCoordinate) const
+	FVector Approximate3DPoint(double InCoordinate) const
 	{
 		return Polyline.Approximate3DPoint(InCoordinate);
 	}
 
-	FPoint GetTangentAt(double InCoordinate) const
+	FVector GetTangentAt(double InCoordinate) const
 	{
 		return Polyline.GetTangentAt(InCoordinate);
 	}
 
-	FPoint2D GetTangent2DAt(double InCoordinate) const
+	FVector2d GetTangent2DAt(double InCoordinate) const
 	{
 		return Polyline.GetTangent2DAt(InCoordinate);
 	}
@@ -122,7 +122,7 @@ public:
 	 * Fast computation of the point in the parametric space of the carrier surface.
 	 * This approximation is based on FSurfacicPolyline::Polyline
 	 */
-	void Approximate2DPoints(const TArray<double>& InCoordinates, TArray<FPoint2D>& OutPoints) const
+	void Approximate2DPoints(const TArray<double>& InCoordinates, TArray<FVector2d>& OutPoints) const
 	{
 		return Polyline.Approximate2DPoints(InCoordinates, OutPoints);
 	}
@@ -131,7 +131,7 @@ public:
 	 * Fast computation of the 3D point of the curve.
 	 * This approximation is based on FSurfacicPolyline::Polyline
 	 */
-	void Approximate3DPoints(const TArray<double>& InCoordinates, TArray<FPoint>& OutPoints) const
+	void Approximate3DPoints(const TArray<double>& InCoordinates, TArray<FVector>& OutPoints) const
 	{
 		return Polyline.Approximate3DPoints(InCoordinates, OutPoints);
 	}
@@ -160,7 +160,7 @@ public:
 	 * Project a set of points of a twin curve on the 3D polyline and return the coordinate of the projected point
 	 * @param ToleranceOfProjection: Max error between the both curve to stop the search of projection of a point 
 	 */
-	void ProjectTwinCurvePoints(const FLinearBoundary& InBoundary, const TArray<FPoint>& InPointsToProject, bool bSameOrientation, TArray<double>& OutProjectedPointCoords, double ToleranceOfProjection) const
+	void ProjectTwinCurvePoints(const FLinearBoundary& InBoundary, const TArray<FVector>& InPointsToProject, bool bSameOrientation, TArray<double>& OutProjectedPointCoords, double ToleranceOfProjection) const
 	{
 		Polyline.ProjectCoincidentalPolyline(InBoundary, InPointsToProject, bSameOrientation, OutProjectedPointCoords, ToleranceOfProjection);
 	}
@@ -207,7 +207,7 @@ public:
 		return Polyline.ComputeLinearToleranceAt(Tolerance, MinLinearTolerance, Index);
 	}
 
-	FPoint2D GetExtremityTolerances(const FLinearBoundary& InBoundary)
+	FVector2d GetExtremityTolerances(const FLinearBoundary& InBoundary)
 	{
 		FDichotomyFinder Finder(Polyline.GetCoordinates());
 		const int32 StartIndex = Finder.Find(InBoundary.Min);
@@ -215,7 +215,7 @@ public:
 
 		const double Tolerance = GetCarrierSurface()->Get3DTolerance();
 
-		FPoint2D ExtremityTolerance;
+		FVector2d ExtremityTolerance;
 		ExtremityTolerance[0] = Polyline.ComputeLinearToleranceAt(Tolerance, MinLinearTolerance, StartIndex);
 		ExtremityTolerance[1] = Polyline.ComputeLinearToleranceAt(Tolerance, MinLinearTolerance, EndIndex);
 		return ExtremityTolerance;
@@ -260,19 +260,19 @@ public:
 		return Polyline.GetLength(InBoundary);
 	}
 
-	void ExtendTo(const FPoint2D& Point) override;
+	void ExtendTo(const FVector2d& Point) override;
 
 	bool IsIso(EIso Iso, double ErrorTolerance = DOUBLE_SMALL_NUMBER) const
 	{
 		return Polyline.IsIso(Iso, ErrorTolerance);
 	}
 
-	void Offset2D(const FPoint2D& OffsetDirection);
+	void Offset2D(const FVector2d& OffsetDirection);
 
 	/**
 	 * must not be call
 	 */
-	virtual void Offset(const FPoint& OffsetDirection) override
+	virtual void Offset(const FVector& OffsetDirection) override
 	{
 		ensureCADKernel(false);
 	}
diff --git a/Engine/Source/Runtime/Datasmith/CADKernel/Base/Public/Geo/Curves/SurfacicCurve.h b/Engine/Source/Runtime/Datasmith/CADKernel/Base/Public/Geo/Curves/SurfacicCurve.h
--- a/Engine/Source/Runtime/Datasmith/CADKernel/Base/Public/Geo/Curves/SurfacicCurve.h
+++ b/Engine/Source/Runtime/Datasmith/CADKernel/Base/Public/Geo/Curves/SurfacicCurve.h
@@ -109,7 +109,7 @@ public:
 	/**
 	 * Must not be call
 	 */
-	virtual void Offset(const FPoint& OffsetDirection) override
+	virtual void Offset(const FVector& OffsetDirection) override
 	{
 		ensureCADKernel(false);
 	}
@@ -121,7 +121,7 @@ public:
 		Curve2D->Evaluate2DPoint(Coordinate, OutPoint, DerivativeOrder);
 	}
 
-	virtual void Evaluate2DPoint(double Coordinate, FPoint2D& OutPoint) const override
+	virtual void Evaluate2DPoint(double Coordinate, FVector2d& OutPoint) const override
 	{
 		Curve2D->Evaluate2DPoint(Coordinate, OutPoint);
 	}
@@ -130,7 +130,7 @@ public:
 
 	void EvaluateSurfacicPolyline(FSurfacicPolyline& OutPolyline) const;
 
-	virtual void Evaluate2DPoints(const TArray<double>& Coordinates, TArray<FPoint2D>& OutPoints) const override
+	virtual void Evaluate2DPoints(const TArray<double>& Coordinates, TArray<FVector2d>& OutPoints) const override
 	{
 		Curve2D->Evaluate2DPoints(Coordinates, OutPoints);
 	}
diff --git a/Engine/Source/Runtime/Datasmith/CADKernel/Base/Public/Geo/GeoPoint.h b/Engine/Source/Runtime/Datasmith/CADKernel/Base/Public/Geo/GeoPoint.h
--- a/Engine/Source/Runtime/Datasmith/CADKernel/Base/Public/Geo/GeoPoint.h
+++ b/Engine/Source/Runtime/Datasmith/CADKernel/Base/Public/Geo/GeoPoint.h
@@ -10,20 +10,20 @@ namespace UE::CADKernel
 struct FCurvePoint2D
 {
 	int32 DerivativeOrder = -1;
-	FPoint2D Point = FVector2d::ZeroVector;
-	FPoint2D Gradient = FVector2d::ZeroVector;
-	FPoint2D Laplacian = FVector2d::ZeroVector;
+	FVector2d Point = FVector2d::ZeroVector;
+	FVector2d Gradient = FVector2d::ZeroVector;
+	FVector2d Laplacian = FVector2d::ZeroVector;
 };
 
 struct FSurfacicPoint
 {
 	int32 DerivativeOrder = -1;
-	FPoint Point = FPoint::ZeroVector;
-	FPoint GradientU = FPoint::ZeroVector;
-	FPoint GradientV = FPoint::ZeroVector;
-	FPoint LaplacianU = FPoint::ZeroVector;
-	FPoint LaplacianV = FPoint::ZeroVector;
-	FPoint LaplacianUV = FPoint::ZeroVector;
+	FVector Point = FVector::ZeroVector;
+	FVector GradientU = FVector::ZeroVector;
+	FVector GradientV = FVector::ZeroVector;
+	FVector LaplacianU = FVector::ZeroVector;
+	FVector LaplacianV = FVector::ZeroVector;
+	FVector LaplacianUV = FVector::ZeroVector;
 };
 
 struct FSurfacicCurvePoint
@@ -31,16 +31,16 @@ struct FSurfacicCurvePoint
 	bool bWithNormals;
 	bool bWithTangent;
 
-	FPoint2D Point2D = FVector2d::ZeroVector;
-	FPoint Point = FPoint::ZeroVector;
-	FPoint Normal = FPoint::ZeroVector;
-	FPoint Tangent = FPoint::ZeroVector;
+	FVector2d Point2D = FVector2d::ZeroVector;
+	FVector Point = FVector::ZeroVector;
+	FVector Normal = FVector::ZeroVector;
+	FVector Tangent = FVector::ZeroVector;
 };
 
 struct FSurfacicCurvePointWithTolerance
 {
-	FPoint2D Point2D = FVector2d::ZeroVector;
-	FPoint Point = FPoint::ZeroVector;
+	FVector2d Point2D = FVector2d::ZeroVector;
+	FVector Point = FVector::ZeroVector;
 	FSurfacicTolerance Tolerance = FVectorUtil::FarawayPoint2D;
 };
 
@@ -49,24 +49,24 @@ typedef FSurfacicCurvePointWithTolerance FSurfacicCurveExtremities[2];
 struct FCurvePoint
 {
 	int32 DerivativeOrder = -1;
-	FPoint Point = FPoint::ZeroVector;
-	FPoint Gradient = FPoint::ZeroVector;
-	FPoint Laplacian = FPoint::ZeroVector;
+	FVector Point = FVector::ZeroVector;
+	FVector Gradient = FVector::ZeroVector;
+	FVector Laplacian = FVector::ZeroVector;
 
 	FCurvePoint() = default;
 
-	FCurvePoint(FPoint InPoint)
+	FCurvePoint(FVector InPoint)
 		: Point(InPoint)
 	{
-		Gradient = FPoint::ZeroVector;
-		Laplacian = FPoint::ZeroVector;
+		Gradient = FVector::ZeroVector;
+		Laplacian = FVector::ZeroVector;
 	}
 
 	void Init()
 	{
-		Point = FPoint::ZeroVector;
-		Gradient = FPoint::ZeroVector;
-		Laplacian = FPoint::ZeroVector;
+		Point = FVector::ZeroVector;
+		Gradient = FVector::ZeroVector;
+		Laplacian = FVector::ZeroVector;
 	}
 
 	/**
diff --git a/Engine/Source/Runtime/Datasmith/CADKernel/Base/Public/Geo/Sampling/Polyline.h b/Engine/Source/Runtime/Datasmith/CADKernel/Base/Public/Geo/Sampling/Polyline.h
--- a/Engine/Source/Runtime/Datasmith/CADKernel/Base/Public/Geo/Sampling/Polyline.h
+++ b/Engine/Source/Runtime/Datasmith/CADKernel/Base/Public/Geo/Sampling/Polyline.h
@@ -198,9 +198,9 @@ public:
 	}
 };
 
-typedef TPolyline<FPoint> FPolyline3D;
+typedef TPolyline<FVector> FPolyline3D;
 
-class FPolyline2D : public TPolyline<FPoint2D>
+class FPolyline2D : public TPolyline<FVector2d>
 {
 public:
 
diff --git a/Engine/Source/Runtime/Datasmith/CADKernel/Base/Public/Geo/Sampling/PolylineTools.h b/Engine/Source/Runtime/Datasmith/CADKernel/Base/Public/Geo/Sampling/PolylineTools.h
--- a/Engine/Source/Runtime/Datasmith/CADKernel/Base/Public/Geo/Sampling/PolylineTools.h
+++ b/Engine/Source/Runtime/Datasmith/CADKernel/Base/Public/Geo/Sampling/PolylineTools.h
@@ -18,10 +18,10 @@ namespace UE::CADKernel
 
 struct FPolylineBBox
 {
-	FPoint Max;
-	FPoint Min;
-	FPoint MaxPoints[3];
-	FPoint MinPoints[3];
+	FVector Max;
+	FVector Min;
+	FVector MaxPoints[3];
+	FVector MinPoints[3];
 	double CoordinateOfMaxPoint[3] = { -HUGE_VALUE, -HUGE_VALUE, -HUGE_VALUE };
 	double CoordinateOfMinPoint[3] = { HUGE_VALUE, HUGE_VALUE, HUGE_VALUE };
 
@@ -31,7 +31,7 @@ struct FPolylineBBox
 	{
 	}
 
-	void Update(const double Coordinate, const FPoint& Point)
+	void Update(const double Coordinate, const FVector& Point)
 	{
 		for (int32 Index = 0; Index < 3; ++Index)
 		{
@@ -50,6 +50,11 @@ struct FPolylineBBox
 			}
 		}
 	}
+
+	void Update(const double Coordinate, const FVector2d& Point2D)
+	{
+		Update(Coordinate, FVector(Point2D, 0.));
+	}
 };
 
 namespace PolylineTools
@@ -71,7 +76,7 @@ template<typename PointType>
 inline PointType LinearInterpolation(const TArray<PointType>& Array, const int32 Index, const double Coordinate)
 {
 	ensureCADKernel(Index + 1 < Array.Num());
-	return PointOnSegment(Array[Index], Array[Index + 1], Coordinate);
+	return PointOnSegment<PointType>(Array[Index], Array[Index + 1], Coordinate);
 }
 
 inline double SectionCoordinate(const TArray<double>& Array, const int32 Index, const double Coordinate)
@@ -118,7 +123,7 @@ TArray<double> ComputePolylineSegmentLengths(const PointType& StartPoint, const
 
 		for (int32 Index = 1; Index < InnerPointCount; ++Index)
 		{
-			const FPoint& CurrentPoint = InnerPolyline[Index];
+			const FVector& CurrentPoint = InnerPolyline[Index];
 			ElementLength.Add(PointType::Distance(*PrevPoint, CurrentPoint));
 			PrevPoint = &CurrentPoint;
 		}
@@ -273,7 +278,7 @@ protected:
 
 		for (int32 Index = BoundaryIndices[0]; Index <= BoundaryIndices[1]; ++Index)
 		{
-			FPoint ProjectPoint = ProjectPointOnSegment(InPointToProject, PolylinePoints[Index], PolylinePoints[Index + 1], ParamU, true);
+			PointType ProjectPoint = ProjectPointOnSegment<PointType>(InPointToProject, PolylinePoints[Index], PolylinePoints[Index + 1], ParamU, true);
 			double SquareDistance = FMath::Square(ProjectPoint[0] - InPointToProject[0]);
 			if (SquareDistance > MinDistance)
 			{
@@ -608,7 +613,7 @@ public:
 			int32 SegmentIndex = 0;
 			for (int32 Index = StartIndex; Index < EndIndex; ++Index)
 			{
-				PointType ProjectPoint = ProjectPointOnSegment(InPointToProject, PolylinePoints[Index], PolylinePoints[Index + 1], ParamU, true);
+				PointType ProjectPoint = ProjectPointOnSegment<PointType>(InPointToProject, PolylinePoints[Index], PolylinePoints[Index + 1], ParamU, true);
 				double SquareDistance = FVector::DistSquared(ProjectPoint, InPointToProject);
 				if (SquareDistance > MinDistance + SquareTol)
 				{
diff --git a/Engine/Source/Runtime/Datasmith/CADKernel/Base/Public/Geo/Sampling/SurfacicPolyline.h b/Engine/Source/Runtime/Datasmith/CADKernel/Base/Public/Geo/Sampling/SurfacicPolyline.h
--- a/Engine/Source/Runtime/Datasmith/CADKernel/Base/Public/Geo/Sampling/SurfacicPolyline.h
+++ b/Engine/Source/Runtime/Datasmith/CADKernel/Base/Public/Geo/Sampling/SurfacicPolyline.h
@@ -32,10 +32,10 @@ class CADKERNEL_API FSurfacicPolyline
 public:
 
 	TArray<double> Coordinates;
-	TArray<FPoint2D> Points2D;
-	TArray<FPoint> Points3D;
+	TArray<FVector2d> Points2D;
+	TArray<FVector> Points3D;
 	TArray<FVector3f> Normals;
-	TArray<FPoint> Tangents;
+	TArray<FVector> Tangents;
 
 	FSurfacicBoundary BoundingBox;
 
@@ -73,32 +73,32 @@ public:
 
 	void GetExtremities(const FLinearBoundary& InBoundary, const double Tolerance3D, const FSurfacicTolerance& Tolerances2D, FSurfacicCurveExtremities& Extremities) const;
 
-	FPoint Approximate3DPoint(double InCoordinate) const
+	FVector Approximate3DPoint(double InCoordinate) const
 	{
-		TPolylineApproximator<FPoint> Approximator3D(Coordinates, Points3D);
+		TPolylineApproximator<FVector> Approximator3D(Coordinates, Points3D);
 		return Approximator3D.ApproximatePoint(InCoordinate);
 	}
 
-	void Approximate3DPoints(const TArray<double>& InCoordinates, TArray<FPoint>& OutPoints) const
+	void Approximate3DPoints(const TArray<double>& InCoordinates, TArray<FVector>& OutPoints) const
 	{
-		TPolylineApproximator<FPoint> Approximator3D(Coordinates, Points3D);
+		TPolylineApproximator<FVector> Approximator3D(Coordinates, Points3D);
 		Approximator3D.ApproximatePoints(InCoordinates, OutPoints);
 	}
 
-	FPoint2D Approximate2DPoint(double InCoordinate) const
+	FVector2d Approximate2DPoint(double InCoordinate) const
 	{
-		TPolylineApproximator<FPoint2D> Approximator(Coordinates, Points2D);
+		TPolylineApproximator<FVector2d> Approximator(Coordinates, Points2D);
 		return Approximator.ApproximatePoint(InCoordinate);
 	}
 
-	FPoint GetTangentAt(double InCoordinate) const
+	FVector GetTangentAt(double InCoordinate) const
 	{
 		FDichotomyFinder Finder(Coordinates);
 		int32 Index = Finder.Find(InCoordinate);
 		return Points3D[Index + 1] - Points3D[Index];
 	}
 
-	FPoint2D GetTangent2DAt(double InCoordinate) const
+	FVector2d GetTangent2DAt(double InCoordinate) const
 	{
 		FDichotomyFinder Finder(Coordinates);
 		int32 Index = Finder.Find(InCoordinate);
@@ -133,9 +133,9 @@ public:
 		}
 	};
 
-	void Approximate2DPoints(const TArray<double>& InCoordinates, TArray<FPoint2D>& OutPoints) const
+	void Approximate2DPoints(const TArray<double>& InCoordinates, TArray<FVector2d>& OutPoints) const
 	{
-		TPolylineApproximator<FPoint2D> Approximator(Coordinates, Points2D);
+		TPolylineApproximator<FVector2d> Approximator(Coordinates, Points2D);
 		Approximator.ApproximatePoints(InCoordinates, OutPoints);
 	}
 
@@ -227,31 +227,31 @@ public:
 
 	void Sample(const FLinearBoundary& Boundary, const double DesiredSegmentLength, TArray<double>& OutCoordinates) const
 	{
-		TPolylineApproximator<FPoint> Approximator3D(Coordinates, Points3D);
+		TPolylineApproximator<FVector> Approximator3D(Coordinates, Points3D);
 		Approximator3D.SamplePolyline(Boundary, DesiredSegmentLength, OutCoordinates);
 	}
 
-	double GetCoordinateOfProjectedPoint(const FLinearBoundary& Boundary, const FPoint& PointOnEdge, FPoint& ProjectedPoint) const
+	double GetCoordinateOfProjectedPoint(const FLinearBoundary& Boundary, const FVector& PointOnEdge, FVector& ProjectedPoint) const
 	{
-		TPolylineApproximator<FPoint> Approximator3D(Coordinates, Points3D);
+		TPolylineApproximator<FVector> Approximator3D(Coordinates, Points3D);
 		return Approximator3D.ProjectPointToPolyline(Boundary, PointOnEdge, ProjectedPoint);
 	}
 
-	double GetCoordinateOfProjectedPoint(const FLinearBoundary& Boundary, const FPoint2D& PointOnEdge, FPoint2D& ProjectedPoint) const
+	double GetCoordinateOfProjectedPoint(const FLinearBoundary& Boundary, const FVector2d& PointOnEdge, FVector2d& ProjectedPoint) const
 	{
-		TPolylineApproximator<FPoint2D> Approximator2D(Coordinates, Points2D);
+		TPolylineApproximator<FVector2d> Approximator2D(Coordinates, Points2D);
 		return Approximator2D.ProjectPointToPolyline(Boundary, PointOnEdge, ProjectedPoint);
 	}
 
-	void ProjectPoints(const FLinearBoundary& InBoundary, const TArray<FPoint>& InPointsToProject, TArray<double>& ProjectedPointCoordinates, TArray<FPoint>& ProjectedPoints) const
+	void ProjectPoints(const FLinearBoundary& InBoundary, const TArray<FVector>& InPointsToProject, TArray<double>& ProjectedPointCoordinates, TArray<FVector>& ProjectedPoints) const
 	{
-		TPolylineApproximator<FPoint> Approximator3D(Coordinates, Points3D);
+		TPolylineApproximator<FVector> Approximator3D(Coordinates, Points3D);
 		Approximator3D.ProjectPointsToPolyline(InBoundary, InPointsToProject, ProjectedPointCoordinates, ProjectedPoints);
 	}
 
-	void ProjectPoints(const FLinearBoundary& InBoundary, const TArray<FPoint2D>& InPointsToProject, TArray<double>& ProjectedPointCoordinates, TArray<FPoint2D>& ProjectedPoints) const
+	void ProjectPoints(const FLinearBoundary& InBoundary, const TArray<FVector2d>& InPointsToProject, TArray<double>& ProjectedPointCoordinates, TArray<FVector2d>& ProjectedPoints) const
 	{
-		TPolylineApproximator<FPoint2D> Approximator(Coordinates, Points2D);
+		TPolylineApproximator<FVector2d> Approximator(Coordinates, Points2D);
 		Approximator.ProjectPointsToPolyline(InBoundary, InPointsToProject, ProjectedPointCoordinates, ProjectedPoints);
 	}
 
@@ -259,9 +259,9 @@ public:
 	 * Project each point of a coincidental polyline on the Polyline.
 	 * @param ToleranceOfProjection: Max error between the both curve to stop the search of projection
 	 */
-	void ProjectCoincidentalPolyline(const FLinearBoundary& InBoundary, const TArray<FPoint>& InPointsToProject, bool bSameOrientation, TArray<double>& OutProjectedPointCoordinates, double ToleranceOfProjection) const
+	void ProjectCoincidentalPolyline(const FLinearBoundary& InBoundary, const TArray<FVector>& InPointsToProject, bool bSameOrientation, TArray<double>& OutProjectedPointCoordinates, double ToleranceOfProjection) const
 	{
-		TPolylineApproximator<FPoint> Approximator3D(Coordinates, Points3D);
+		TPolylineApproximator<FVector> Approximator3D(Coordinates, Points3D);
 		Approximator3D.ProjectCoincidentalPolyline(InBoundary, InPointsToProject, bSameOrientation, OutProjectedPointCoordinates, ToleranceOfProjection);
 	}
 
@@ -275,17 +275,17 @@ public:
 		return Coordinates;
 	}
 
-	const TArray<FPoint2D>& Get2DPoints() const
+	const TArray<FVector2d>& Get2DPoints() const
 	{
 		return Points2D;
 	}
 
-	const FPoint& GetPointAt(int32 Index) const
+	const FVector& GetPointAt(int32 Index) const
 	{
 		return Points3D[Index];
 	}
 
-	const TArray<FPoint>& GetPoints() const
+	const TArray<FVector>& GetPoints() const
 	{
 		return Points3D;
 	}
@@ -295,7 +295,7 @@ public:
 		return Normals;
 	}
 
-	const TArray<FPoint>& GetTangents() const
+	const TArray<FVector>& GetTangents() const
 	{
 		return Tangents;
 	}
@@ -326,36 +326,36 @@ public:
 	/**
 	 * Get the sub 2d polyline bounded by the input InBoundary in the orientation of the input InOrientation and append it to the output OutPoints
 	 */
-	void GetSubPolyline(const FLinearBoundary& InBoundary, const EOrientation InOrientation, TArray<FPoint2D>& OutPoints) const
+	void GetSubPolyline(const FLinearBoundary& InBoundary, const EOrientation InOrientation, TArray<FVector2d>& OutPoints) const
 	{
-		TPolylineApproximator<FPoint2D> Approximator(Coordinates, Points2D);
+		TPolylineApproximator<FVector2d> Approximator(Coordinates, Points2D);
 		Approximator.GetSubPolyline(InBoundary, InOrientation, OutPoints);
 	}
 
 	/**
 	 * Get the sub polyline bounded by the input InBoundary in the orientation of the input InOrientation and append it to the output OutPoints
 	 */
-	void GetSubPolyline(const FLinearBoundary& InBoundary, TArray<double>& OutCoordinates, TArray<FPoint2D>& OutPoints) const
+	void GetSubPolyline(const FLinearBoundary& InBoundary, TArray<double>& OutCoordinates, TArray<FVector2d>& OutPoints) const
 	{
-		TPolylineApproximator<FPoint2D> Approximator(Coordinates, Points2D);
+		TPolylineApproximator<FVector2d> Approximator(Coordinates, Points2D);
 		Approximator.GetSubPolyline(InBoundary, OutCoordinates, OutPoints);
 	}
 
 	/**
 	 * Get the sub polyline bounded by the input InBoundary in the orientation of the input InOrientation and append it to the output OutPoints
 	 */
-	void GetSubPolyline(const FLinearBoundary& InBoundary, const EOrientation InOrientation, TArray<FPoint>& OutPoints) const
+	void GetSubPolyline(const FLinearBoundary& InBoundary, const EOrientation InOrientation, TArray<FVector>& OutPoints) const
 	{
-		TPolylineApproximator<FPoint> Approximator3D(Coordinates, Points3D);
+		TPolylineApproximator<FVector> Approximator3D(Coordinates, Points3D);
 		Approximator3D.GetSubPolyline(InBoundary, InOrientation, OutPoints);
 	}
 
 	/**
 	 * Get the sub polyline bounded by the input InBoundary in the orientation of the input InOrientation and append it to the output OutPoints
 	 */
-	void GetSubPolyline(const FLinearBoundary& InBoundary, TArray<double>& OutCoordinates, TArray<FPoint>& OutPoints) const
+	void GetSubPolyline(const FLinearBoundary& InBoundary, TArray<double>& OutCoordinates, TArray<FVector>& OutPoints) const
 	{
-		TPolylineApproximator<FPoint> Approximator3D(Coordinates, Points3D);
+		TPolylineApproximator<FVector> Approximator3D(Coordinates, Points3D);
 		Approximator3D.GetSubPolyline(InBoundary, OutCoordinates, OutPoints);
 	}
 
@@ -480,20 +480,20 @@ public:
 
 	double GetLength(const FLinearBoundary& InBoundary) const
 	{
-		TPolylineApproximator<FPoint> Approximator3D(Coordinates, Points3D);
+		TPolylineApproximator<FVector> Approximator3D(Coordinates, Points3D);
 		return Approximator3D.ComputeLengthOfSubPolyline(InBoundary);
 	}
 
 	double Get2DLength(const FLinearBoundary& InBoundary) const
 	{
-		TPolylineApproximator<FPoint2D> Approximator(Coordinates, Points2D);
+		TPolylineApproximator<FVector2d> Approximator(Coordinates, Points2D);
 		return Approximator.ComputeLengthOfSubPolyline(InBoundary);
 	}
 
 	bool IsIso(EIso Iso, double ErrorTolerance = DOUBLE_SMALL_NUMBER) const
 	{
-		FPoint2D StartPoint = Points2D[0];
-		return Algo::AllOf(Points2D, [&](const FPoint2D& Point) { 
+		FVector2d StartPoint = Points2D[0];
+		return Algo::AllOf(Points2D, [&](const FVector2d& Point) { 
 			return FMath::IsNearlyEqual(Point[Iso], StartPoint[Iso], ErrorTolerance);
 			});
 	}
diff --git a/Engine/Source/Runtime/Datasmith/CADKernel/Base/Public/Geo/Surfaces/NurbsSurfaceData.h b/Engine/Source/Runtime/Datasmith/CADKernel/Base/Public/Geo/Surfaces/NurbsSurfaceData.h
--- a/Engine/Source/Runtime/Datasmith/CADKernel/Base/Public/Geo/Surfaces/NurbsSurfaceData.h
+++ b/Engine/Source/Runtime/Datasmith/CADKernel/Base/Public/Geo/Surfaces/NurbsSurfaceData.h
@@ -37,7 +37,7 @@ struct FNurbsSurfaceData
 
 	// if Weights.Num() == 0  => bIsRational = false
 	TArray<double> Weights;
-	TArray<FPoint> Poles;
+	TArray<FVector> Poles;
 };
 
 }
\ No newline at end of file
diff --git a/Engine/Source/Runtime/Datasmith/CADKernel/Base/Public/Geo/Surfaces/Surface.h b/Engine/Source/Runtime/Datasmith/CADKernel/Base/Public/Geo/Surfaces/Surface.h
--- a/Engine/Source/Runtime/Datasmith/CADKernel/Base/Public/Geo/Surfaces/Surface.h
+++ b/Engine/Source/Runtime/Datasmith/CADKernel/Base/Public/Geo/Surfaces/Surface.h
@@ -81,7 +81,7 @@ protected:
 
 public:
 
-	static TSharedPtr<FSurface> MakeBezierSurface(const double InToleranceGeometric, int32 InUDegre, int32 InVDegre, const TArray<FPoint>& InPoles);
+	static TSharedPtr<FSurface> MakeBezierSurface(const double InToleranceGeometric, int32 InUDegre, int32 InVDegre, const TArray<FVector>& InPoles);
 	static TSharedPtr<FSurface> MakeConeSurface(const double InToleranceGeometric, const FMatrixH& InMatrix, double InStartRadius, double InConeAngle, const FSurfacicBoundary& InBoundary);
 	static TSharedPtr<FSurface> MakeCylinderSurface(const double InToleranceGeometric, const FMatrixH& InMatrix, const double InRadius, const FSurfacicBoundary& InBoundary);
 	static TSharedPtr<FSurface> MakeNurbsSurface(const double InToleranceGeometric, const FNurbsSurfaceData& NurbsData);
@@ -90,7 +90,7 @@ public:
 	static TSharedPtr<FSurface> MakeSphericalSurface(const double InToleranceGeometric, const FMatrixH& InMatrix, double InRadius, const FSurfacicBoundary& InBoundary);
 	static TSharedPtr<FSurface> MakeTorusSurface(const double InToleranceGeometric, const FMatrixH& InMatrix, double InMajorRadius, double InMinorRadius, const FSurfacicBoundary& InBoundary);
 
-	virtual void ValidateUVPoints(TArray<FPoint2D>& UVPoints) const {}
+	virtual void ValidateUVPoints(TArray<FVector2d>& UVPoints) const {}
 
 	virtual void Serialize(FCADKernelArchive& Ar) override
 	{
@@ -173,9 +173,9 @@ public:
 	}
 
 
-	virtual void EvaluatePoint(const FPoint2D& InSurfacicCoordinate, FSurfacicPoint& OutPoint3D, int32 InDerivativeOrder = 0) const = 0;
+	virtual void EvaluatePoint(const FVector2d& InSurfacicCoordinate, FSurfacicPoint& OutPoint3D, int32 InDerivativeOrder = 0) const = 0;
 
-	virtual void EvaluatePoints(const TArray<FPoint2D>& InSurfacicCoordinates, TArray<FSurfacicPoint>& OutPoint3D, int32 InDerivativeOrder = 0) const;
+	virtual void EvaluatePoints(const TArray<FVector2d>& InSurfacicCoordinates, TArray<FSurfacicPoint>& OutPoint3D, int32 InDerivativeOrder = 0) const;
 	virtual void EvaluatePoints(const TArray<FCurvePoint2D>& InSurfacicCoordinates, TArray<FSurfacicPoint>& OutPoint3D, int32 InDerivativeOrder = 0) const;
 	virtual void EvaluatePoints(const TArray<FCurvePoint2D>& InSurfacicCoordinates, TArray<FCurvePoint>& OutPoint3D, int32 InDerivativeOrder = 0) const;
 	virtual void EvaluatePoints(FSurfacicPolyline& Polyline) const;
@@ -192,15 +192,15 @@ public:
 	 * For Normal(u,v) parallel Rotation Axis -> use Plan projection
 	 * For Normal(u,v) perpendicular to Rotation Axis -> use cylindrical projection
 	 */
-	virtual FPoint2D EvaluatePointInCylindricalSpace(const FPoint2D& InSurfacicCoordinate) const { return FVector2d::ZeroVector; }
-	virtual void EvaluatePointGridInCylindricalSpace(const FCoordinateGrid& Coordinates, TArray<FPoint2D>&) const { ; }
+	virtual FVector2d EvaluatePointInCylindricalSpace(const FVector2d& InSurfacicCoordinate) const { return FVector2d::ZeroVector; }
+	virtual void EvaluatePointGridInCylindricalSpace(const FCoordinateGrid& Coordinates, TArray<FVector2d>&) const { ; }
 
 	virtual void EvaluatePointGrid(const FCoordinateGrid& Coordinates, FSurfacicSampling& OutPoints, bool bComputeNormals = false) const;
 	void EvaluateGrid(FGrid& Grid) const;
 
-	virtual void EvaluateNormals(const TArray<FPoint2D>& Points2D, TArray<FVector3f>& Normals) const;
+	virtual void EvaluateNormals(const TArray<FVector2d>& Points2D, TArray<FVector3f>& Normals) const;
 
-	virtual void EvaluatePointAndNormal(const FPoint2D& InPoint2D, FPoint& Point, FVector3f& Normal) const
+	virtual void EvaluatePointAndNormal(const FVector2d& InPoint2D, FVector& Point, FVector3f& Normal) const
 	{
 		int32 DerivativeOrder = 1;
 		FSurfacicPoint SurfacicPoint;
@@ -212,7 +212,7 @@ public:
 		Normal = (FVector3f)NormalD.GetSafeNormal();
 	}
 
-	virtual FVector EvaluateNormal(const FPoint2D& InPoint2D) const
+	virtual FVector EvaluateNormal(const FVector2d& InPoint2D) const
 	{
 		int32 DerivativeOrder = 1;
 		FSurfacicPoint Point3D;
@@ -281,11 +281,11 @@ public:
 	 * This function is useful to estimate tolerance when scales are defined in the Matrix
 	 * @param InAxis a vetor of length 1
 	 * @param InMatrix
-	 * @param InOrigin = InMatrix*FPoint::ZeroVector
+	 * @param InOrigin = InMatrix*FVector::ZeroVector
 	 */
-	static double ComputeScaleAlongAxis(const FPoint& InAxis, const FMatrixH& InMatrix, const FPoint& InOrigin)
+	static double ComputeScaleAlongAxis(const FVector& InAxis, const FMatrixH& InMatrix, const FVector& InOrigin)
 	{
-		FPoint Point = InMatrix.Multiply(InAxis);
+		FVector Point = InMatrix.Multiply(InAxis);
 		double Length = FVector::Distance(InOrigin, Point);
 		return Length;
 	};
diff --git a/Engine/Source/Runtime/Datasmith/CADKernel/Base/Public/Math/Aabb.h b/Engine/Source/Runtime/Datasmith/CADKernel/Base/Public/Math/Aabb.h
--- a/Engine/Source/Runtime/Datasmith/CADKernel/Base/Public/Math/Aabb.h
+++ b/Engine/Source/Runtime/Datasmith/CADKernel/Base/Public/Math/Aabb.h
@@ -206,51 +206,51 @@ public:
 };
 
 template<>
-FORCEINLINE TAABB<FPoint2D, 2>::TAABB()
+FORCEINLINE TAABB<FVector2d, 2>::TAABB()
 	: MinCorner(FVectorUtil::FarawayPoint2D)
 	, MaxCorner(-FVectorUtil::FarawayPoint2D)
 {
 }
 
 template<>
-FORCEINLINE TAABB<FPoint, 3>::TAABB()
+FORCEINLINE TAABB<FVector, 3>::TAABB()
 	: MinCorner(FVectorUtil::FarawayPoint3D)
 	, MaxCorner(-FVectorUtil::FarawayPoint3D)
 {
 }
 
 template<>
-FORCEINLINE void TAABB<FPoint2D, 2>::Empty()
+FORCEINLINE void TAABB<FVector2d, 2>::Empty()
 {
 	MinCorner = FVectorUtil::FarawayPoint2D;
 	MaxCorner = -FVectorUtil::FarawayPoint2D;
 }
 
 template<>
-FORCEINLINE void TAABB<FPoint, 3>::Empty()
+FORCEINLINE void TAABB<FVector, 3>::Empty()
 {
 	MinCorner = FVectorUtil::FarawayPoint3D;
 	MaxCorner = -FVectorUtil::FarawayPoint3D;
 }
 
 
-class FAABB : public TAABB<FPoint>
+class FAABB : public TAABB<FVector>
 {
 
 public:
 	FAABB()
-		: TAABB<FPoint>()
+		: TAABB<FVector>()
 	{
 	}
 
-	FAABB(const FPoint& InMinCorner, const FPoint& InMaxCorner)
-		: TAABB<FPoint>(InMinCorner, InMaxCorner)
+	FAABB(const FVector& InMinCorner, const FVector& InMaxCorner)
+		: TAABB<FVector>(InMinCorner, InMaxCorner)
 	{
 	}
 
-	FPoint GetCorner(int32 Corner) const
+	FVector GetCorner(int32 Corner) const
 	{
-		return FPoint(
+		return FVector(
 			Corner & EAABBBoundary::XMin ? MinCorner[0] : MaxCorner[0],
 			Corner & EAABBBoundary::YMin ? MinCorner[1] : MaxCorner[1],
 			Corner & EAABBBoundary::ZMin ? MinCorner[2] : MaxCorner[2]
@@ -258,22 +258,22 @@ public:
 	}
 };
 
-class CADKERNEL_API FAABB2D : public TAABB<FPoint2D, 2>
+class CADKERNEL_API FAABB2D : public TAABB<FVector2d, 2>
 {
 public:
 	FAABB2D()
-		: TAABB<FPoint2D, 2>()
+		: TAABB<FVector2d, 2>()
 	{
 	}
 
-	FAABB2D(const FPoint& InMinCorner, const FPoint& InMaxCorner)
-		: TAABB<FPoint2D, 2>(InMinCorner, InMaxCorner)
+	FAABB2D(const FVector& InMinCorner, const FVector& InMaxCorner)
+		: TAABB<FVector2d, 2>(FVector2d(InMinCorner.X, InMinCorner.Y), FVector2d(InMaxCorner.X, InMaxCorner.Y))
 	{
 	}
 
-	FPoint2D GetCorner(int32 CornerIndex) const
+	FVector2d GetCorner(int32 CornerIndex) const
 	{
-		return FPoint2D(
+		return FVector2d(
 			CornerIndex & EAABBBoundary::XMin ? MinCorner[0] : MaxCorner[0],
 			CornerIndex & EAABBBoundary::YMin ? MinCorner[1] : MaxCorner[1]
 		);
diff --git a/Engine/Source/Runtime/Datasmith/CADKernel/Base/Public/Math/BSpline.h b/Engine/Source/Runtime/Datasmith/CADKernel/Base/Public/Math/BSpline.h
--- a/Engine/Source/Runtime/Datasmith/CADKernel/Base/Public/Math/BSpline.h
+++ b/Engine/Source/Runtime/Datasmith/CADKernel/Base/Public/Math/BSpline.h
@@ -32,7 +32,7 @@ CADKERNEL_API void FindNotDerivableParameters(const FNURBSSurface&, int32 InDeri
 CADKERNEL_API void Evaluate2DPoint(const FNURBSCurve&, double Coordinate, FCurvePoint2D& OutPoint, int32 DerivativeOrder);
 CADKERNEL_API void EvaluatePoint(const FNURBSCurve&, double Coordinate, FCurvePoint& OutPoint, int32 DerivativeOrder);
 
-CADKERNEL_API void EvaluatePoint(const FNURBSSurface&, const FPoint2D& InPoint2D, FSurfacicPoint& OutPoint3D, int32 InDerivativeOrder);
+CADKERNEL_API void EvaluatePoint(const FNURBSSurface&, const FVector2d& InPoint2D, FSurfacicPoint& OutPoint3D, int32 InDerivativeOrder);
 CADKERNEL_API void EvaluatePointGrid(const FNURBSSurface&, const FCoordinateGrid& Coords, FSurfacicSampling& OutPoints, bool bComputeNormals);
 
 CADKERNEL_API TSharedPtr<FNURBSCurve> DuplicateNurbsCurveWithHigherDegree(int32 degre, const FNURBSCurve& InCurve);
diff --git a/Engine/Source/Runtime/Datasmith/CADKernel/Base/Public/Math/Boundary.h b/Engine/Source/Runtime/Datasmith/CADKernel/Base/Public/Math/Boundary.h
--- a/Engine/Source/Runtime/Datasmith/CADKernel/Base/Public/Math/Boundary.h
+++ b/Engine/Source/Runtime/Datasmith/CADKernel/Base/Public/Math/Boundary.h
@@ -268,18 +268,18 @@ public:
 		Offset(OffsetTolerance);
 	}
 
-	FSurfacicBoundary(const FPoint2D& Point1, const FPoint2D& Point2)
+	FSurfacicBoundary(const FVector2d& Point1, const FVector2d& Point2)
 	{
 		Set(Point1, Point2);
 	}
 
-	FSurfacicBoundary(const FPoint2D& Point1, const FPoint2D& Point2, const double OffsetTolerance)
+	FSurfacicBoundary(const FVector2d& Point1, const FVector2d& Point2, const double OffsetTolerance)
 	{
 		Set(Point1, Point2);
 		Offset(OffsetTolerance);
 	}
 
-	void Set(const FPoint2D& Point1, const FPoint2D& Point2)
+	void Set(const FVector2d& Point1, const FVector2d& Point2)
 	{
 		UVBoundaries[EIso::IsoU].Set(Point1.X, Point2.X);
 		UVBoundaries[EIso::IsoV].Set(Point1.Y, Point2.Y);
@@ -313,10 +313,10 @@ public:
 	/**
 	 * Set the boundary with the min and max of this array
 	 */
-	void Set(const TArray<FPoint2D>& Points)
+	void Set(const TArray<FVector2d>& Points)
 	{
 		Init();
-		for (const FPoint2D& Point : Points)
+		for (const FVector2d& Point : Points)
 		{
 			ExtendTo(Point);
 		}
@@ -393,13 +393,13 @@ public:
 		UVBoundaries[EIso::IsoV].ExtendTo(MaxLimit[EIso::IsoV]);
 	}
 
-	void ExtendTo(const FPoint2D& Point)
+	void ExtendTo(const FVector2d& Point)
 	{
 		UVBoundaries[EIso::IsoU].ExtendTo(Point.X);
 		UVBoundaries[EIso::IsoV].ExtendTo(Point.Y);
 	}
 
-	void ExtendTo(const FPoint& Point)
+	void ExtendTo(const FVector& Point)
 	{
 		UVBoundaries[EIso::IsoU].ExtendTo(Point.X);
 		UVBoundaries[EIso::IsoV].ExtendTo(Point.Y);
@@ -423,7 +423,7 @@ public:
 	/**
 	 * If a point is outside the bounds, set the coordinate to insert the point inside the bounds
 	 */
-	void MoveInsideIfNot(FPoint& Point, const double Tolerance = DOUBLE_SMALL_NUMBER) const
+	void MoveInsideIfNot(FVector& Point, const double Tolerance = DOUBLE_SMALL_NUMBER) const
 	{
 		UVBoundaries[EIso::IsoU].MoveInsideIfNot(Point.X, Tolerance);
 		UVBoundaries[EIso::IsoV].MoveInsideIfNot(Point.Y, Tolerance);
@@ -432,7 +432,7 @@ public:
 	/**
 	 * If a point is outside the bounds, set the coordinate to insert the point inside the bounds
 	 */
-	void MoveInsideIfNot(FPoint2D& Point, const double Tolerance = DOUBLE_SMALL_NUMBER) const
+	void MoveInsideIfNot(FVector2d& Point, const double Tolerance = DOUBLE_SMALL_NUMBER) const
 	{
 		UVBoundaries[EIso::IsoU].MoveInsideIfNot(Point.X, Tolerance);
 		UVBoundaries[EIso::IsoV].MoveInsideIfNot(Point.Y, Tolerance);
diff --git a/Engine/Source/Runtime/Datasmith/CADKernel/Base/Public/Math/Geometry.h b/Engine/Source/Runtime/Datasmith/CADKernel/Base/Public/Math/Geometry.h
--- a/Engine/Source/Runtime/Datasmith/CADKernel/Base/Public/Math/Geometry.h
+++ b/Engine/Source/Runtime/Datasmith/CADKernel/Base/Public/Math/Geometry.h
@@ -29,10 +29,10 @@ CADKERNEL_API void SetTolerance(const double Tolerance);
  * https://en.wikipedia.org/wiki/Circumscribed_circle#Cartesian_coordinates_2
  * With A = (0, 0)
  */
-CADKERNEL_API inline FPoint2D ComputeCircumCircleCenter(const FPoint2D& InPoint0, const FPoint2D& InPoint1, const FPoint2D& InPoint2)
+CADKERNEL_API inline FVector2d ComputeCircumCircleCenter(const FVector2d& InPoint0, const FVector2d& InPoint1, const FVector2d& InPoint2)
 {
-	FPoint2D Segment_P0_P1 = InPoint1 - InPoint0;
-	FPoint2D Segment_P0_P2 = InPoint2 - InPoint0;
+	FVector2d Segment_P0_P1 = InPoint1 - InPoint0;
+	FVector2d Segment_P0_P2 = InPoint2 - InPoint0;
 
 	// D = 2(BuCv - BvCu)
 	double D = 2. * Segment_P0_P1 ^ Segment_P0_P2;
@@ -48,30 +48,30 @@ CADKERNEL_API inline FPoint2D ComputeCircumCircleCenter(const FPoint2D& InPoint0
 	double CenterU = (SquareNormB * Segment_P0_P2.Y - SquareNormC * Segment_P0_P1.Y) / D;
 	double CenterV = (SquareNormC * Segment_P0_P1.X - SquareNormB * Segment_P0_P2.X) / D;
 
-	return FPoint2D(CenterU, CenterV) + InPoint0;
+	return FVector2d(CenterU, CenterV) + InPoint0;
 }
 
-CADKERNEL_API inline FPoint ComputeCircumCircleCenter(const FPoint& Point0, const FPoint& Point1, const FPoint& Point2)
+CADKERNEL_API inline FVector ComputeCircumCircleCenter(const FVector& Point0, const FVector& Point1, const FVector& Point2)
 {
-	FPoint Test = Point1 - Point0;
-	FPoint AxisZ = (Point1 - Point0) ^ (Point2 - Point0);
+	FVector Test = Point1 - Point0;
+	FVector AxisZ = (Point1 - Point0) ^ (Point2 - Point0);
 	AxisZ.Normalize();
 
-	FPoint AxisX = Point1 - Point0;
+	FVector AxisX = Point1 - Point0;
 	AxisX.Normalize();
-	FPoint AxisY = AxisZ ^ AxisX;
+	FVector AxisY = AxisZ ^ AxisX;
 
 	FMatrixH Matrix(Point0, AxisX, AxisY, AxisZ);
 	FMatrixH MatrixInverse = Matrix;
 	MatrixInverse.Inverse();
 
-	FPoint2D D2Point1 = MatrixInverse * Point1;
-	FPoint2D D2Point2 = MatrixInverse * Point2;
+	FVector2d D2Point1 = FVectorUtil::FromVector(MatrixInverse * Point1);
+	FVector2d D2Point2 = FVectorUtil::FromVector(MatrixInverse * Point2);
 
 	double D = 2. * D2Point1 ^ D2Point2;
 	if (FMath::IsNearlyZero(D, SMALL_NUMBER_SQUARE))
 	{
-		return FPoint::ZeroVector;
+		return FVector::ZeroVector;
 	}
 
 	double SquareNormB = D2Point1.SquaredLength();
@@ -80,15 +80,15 @@ CADKERNEL_API inline FPoint ComputeCircumCircleCenter(const FPoint& Point0, cons
 	double CenterU = (SquareNormB * D2Point2.Y - SquareNormC * D2Point1.Y) / D;
 	double CenterV = (SquareNormC * D2Point1.X - SquareNormB * D2Point2.X) / D;
 
-	FPoint Center2D(CenterU, CenterV, 0);
+	FVector Center2D(CenterU, CenterV, 0);
 	return Matrix * Center2D;
 }
 
 
-CADKERNEL_API inline FPoint2D ComputeCircumCircleCenterAndSquareRadius(const FPoint2D& InPoint0, const FPoint2D& InPoint1, const FPoint2D& InPoint2, double& OutSquareRadius)
+CADKERNEL_API inline FVector2d ComputeCircumCircleCenterAndSquareRadius(const FVector2d& InPoint0, const FVector2d& InPoint1, const FVector2d& InPoint2, double& OutSquareRadius)
 {
-	FPoint2D Segment_P0_P1 = InPoint1 - InPoint0;
-	FPoint2D Segment_P0_P2 = InPoint2 - InPoint0;
+	FVector2d Segment_P0_P1 = InPoint1 - InPoint0;
+	FVector2d Segment_P0_P2 = InPoint2 - InPoint0;
 
 	double D = 2. * Segment_P0_P1 ^ Segment_P0_P2;
 	if (FMath::IsNearlyZero(D, SMALL_NUMBER_SQUARE))
@@ -102,12 +102,149 @@ CADKERNEL_API inline FPoint2D ComputeCircumCircleCenterAndSquareRadius(const FPo
 	double CenterU = (SquareNormB * Segment_P0_P2.Y - SquareNormC * Segment_P0_P1.Y) / D;
 	double CenterV = (SquareNormC * Segment_P0_P1.X - SquareNormB * Segment_P0_P2.X) / D;
 
-	FPoint2D Center(CenterU, CenterV);
+	FVector2d Center(CenterU, CenterV);
 	OutSquareRadius = Center.SquaredLength();
 
 	return Center + InPoint0;
 }
 
+/**
+ * @param OutCoordinate: the coordinate of the projected point in the segment AB (coodinate of A = 0 and of B = 1)
+ * @return Projected point
+ */
+template<class PointType>
+inline PointType ProjectPointOnSegment(const PointType& Point, const PointType& InSegmentA, const PointType& InSegmentB, double& OutCoordinate, bool bRestrictCoodinateToInside = true)
+{
+	PointType Segment = InSegmentB - InSegmentA;
+
+	double SquaredLength = Segment | Segment;
+
+	if (SquaredLength <= 0.0)
+	{
+		OutCoordinate = 0.0;
+		return InSegmentA;
+	}
+	else
+	{
+		PointType APoint = Point - InSegmentA;
+		OutCoordinate = (APoint | Segment) / SquaredLength;
+
+		if (bRestrictCoodinateToInside)
+		{
+			if (OutCoordinate < 0.0)
+			{
+				OutCoordinate = 0.0;
+				return InSegmentA;
+			}
+
+			if (OutCoordinate > 1.0)
+			{
+				OutCoordinate = 1.0;
+				return InSegmentB;
+			}
+		}
+
+		PointType ProjectedPoint = Segment * OutCoordinate;
+		ProjectedPoint += InSegmentA;
+		return ProjectedPoint;
+	}
+}
+
+inline FVector ProjectPointOnPlane(const FVector& Point, const FVector& Origin, const FVector& InNormal, double& OutDistance)
+{
+	FVector Normal = InNormal;
+	ensureCADKernel(!FMath::IsNearlyZero(Normal.Length()));
+	Normal.Normalize();
+
+	FVector OP = Point - Origin;
+	OutDistance = OP | Normal;
+
+	return Point - (Normal * OutDistance);
+}
+
+/**
+ * @return InSegmentA + (InSegmentB - InSegmentA) | InCoordinate
+ */
+template<class PointType>
+inline PointType PointOnSegment(const PointType& InSegmentA, const PointType& InSegmentB, double InCoordinate)
+{
+	PointType Segment = InSegmentB - InSegmentA;
+	return InSegmentA + Segment * InCoordinate;
+}
+
+/**
+ * @return the distance between the point and the segment. If the projection of the point on the segment
+ * is not inside it, return the distance of the point to nearest the segment extremity
+ */
+template<class PointType>
+inline double DistanceOfPointToSegment(const PointType& Point, const PointType& SegmentPoint1, const PointType& SegmentPoint2)
+{
+	double Coordinate;
+	return PointType::Distance(ProjectPointOnSegment<PointType>(Point, SegmentPoint1, SegmentPoint2, Coordinate, /*bRestrictCoodinateToInside*/ true), Point);
+}
+
+/**
+ * @return the distance between the point and the segment. If the projection of the point on the segment
+ * is not inside it, return the distance of the point to nearest the segment extremity
+ */
+template<class PointType>
+inline double SquareDistanceOfPointToSegment(const PointType& Point, const PointType& SegmentPoint1, const PointType& SegmentPoint2)
+{
+	double Coordinate;
+	return FVector2d::DistSquared(ProjectPointOnSegment<PointType>(Point, SegmentPoint1, SegmentPoint2, Coordinate, /*bRestrictCoodinateToInside*/ true), Point);
+}
+
+/**
+ * @return the distance between the point and the line i.e. the distance between the point and its projection on the line
+ */
+template<class PointType>
+inline double DistanceOfPointToLine(const PointType& Point, const PointType& LinePoint1, const PointType& LineDirection)
+{
+	double Coordinate;
+	return ProjectPointOnSegment<PointType>(Point, LinePoint1, LinePoint1 + LineDirection, Coordinate, /*bRestrictCoodinateToInside*/ false).Distance(Point);
+}
+
+CADKERNEL_API double ComputeCurvature(const FVector& Gradient, const FVector& Laplacian);
+CADKERNEL_API double ComputeCurvature(const FVector& normal, const FVector& Gradient, const FVector& Laplacian);
+
+/**
+ * @return Coordinate of the projected point in the segment AB (coordinate of A = 0 and of B = 1)
+ */
+template<class PointType>
+inline double CoordinateOfProjectedPointOnSegment(const PointType& Point, const PointType& InSegmentA, const PointType& InSegmentB, bool bRestrictCoodinateToInside = true)
+{
+	PointType Segment = InSegmentB - InSegmentA;
+
+	double SquaredLength = Segment | Segment;
+
+	if (SquaredLength <= 0.0)
+	{
+		return 0.0;
+	}
+	else
+	{
+		PointType APoint = Point - InSegmentA;
+		double Coordinate = (APoint | Segment) / SquaredLength;
+
+		if (bRestrictCoodinateToInside)
+		{
+			if (Coordinate < 0.0)
+			{
+				return 0.0;
+			}
+
+			if (Coordinate > 1.0)
+			{
+				return 1.0;
+			}
+		}
+
+		return Coordinate;
+	}
+}
+
+CADKERNEL_API void FindLoopIntersectionsWithIso(const EIso Iso, const double IsoParameter, const TArray<TArray<FVector2d>>& Loops, TArray<double>& OutIntersections);
+
 template<class PointType>
 struct CADKERNEL_API TSegment
 {
@@ -137,7 +274,7 @@ struct CADKERNEL_API TSegment
 
 	double SquaredLength() const
 	{
-		return FPoint2D::DistSquared(Point0, Point1);
+		return FVector2d::DistSquared(Point0, Point1);
 	}
 
 	PointType GetVector() const
@@ -146,8 +283,8 @@ struct CADKERNEL_API TSegment
 	}
 };
 
-using FSegment2D = TSegment<FPoint2D>;
-using FSegment3D = TSegment<FPoint>;
+using FSegment2D = TSegment<FVector2d>;
+using FSegment3D = TSegment<FVector>;
 
 template<class PointType>
 struct CADKERNEL_API TTriangle
@@ -179,7 +316,7 @@ struct CADKERNEL_API TTriangle
 		}
 	}
 
-	virtual inline PointType ProjectPoint(const PointType& InPoint, FPoint2D& OutCoordinate)
+	virtual inline PointType ProjectPoint(const PointType& InPoint, FVector2d& OutCoordinate)
 	{
 		PointType Segment01 = Point1 - Point0;
 		PointType Segment02 = Point2 - Point0;
@@ -248,17 +385,17 @@ struct CADKERNEL_API TTriangle
 		switch (SideIndex)
 		{
 		case Side01:
-			ProjectedPoint = ProjectPointOnSegment(InPoint, Point0, Point1, SegmentCoordinate);
+			ProjectedPoint = ProjectPointOnSegment<PointType>(InPoint, Point0, Point1, SegmentCoordinate);
 			OutCoordinate.X = SegmentCoordinate;
 			OutCoordinate.Y = 0.0;
 			break;
 		case Side20:
-			ProjectedPoint = ProjectPointOnSegment(InPoint, Point0, Point2, SegmentCoordinate);
+			ProjectedPoint = ProjectPointOnSegment<PointType>(InPoint, Point0, Point2, SegmentCoordinate);
 			OutCoordinate.X = 0.0;
 			OutCoordinate.Y = SegmentCoordinate;
 			break;
 		case Side12:
-			ProjectedPoint = ProjectPointOnSegment(InPoint, Point1, Point2, SegmentCoordinate);
+			ProjectedPoint = ProjectPointOnSegment<PointType>(InPoint, Point1, Point2, SegmentCoordinate);
 			OutCoordinate.X = 1.0 - SegmentCoordinate;
 			OutCoordinate.Y = SegmentCoordinate;
 			break;
@@ -272,179 +409,42 @@ struct CADKERNEL_API TTriangle
 	}
 };
 
-struct CADKERNEL_API FTriangle : public TTriangle<FPoint>
+struct CADKERNEL_API FTriangle : public TTriangle<FVector>
 {
-	FTriangle(const FPoint& InPoint0, const FPoint& InPoint1, const FPoint& InPoint2)
-		: TTriangle<FPoint>(InPoint0, InPoint1, InPoint2)
+	FTriangle(const FVector& InPoint0, const FVector& InPoint1, const FVector& InPoint2)
+		: TTriangle<FVector>(InPoint0, InPoint1, InPoint2)
 	{
 	}
 
-	virtual FPoint ComputeNormal() const
+	virtual FVector ComputeNormal() const
 	{
-		FPoint Normal = (Point1 - Point0) ^ (Point2 - Point0);
+		FVector Normal = (Point1 - Point0) ^ (Point2 - Point0);
 		Normal.Normalize();
 		return Normal;
 	}
 };
 
-struct CADKERNEL_API FTriangle2D : public TTriangle<FPoint2D>
+struct CADKERNEL_API FTriangle2D : public TTriangle<FVector2d>
 {
-	FTriangle2D(const FPoint2D& InPoint0, const FPoint2D& InPoint1, const FPoint2D& InPoint2)
-		: TTriangle<FPoint2D>(InPoint0, InPoint1, InPoint2)
+	FTriangle2D(const FVector2d& InPoint0, const FVector2d& InPoint1, const FVector2d& InPoint2)
+		: TTriangle<FVector2d>(InPoint0, InPoint1, InPoint2)
 	{
 	}
 
-	FPoint2D CircumCircleCenterWithSquareRadius(double& SquareRadius) const
+	FVector2d CircumCircleCenterWithSquareRadius(double& SquareRadius) const
 	{
 		return ComputeCircumCircleCenterAndSquareRadius(this->Point0, this->Point1, this->Point2, SquareRadius);
 	}
 };
 
-inline FPoint ProjectPointOnPlane(const FPoint& Point, const FPoint& Origin, const FPoint& InNormal, double& OutDistance)
-{
-	FPoint Normal = InNormal;
-	ensureCADKernel(!FMath::IsNearlyZero(Normal.Length()));
-	Normal.Normalize();
-
-	FPoint OP = Point - Origin;
-	OutDistance = OP | Normal;
-
-	return Point - (Normal * OutDistance);
-}
-
-/**
- * @return InSegmentA + (InSegmentB - InSegmentA) | InCoordinate
- */
-template<class PointType>
-inline PointType PointOnSegment(const PointType& InSegmentA, const PointType& InSegmentB, double InCoordinate)
-{
-	PointType Segment = InSegmentB - InSegmentA;
-	return InSegmentA + Segment * InCoordinate;
-}
-
-/**
- * @return the distance between the point and the segment. If the projection of the point on the segment
- * is not inside it, return the distance of the point to nearest the segment extremity
- */
-template<class PointType>
-inline double DistanceOfPointToSegment(const PointType& Point, const PointType& SegmentPoint1, const PointType& SegmentPoint2)
-{
-	double Coordinate;
-	return PointType::Distance(ProjectPointOnSegment(Point, SegmentPoint1, SegmentPoint2, Coordinate, /*bRestrictCoodinateToInside*/ true), Point);
-}
-
-/**
- * @return the distance between the point and the segment. If the projection of the point on the segment
- * is not inside it, return the distance of the point to nearest the segment extremity
- */
-template<class PointType>
-inline double SquareDistanceOfPointToSegment(const PointType& Point, const PointType& SegmentPoint1, const PointType& SegmentPoint2)
-{
-	double Coordinate;
-	return FPoint2D::DistSquared(ProjectPointOnSegment(Point, SegmentPoint1, SegmentPoint2, Coordinate, /*bRestrictCoodinateToInside*/ true), Point);
-}
-
-/**
- * @return the distance between the point and the line i.e. the distance between the point and its projection on the line
- */
-template<class PointType>
-inline double DistanceOfPointToLine(const PointType& Point, const PointType& LinePoint1, const PointType& LineDirection)
-{
-	double Coordinate;
-	return ProjectPointOnSegment(Point, LinePoint1, LinePoint1 + LineDirection, Coordinate, /*bRestrictCoodinateToInside*/ false).Distance(Point);
-}
-
-CADKERNEL_API double ComputeCurvature(const FPoint& Gradient, const FPoint& Laplacian);
-CADKERNEL_API double ComputeCurvature(const FPoint& normal, const FPoint& Gradient, const FPoint& Laplacian);
-
-/**
- * @param OutCoordinate: the coordinate of the projected point in the segment AB (coodinate of A = 0 and of B = 1)
- * @return Projected point
- */
-template<class PointType>
-inline PointType ProjectPointOnSegment(const PointType& Point, const PointType& InSegmentA, const PointType& InSegmentB, double& OutCoordinate, bool bRestrictCoodinateToInside = true)
-{
-	PointType Segment = InSegmentB - InSegmentA;
-
-	double SquaredLength = Segment | Segment;
-
-	if (SquaredLength <= 0.0)
-	{
-		OutCoordinate = 0.0;
-		return InSegmentA;
-	}
-	else
-	{
-		PointType APoint = Point - InSegmentA;
-		OutCoordinate = (APoint | Segment) / SquaredLength;
-
-		if (bRestrictCoodinateToInside)
-		{
-			if (OutCoordinate < 0.0)
-			{
-				OutCoordinate = 0.0;
-				return InSegmentA;
-			}
-
-			if (OutCoordinate > 1.0)
-			{
-				OutCoordinate = 1.0;
-				return InSegmentB;
-			}
-		}
-
-		PointType ProjectedPoint = Segment * OutCoordinate;
-		ProjectedPoint += InSegmentA;
-		return ProjectedPoint;
-	}
-}
-
-/**
- * @return Coordinate of the projected point in the segment AB (coordinate of A = 0 and of B = 1)
- */
-template<class PointType>
-inline double CoordinateOfProjectedPointOnSegment(const PointType& Point, const PointType& InSegmentA, const PointType& InSegmentB, bool bRestrictCoodinateToInside = true)
-{
-	PointType Segment = InSegmentB - InSegmentA;
-
-	double SquaredLength = Segment | Segment;
-
-	if (SquaredLength <= 0.0)
-	{
-		return 0.0;
-	}
-	else
-	{
-		PointType APoint = Point - InSegmentA;
-		double Coordinate = (APoint | Segment) / SquaredLength;
-
-		if (bRestrictCoodinateToInside)
-		{
-			if (Coordinate < 0.0)
-			{
-				return 0.0;
-			}
-
-			if (Coordinate > 1.0)
-			{
-				return 1.0;
-			}
-		}
-
-		return Coordinate;
-	}
-}
-
-CADKERNEL_API void FindLoopIntersectionsWithIso(const EIso Iso, const double IsoParameter, const TArray<TArray<FPoint2D>>& Loops, TArray<double>& OutIntersections);
-
 /**
  * The segments must intersect because no check is done
  */
-inline FPoint2D FindIntersectionOfSegments2D(const FSegment2D& SegmentAB, const FSegment2D& SegmentCD, double& OutABIntersectionCoordinate)
+inline FVector2d FindIntersectionOfSegments2D(const FSegment2D& SegmentAB, const FSegment2D& SegmentCD, double& OutABIntersectionCoordinate)
 {
-	const FPoint2D AB = SegmentAB[1] - SegmentAB[0];
-	const FPoint2D DC = SegmentCD[0] - SegmentCD[1];
-	const FPoint2D AC = SegmentCD[0] - SegmentAB[0];
+	const FVector2d AB = SegmentAB[1] - SegmentAB[0];
+	const FVector2d DC = SegmentCD[0] - SegmentCD[1];
+	const FVector2d AC = SegmentCD[0] - SegmentAB[0];
 
 	double ParallelCoef = DC ^ AB;
 	if (FMath::IsNearlyZero(ParallelCoef))
@@ -452,7 +452,7 @@ inline FPoint2D FindIntersectionOfSegments2D(const FSegment2D& SegmentAB, const
 		const double SquareAB = AB | AB;
 		double CCoordinate = (AB | AC) / SquareAB;
 
-		const FPoint2D AD = SegmentCD[1] - SegmentAB[0];
+		const FVector2d AD = SegmentCD[1] - SegmentAB[0];
 		double DCoordinate = (AB | AD) / SquareAB;
 
 		if (CCoordinate >= -DOUBLE_KINDA_SMALL_NUMBER && CCoordinate <= 1 + DOUBLE_KINDA_SMALL_NUMBER)
@@ -489,7 +489,7 @@ inline FPoint2D FindIntersectionOfSegments2D(const FSegment2D& SegmentAB, const
 /**
  * The segments must intersect because no check is done
  */
-inline FPoint2D FindIntersectionOfSegments2D(const FSegment2D& SegmentAB, const FSegment2D& SegmentCD)
+inline FVector2d FindIntersectionOfSegments2D(const FSegment2D& SegmentAB, const FSegment2D& SegmentCD)
 {
 	double ABIntersectionCoordinate;
 	return FindIntersectionOfSegments2D(SegmentAB, SegmentCD, ABIntersectionCoordinate);
@@ -498,14 +498,14 @@ inline FPoint2D FindIntersectionOfSegments2D(const FSegment2D& SegmentAB, const
 /**
  * @return false if the lines are parallel
  */
-inline bool FindIntersectionOfLines2D(const FSegment2D& LineAB, const FSegment2D& LineCD, FPoint2D& OutIntersectionPoint)
+inline bool FindIntersectionOfLines2D(const FSegment2D& LineAB, const FSegment2D& LineCD, FVector2d& OutIntersectionPoint)
 {
 	constexpr const double Min = -DOUBLE_SMALL_NUMBER;
 	constexpr const double Max = 1. + DOUBLE_SMALL_NUMBER;
 
-	const FPoint2D AB = LineAB[1] - LineAB[0];
-	const FPoint2D DC = LineCD[0] - LineCD[1];
-	const FPoint2D AC = LineCD[0] - LineAB[0];
+	const FVector2d AB = LineAB[1] - LineAB[0];
+	const FVector2d DC = LineCD[0] - LineCD[1];
+	const FVector2d AC = LineCD[0] - LineAB[0];
 
 	double ParallelCoef = DC ^ AB;
 	if (FMath::IsNearlyZero(ParallelCoef))
@@ -527,8 +527,8 @@ CADKERNEL_API bool DoIntersectInside(const FSegment2D& SegmentAB, const FSegment
 
 inline bool AreParallel(const FSegment2D& SegmentAB, const FSegment2D& SegmentCD)
 {
-	const FPoint2D AB = SegmentAB.GetVector().GetSafeNormal();
-	const FPoint2D CD = SegmentCD.GetVector().GetSafeNormal();
+	const FVector2d AB = SegmentAB.GetVector().GetSafeNormal();
+	const FVector2d CD = SegmentCD.GetVector().GetSafeNormal();
 	const double ParallelCoef = AB ^ CD;
 	return (FMath::IsNearlyZero(ParallelCoef, DOUBLE_KINDA_SMALL_NUMBER));
 };
diff --git a/Engine/Source/Runtime/Datasmith/CADKernel/Base/Public/Math/MatrixH.h b/Engine/Source/Runtime/Datasmith/CADKernel/Base/Public/Math/MatrixH.h
--- a/Engine/Source/Runtime/Datasmith/CADKernel/Base/Public/Math/MatrixH.h
+++ b/Engine/Source/Runtime/Datasmith/CADKernel/Base/Public/Math/MatrixH.h
@@ -43,7 +43,7 @@ public:
 		}
 	}
 
-	FMatrixH(const FPoint& Origin, const FPoint& Ox, const FPoint& Oy, const FPoint& Oz)
+	FMatrixH(const FVector& Origin, const FVector& Ox, const FVector& Oy, const FVector& Oz)
 	{
 		BuildChangeOfCoordinateSystemMatrix(Ox, Oy, Oz, Origin);
 	}
@@ -65,7 +65,7 @@ public:
 		}
 	}
 
-	void FromAxisOrigin(const FPoint& Axis, const FPoint& Origin);
+	void FromAxisOrigin(const FVector& Axis, const FVector& Origin);
 
 	/**
 	 * @param Xaxis the vector X of the old coordinate system in the new coordinate system
@@ -74,24 +74,26 @@ public:
 	 * @param Origin Origin of the old coordinate system in the new coordinate system
 	 * @return the transform matrix
 	 */
-	void BuildChangeOfCoordinateSystemMatrix(const FPoint& Xaxis, const FPoint& Yaxis, const FPoint& Zaxis, const FPoint& Origin);
+	void BuildChangeOfCoordinateSystemMatrix(const FVector& Xaxis, const FVector& Yaxis, const FVector& Zaxis, const FVector& Origin);
 
-	FPoint Multiply(const FPoint& InPoint) const
+	FVector Multiply(const FVector& InPoint) const
 	{
-		return FPoint(
+		return FVector(
 			InPoint.X * Get(0, 0) + InPoint.Y * Get(0, 1) + InPoint.Z * Get(0, 2) + Get(0, 3),
 			InPoint.X * Get(1, 0) + InPoint.Y * Get(1, 1) + InPoint.Z * Get(1, 2) + Get(1, 3),
 			InPoint.X * Get(2, 0) + InPoint.Y * Get(2, 1) + InPoint.Z * Get(2, 2) + Get(2, 3)
 		);
 	}
 
-	FPoint MultiplyVector(const FPoint& InVector) const
+	FVector Multiply(const FVector2d& Point2D) const
 	{
-		return FPoint(
-			InVector.X * Get(0, 0) + InVector.Y * Get(0, 1) + InVector.Z * Get(0, 2),
-			InVector.X * Get(1, 0) + InVector.Y * Get(1, 1) + InVector.Z * Get(1, 2),
-			InVector.X * Get(2, 0) + InVector.Y * Get(2, 1) + InVector.Z * Get(2, 2)
-		);
+		return Multiply(FVector(Point2D, 0.));
+	}
+
+	FVector2d Multiply2D(const FVector2d& Point2D) const
+	{
+		const FVector Point = Multiply(FVector(Point2D, 0.));
+		return FVector2d(Point.X, Point.Y);
 	}
 
 	FVector MultiplyVector(const FVector& InVector) const
@@ -103,6 +105,17 @@ public:
 		);
 	}
 
+	FVector MultiplyVector(const FVector2d& Point2D) const
+	{
+		return MultiplyVector(FVector(Point2D, 0.));
+	}
+
+	FVector2d MultiplyVector2D(const FVector2d& Point2D) const
+	{
+		const FVector Point = MultiplyVector(FVector(Point2D, 0.));
+		return FVector2d(Point.X, Point.Y);
+	}
+
 	FVector3f MultiplyVector(const FVector3f& InVector) const
 	{
 		return FVector3f(
@@ -112,29 +125,16 @@ public:
 		);
 	}
 
-	static FMatrixH MakeRotationMatrix(double InAngle, const FPoint InAxe);
+	static FMatrixH MakeRotationMatrix(double InAngle, const FVector InAxe);
 
-	static FMatrixH MakeTranslationMatrix(const FPoint& InPoint);
+	static FMatrixH MakeTranslationMatrix(const FVector& InPoint);
 
-	static FMatrixH MakeScaleMatrix(FPoint& Scale);
+	static FMatrixH MakeScaleMatrix(FVector& Scale);
 	static FMatrixH MakeScaleMatrix(double ScaleX, double ScaleY, double ScaleZ);
 
 	/**
 	 * Apply the rotation centered in origin to PointToRotate
 	 */
-	FPoint PointRotation(const FPoint& PointToRotate, const FPoint& Origin) const
-	{
-		FPoint Result = Origin;
-		for (int32 Index = 0; Index < 3; Index++)
-		{
-			for (int32 Jndex = 0; Jndex < 3; Jndex++)
-			{
-				Result[Index] += Get(Index, Jndex) * (PointToRotate[Jndex] - Origin[Jndex]);
-			}
-		}
-		return Result;
-	}
-
 	FVector PointRotation(const FVector& PointToRotate, const FVector& Origin) const
 	{
 		FVector Result = Origin;
@@ -148,9 +148,9 @@ public:
 		return Result;
 	}
 
-	FPoint2D PointRotation(const FPoint2D& PointToRotate, const FPoint2D& Origin) const
+	FVector2d PointRotation(const FVector2d& PointToRotate, const FVector2d& Origin) const
 	{
-		FPoint2D Result = Origin;
+		FVector2d Result = Origin;
 		for (int32 Index = 0; Index < 2; Index++)
 		{
 			for (int32 Jndex = 0; Jndex < 2; Jndex++)
@@ -231,7 +231,7 @@ public:
 		Move(*this, Result);
 	}
 
-	FPoint operator*(const FPoint& Point) const
+	FVector operator*(const FVector& Point) const
 	{
 		return Multiply(Point);
 	}
@@ -251,14 +251,14 @@ public:
 		memcpy(OutMatrix, Matrix, 16 * sizeof(double));
 	}
 
-	FPoint Column(int32 Index) const
+	FVector Column(int32 Index) const
 	{
-		return FPoint(Get(0, Index), Get(1, Index), Get(2, Index));
+		return FVector(Get(0, Index), Get(1, Index), Get(2, Index));
 	}
 
-	FPoint Row(int32 Index) const
+	FVector Row(int32 Index) const
 	{
-		return FPoint(Get(Index, 0), Get(Index, 1), Get(Index, 2));
+		return FVector(Get(Index, 0), Get(Index, 1), Get(Index, 2));
 	}
 
 	void Print(EVerboseLevel level) const;
diff --git a/Engine/Source/Runtime/Datasmith/CADKernel/Base/Public/Math/Point.h b/Engine/Source/Runtime/Datasmith/CADKernel/Base/Public/Math/Point.h
--- a/Engine/Source/Runtime/Datasmith/CADKernel/Base/Public/Math/Point.h
+++ b/Engine/Source/Runtime/Datasmith/CADKernel/Base/Public/Math/Point.h
@@ -4,328 +4,10 @@
 #include "Core/Types.h"
 #include "Math/MathConst.h"
 
-// this will eventually be reconciled with what the Geometry team is using.
 namespace UE::CADKernel
 {
 
-struct FPoint2D;
-
-#if 1
-struct FPoint : public FVector
-{
-public:
-	explicit FPoint(double InCoordX = 0., double InCoordY = 0., double InCoordZ = 0.)
-		: FVector(InCoordX, InCoordY, InCoordZ)
-	{
-	}
-
-	FPoint(const FPoint& Point)
-		: FVector(Point.X, Point.Y, Point.Z)
-	{
-	}
-
-	FPoint(const FVector& Point)
-		: FVector(Point.X, Point.Y, Point.Z)
-	{
-	}
-	
-	FPoint(const FVector2d& InPoint2D)
-		: FVector(InPoint2D.X, InPoint2D.Y, 0.)
-	{
-	}
-
-	FPoint(const double* const InPoint)
-		: FVector(InPoint[0], InPoint[1], InPoint[2])
-	{
-	}
-
-	operator FVector3f() const
-	{
-		return FVector3f((float)X, (float)Y, (float)Z);
-	}
-
-	static double Distance(const FPoint& A, const FPoint& B)
-	{
-		return FVector::Distance(A, B);
-	}
-
-	static double DistSquared(const FPoint& A, const FPoint& B)
-	{
-		return FVector::DistSquared(A, B);
-	}
-
-	friend FArchive& operator<<(FArchive& Ar, FPoint& Point)
-	{
-		Ar.Serialize(&Point, sizeof(FPoint));
-		return Ar;
-	}
-
-	// TO REMOVE
-
-	FPoint operator+(const FPoint& Point) const
-	{
-		return FPoint(X + Point.X, Y + Point.Y, Z + Point.Z);
-	}
-
-	FPoint operator+(const FVector& Point) const
-	{
-		return FPoint(X + Point.X, Y + Point.Y, Z + Point.Z);
-	}
-
-	FPoint operator+(const FVector3f& Point) const
-	{
-		return FPoint(X + Point.X, Y + Point.Y, Z + Point.Z);
-	}
-
-	FPoint operator-(const FPoint& Point) const
-	{
-		return FPoint(X - Point.X, Y - Point.Y, Z - Point.Z);
-	}
-
-	FPoint operator^(const FPoint& Point) const
-	{
-		return FPoint(Y * Point.Z - Z * Point.Y, Z * Point.X - X * Point.Z, X * Point.Y - Y * Point.X);
-	}
-
-	double operator|(const FPoint& Point) const
-	{
-		return X * Point.X + Y * Point.Y + Z * Point.Z;
-	}
-
-	FPoint operator*(double Scale) const
-	{
-		return FPoint(X * Scale, Y * Scale, Z * Scale);
-	}
-
-	FPoint operator/(double InvScale) const
-	{
-		if (FMath::IsNearlyZero(InvScale))
-		{
-			return *this;
-		}
-		return FPoint(X / InvScale, Y / InvScale, Z / InvScale);
-	}
-
-	FPoint operator-() const
-	{
-		return FPoint(-X, -Y, -Z);
-	}
-
-	FPoint& operator=(const FPoint& Point)
-	{
-		X = Point.X;
-		Y = Point.Y;
-		Z = Point.Z;
-		return *this;
-	}
-
-	FPoint& operator+=(const FPoint& Point)
-	{
-		X += Point.X;
-		Y += Point.Y;
-		Z += Point.Z;
-		return *this;
-	}
-
-	FPoint& operator+=(const FVector& Point)
-	{
-		X += Point.X;
-		Y += Point.Y;
-		Z += Point.Z;
-		return *this;
-	}
-
-	FPoint& operator+=(const FVector3f& Point)
-	{
-		X += Point.X;
-		Y += Point.Y;
-		Z += Point.Z;
-		return *this;
-	}
-
-	FPoint& operator-=(const FPoint& Point)
-	{
-		X -= Point.X;
-		Y -= Point.Y;
-		Z -= Point.Z;
-		return *this;
-	}
-
-	FPoint& operator*=(double Scale)
-	{
-		X *= Scale;
-		Y *= Scale;
-		Z *= Scale;
-		return *this;
-	}
-
-	FPoint& operator/=(double InvScale)
-	{
-		if (FMath::IsNearlyZero(InvScale))
-		{
-			return *this;
-		}
-
-		X /= InvScale;
-		Y /= InvScale;
-		Z /= InvScale;
-		return *this;
-	}
-
-	//bool operator<(const FPoint& Other) const
-	//{
-	//	int32 Compare = RealCompare(X, Other.X);
-	//	if (Compare < 0)
-	//	{
-	//		return true;
-	//	}
-	//	if (Compare > 0)
-	//	{
-	//		return false;
-	//	}
-
-	//	Compare = RealCompare(Y, Other.Y);
-	//	if (Compare < 0)
-	//	{
-	//		return true;
-	//	}
-	//	if (Compare > 0)
-	//	{
-	//		return false;
-	//	}
-
-	//	Compare = RealCompare(Z, Other.Z);
-	//	return (Compare < 0);
-	//}
-
-	//void SetMin(const FPoint& Point)
-	//{
-	//	X = FMath::Min(X, Point.X);
-	//	Y = FMath::Min(Y, Point.Y);
-	//	Z = FMath::Min(Z, Point.Z);
-	//}
-
-	//void SetMax(const FPoint& Point)
-	//{
-	//	X = FMath::Max(X, Point.X);
-	//	Y = FMath::Max(Y, Point.Y);
-	//	Z = FMath::Max(Z, Point.Z);
-	//}
-
-	//void Set(double InCoordX = 0., double InCoordY = 0., double InCoordZ = 0.)
-	//{
-	//	X = InCoordX;
-	//	Y = InCoordY;
-	//	Z = InCoordZ;
-	//}
-
-	//void Set(double* InCoordinates)
-	//{
-	//	X = InCoordinates[0];
-	//	Y = InCoordinates[1];
-	//	Z = InCoordinates[2];
-	//}
-
-	//double SquareLength() const
-	//{
-	//	return (X * X) + (Y * Y) + (Z * Z);
-	//}
-
-	//FPoint& Normalize()
-	//{
-	//	double Norm = Length();
-	//	if (FMath::IsNearlyZero(Norm))
-	//	{
-	//		*this = FPoint(1.,1.,1.).Normalize();
-	//	}
-	//	else
-	//	{
-	//		*this /= Norm;
-	//	}
-	//	return *this;
-	//}
-
-	//FPoint Normalize() const 
-	//{
-	//	double Norm = Length();
-	//	if (FMath::IsNearlyZero(Norm))
-	//	{
-	//		return FPoint(1.,1.,1.).Normalize();
-	//	}
-	//	else
-	//	{
-	//		return *this / Norm;
-	//	}
-	//}
-
-	//FPoint Middle(const FPoint& Point) const
-	//{
-	//	return FPoint(0.5 * (X + Point.X), 0.5 * (Y + Point.Y), 0.5 * (Z + Point.Z));
-	//}
-
-	//double Distance(const FPoint& Point) const
-	//{
-	//	return FMath::Sqrt(FMath::Square(Point.X - X) + FMath::Square(Point.Y - Y) + FMath::Square(Point.Z - Z));
-	//}
-
-	//double SquareDistance(const FPoint& Point) const
-	//{
-	//	return FMath::Square(Point.X - X) + FMath::Square(Point.Y - Y) + FMath::Square(Point.Z - Z);
-	//}
-};
-
-struct FPoint2D : public FVector2d
-{
-public:
-	FPoint2D()
-		: FVector2d(0., 0.)
-	{
-	}
-
-	FPoint2D(double NewX, double NewY)
-		: FVector2d(NewX, NewY)
-	{
-	}
-
-	FPoint2D(const FPoint& Point)
-		: FVector2d(Point.X, Point.Y)
-	{
-	}
-
-	FPoint2D(const FPoint2D& Point)
-		: FVector2d(Point.X, Point.Y)
-	{
-	}
-
-	FPoint2D(const FVector2d& Point)
-		: FVector2d(Point.X, Point.Y)
-	{
-	}
-
-	static double Distance(const FPoint2D& A, const FPoint2D& B)
-	{
-		return FVector2d::Distance(A, B);
-	}
-
-	static double DistSquared(const FPoint2D& A, const FPoint2D& B)
-	{
-		return FVector2d::DistSquared(A, B);
-	}
-
-	friend FArchive& operator<<(FArchive& Ar, FPoint2D& Point)
-	{
-		Ar.Serialize(&Point, sizeof(FPoint2D));
-		return Ar;
-	}
-
-	void Set(double NewX, double NewY)
-	{
-		FVector2d::Set(NewX, NewY);
-	}
-};
-#endif
-
-typedef FPoint2D FSurfacicTolerance;
+typedef FVector2d FSurfacicTolerance;
 
 class FVectorUtil
 {
@@ -333,47 +15,9 @@ public:
 	static CADKERNEL_API const FVector FarawayPoint3D;
 	static CADKERNEL_API const FVector2d FarawayPoint2D;
 
-	static FPoint2D FromVector(const FPoint& Src)
-	{
-		return FPoint2D(Src.X, Src.Y);
-	}
-
-	static FPoint FromBuffer(const double* const Src)
-	{
-		return FPoint(Src[0], Src[1], Src[2]);
-	}
-
-	static void SetPoint(FPoint& Point, const double* const Src)
+	static FVector2d FromVector(const FVector& Src)
 	{
-		Point.X = Src[0];
-		Point.Y = Src[1];
-		Point.Z = Src[2];
-	}
-
-	static void SetPoint(FPoint& Point, double X, double Y, double Z)
-	{
-		Point.X = X;
-		Point.Y = Y;
-		Point.Z = Z;
-	}
-
-	static void SetPoint(FPoint& Point, const FPoint2D& Point2D)
-	{
-		Point.X = Point2D.X;
-		Point.Y = Point2D.Y;
-		Point.Z = 0.;
-	}
-
-	static void SetPoint2D(FPoint2D& Point, const double* const Src)
-	{
-		Point.X = Src[0];
-		Point.Y = Src[1];
-	}
-
-	static void SetPoint2D(FPoint2D& Point, double X, double Y)
-	{
-		Point.X = X;
-		Point.Y = Y;
+		return FVector2d(Src.X, Src.Y);
 	}
 
 	template<typename VectorType>
@@ -390,10 +34,11 @@ public:
 		return FMath::Max(-1.0, FMath::Min(Cosinus, 1.0));
 	}
 
-	static double ComputeAngle(const FPoint& Vec1, const FPoint& Vec2)
+	template<typename VectorType>
+	static double ComputeAngle(const VectorType& Vec1, const VectorType& Vec2)
 	{
-		FPoint ThisNormalized(Vec1);
-		FPoint OtherNormalized(Vec2);
+		VectorType ThisNormalized(Vec1);
+		VectorType OtherNormalized(Vec2);
 
 		ThisNormalized.Normalize();
 		OtherNormalized.Normalize();
@@ -418,8 +63,3 @@ public:
 
 }
 
-inline uint32 GetTypeHash(const UE::CADKernel::FPoint& Point)
-{
-	return HashCombine(GetTypeHash(Point.X), HashCombine(GetTypeHash(Point.Y), GetTypeHash(Point.Z)));
-}
-
diff --git a/Engine/Source/Runtime/Datasmith/CADKernel/Base/Public/Math/SlopeUtils.h b/Engine/Source/Runtime/Datasmith/CADKernel/Base/Public/Math/SlopeUtils.h
--- a/Engine/Source/Runtime/Datasmith/CADKernel/Base/Public/Math/SlopeUtils.h
+++ b/Engine/Source/Runtime/Datasmith/CADKernel/Base/Public/Math/SlopeUtils.h
@@ -14,7 +14,7 @@ namespace UE::CADKernel
  *
  * This file propose all useful methods to use slope instead of angle.
  *
- * The method "ComputeSlope(const FPoint2D& StartPoint, const FPoint2D& EndPoint) is the main method.
+ * The method "ComputeSlope(const FVector2d& StartPoint, const FVector2d& EndPoint) is the main method.
  * It compute the slope between the input segment defined by two points and [0, u) axis.
  * The return value is a real in the interval [0, 8] for an angle in the interval [0, 2Pi]
  *
@@ -115,7 +115,7 @@ constexpr double Epsilon = 0.001;
 
 }
 
-typedef TFunction<double(const FPoint2D&, const FPoint2D&, double)> SlopeMethod;
+typedef TFunction<double(const FVector2d&, const FVector2d&, double)> SlopeMethod;
 
 /**
  * Transform a positive slope into an oriented slope [-4, 4] i.e. an equivalent angle between [-Pi, Pi]
@@ -177,7 +177,7 @@ inline double SwapSlopeOrientation(double Slope)
 	return TransformIntoPositiveSlope(SwapedSlope);
 }
 
-inline double ComputeSlope(const FPoint2D& StartPoint, const FPoint2D& EndPoint)
+inline double ComputeSlope(const FVector2d& StartPoint, const FVector2d& EndPoint)
 {
 	double DeltaU = EndPoint.X - StartPoint.X;
 	double DeltaV = EndPoint.Y - StartPoint.Y;
@@ -265,7 +265,7 @@ inline double ComputeSlope(const FPoint2D& StartPoint, const FPoint2D& EndPoint)
  * Compute the slope of a segment according to a reference slope
  * return the slope
  */
-inline double ComputeSlope(const FPoint2D& StartPoint, const FPoint2D& EndPoint, double ReferenceSlope)
+inline double ComputeSlope(const FVector2d& StartPoint, const FVector2d& EndPoint, double ReferenceSlope)
 {
 	const double Slope = ComputeSlope(StartPoint, EndPoint);
 	return Slope - ReferenceSlope;
@@ -275,7 +275,7 @@ inline double ComputeSlope(const FPoint2D& StartPoint, const FPoint2D& EndPoint,
  * Compute the slope between the segments [StartPoint, EndPoint1] and [StartPoint, EndPoint2]
  * @return a slope i.e. an equivalent angle
  */
-inline double ComputeSlope(const FPoint2D& StartPoint, const FPoint2D& EndPoint1, const FPoint2D& EndPoint2)
+inline double ComputeSlope(const FVector2d& StartPoint, const FVector2d& EndPoint1, const FVector2d& EndPoint2)
 {
 	const double ReferenceSlope = ComputeSlope(StartPoint, EndPoint1);
 	const double Slope = ComputeSlope(StartPoint, EndPoint2);
@@ -287,7 +287,7 @@ inline double ComputeSlope(const FPoint2D& StartPoint, const FPoint2D& EndPoint1
  * This method is used to compute an approximation of the angle between two segments in 2D.
  * return a slope between [0, 8] i.e. an equivalent angle between [0, 2Pi]
  */
-inline double ComputePositiveSlope(const FPoint2D& StartPoint, const FPoint2D& EndPoint, double ReferenceSlope)
+inline double ComputePositiveSlope(const FVector2d& StartPoint, const FVector2d& EndPoint, double ReferenceSlope)
 {
 	double Slope = ComputeSlope(StartPoint, EndPoint, ReferenceSlope);
 	return TransformIntoPositiveSlope(Slope);
@@ -297,18 +297,18 @@ inline double ComputePositiveSlope(const FPoint2D& StartPoint, const FPoint2D& E
  * Compute the positive slope between the segments [StartPoint, EndPoint1] and [StartPoint, EndPoint2]
  * @return a slope between [0, 8] i.e. an equivalent angle between [0, 2Pi]
  */
-inline double ComputePositiveSlope(const FPoint2D& StartPoint, const FPoint2D& EndPoint1, const FPoint2D& EndPoint2)
+inline double ComputePositiveSlope(const FVector2d& StartPoint, const FVector2d& EndPoint1, const FVector2d& EndPoint2)
 {
 	double Slope = ComputeSlope(StartPoint, EndPoint1, EndPoint2);
 	return TransformIntoPositiveSlope(Slope);
 }
 
-inline double ClockwiseSlope(const FPoint2D& StartPoint, const FPoint2D& EndPoint, double ReferenceSlope)
+inline double ClockwiseSlope(const FVector2d& StartPoint, const FVector2d& EndPoint, double ReferenceSlope)
 {
 	return TransformIntoClockwiseSlope(ComputePositiveSlope(StartPoint, EndPoint, ReferenceSlope));
 }
 
-inline double CounterClockwiseSlope(const FPoint2D& StartPoint, const FPoint2D& EndPoint, double ReferenceSlope)
+inline double CounterClockwiseSlope(const FVector2d& StartPoint, const FVector2d& EndPoint, double ReferenceSlope)
 {
 	return ComputePositiveSlope(StartPoint, EndPoint, ReferenceSlope);
 }
@@ -317,7 +317,7 @@ inline double CounterClockwiseSlope(const FPoint2D& StartPoint, const FPoint2D&
  * Compute the oriented slope of a segment according to a reference slope
  * @return a slope between [-4, 4] i.e. an equivalent angle between [-Pi, Pi]
  */
-inline double ComputeOrientedSlope(const FPoint2D& StartPoint, const FPoint2D& EndPoint, double ReferenceSlope)
+inline double ComputeOrientedSlope(const FVector2d& StartPoint, const FVector2d& EndPoint, double ReferenceSlope)
 {
 	return TransformIntoOrientedSlope(ComputeSlope(StartPoint, EndPoint, ReferenceSlope));
 }
@@ -326,7 +326,7 @@ inline double ComputeOrientedSlope(const FPoint2D& StartPoint, const FPoint2D& E
  * Compute the positive slope between the segments [StartPoint, EndPoint1] and [StartPoint, EndPoint2]
  * @return a slope between [-4, 4] i.e. an equivalent angle between [-Pi, Pi]
  */
-inline double ComputeOrientedSlope(const FPoint2D& StartPoint, const FPoint2D& EndPoint1, const FPoint2D& EndPoint2)
+inline double ComputeOrientedSlope(const FVector2d& StartPoint, const FVector2d& EndPoint1, const FVector2d& EndPoint2)
 {
 	return TransformIntoOrientedSlope(ComputeSlope(StartPoint, EndPoint1, EndPoint2));
 }
@@ -334,7 +334,7 @@ inline double ComputeOrientedSlope(const FPoint2D& StartPoint, const FPoint2D& E
 /**
  * return a slope between [0, 4] i.e. an angle between [0, Pi]
  */
-inline double ComputeUnorientedSlope(const FPoint2D& StartPoint, const FPoint2D& EndPoint, double ReferenceSlope)
+inline double ComputeUnorientedSlope(const FVector2d& StartPoint, const FVector2d& EndPoint, double ReferenceSlope)
 {
 	const double Slope = ComputeSlope(StartPoint, EndPoint, ReferenceSlope);
 	return TransformIntoUnorientedSlope(Slope);
@@ -346,7 +346,7 @@ inline double ComputeUnorientedSlope(const FPoint2D& StartPoint, const FPoint2D&
  * ComputeUnorientedSlope => 2.3 return 0.3
  * ComputeUnorientedSlope => 3.6 return 0.4
  */
-inline double ComputeSlopeRelativeToNearestAxis(const FPoint2D& StartPoint, const FPoint2D& EndPoint)
+inline double ComputeSlopeRelativeToNearestAxis(const FVector2d& StartPoint, const FVector2d& EndPoint)
 {
 	double Slope = TransformIntoUnorientedSlope(ComputeSlope(StartPoint, EndPoint));
 	if (Slope > Slope::RightSlope)
@@ -370,7 +370,7 @@ inline double ComputeSlopeRelativeToNearestAxis(const FPoint2D& StartPoint, cons
  * ComputeUnorientedSlope => 0.5 return 0.5
  * ComputeUnorientedSlope => 2.3 return 1.7
  */
-inline double ComputeSlopeRelativeToReferenceAxis(const FPoint2D& StartPoint, const FPoint2D& EndPoint, double ReferenceAxisSlope)
+inline double ComputeSlopeRelativeToReferenceAxis(const FVector2d& StartPoint, const FVector2d& EndPoint, double ReferenceAxisSlope)
 {
 	const double Slope = ComputeSlope(StartPoint, EndPoint, ReferenceAxisSlope);
 	return TransformIntoSlopeRelativeToReferenceAxis(Slope);
@@ -379,7 +379,7 @@ inline double ComputeSlopeRelativeToReferenceAxis(const FPoint2D& StartPoint, co
 /**
  * return a slope between [0, 4] i.e. an angle between [0, Pi]
  */
-inline double ComputeUnorientedSlope(const FPoint2D& StartPoint, const FPoint2D& EndPoint1, const FPoint2D& EndPoint2)
+inline double ComputeUnorientedSlope(const FVector2d& StartPoint, const FVector2d& EndPoint1, const FVector2d& EndPoint2)
 {
 	return TransformIntoUnorientedSlope(ComputeSlope(StartPoint, EndPoint1, EndPoint2));
 }
@@ -398,7 +398,7 @@ inline double ComputeUnorientedSlope(const FPoint2D& StartPoint, const FPoint2D&
  * Return true if the segment BP is inside the sector defined the half-lines [BA) and [BC) in the counterclockwise.
  * Return false if ABP angle or PBC angle is too flat (smaller than FlatAngle)
  */
-inline bool IsPointPInsideSectorABC(const FPoint2D& PointA, const FPoint2D& PointB, const FPoint2D& PointC, const FPoint2D& PointP, const double FlatAngle)
+inline bool IsPointPInsideSectorABC(const FVector2d& PointA, const FVector2d& PointB, const FVector2d& PointC, const FVector2d& PointP, const double FlatAngle)
 {
 	double SlopWithNextBoundary = ComputeSlope(PointB, PointC);
 	double BoundaryDeltaSlope = ComputePositiveSlope(PointB, PointA, SlopWithNextBoundary);
@@ -423,12 +423,12 @@ inline bool IsPointPInsideSectorABC(const FPoint2D& PointA, const FPoint2D& Poin
  *
  * Return true if all of the segment BPi is inside the sector defined the half-lines [BA) and [BC) in the counterclockwise.
  */
-inline bool ArePointsInsideSectorABC(const FPoint2D& PointA, const FPoint2D& PointB, const FPoint2D& PointC, const TArray<const FPoint2D*>& Points, const double FlatAngle = -DOUBLE_SMALL_NUMBER)
+inline bool ArePointsInsideSectorABC(const FVector2d& PointA, const FVector2d& PointB, const FVector2d& PointC, const TArray<const FVector2d*>& Points, const double FlatAngle = -DOUBLE_SMALL_NUMBER)
 {
 	double SlopWithNextBoundary = ComputeSlope(PointB, PointC);
 	double BoundaryDeltaSlope = ComputePositiveSlope(PointB, PointA, SlopWithNextBoundary);
 
-	return Algo::AllOf(Points, [&](const FPoint2D* PointP) 
+	return Algo::AllOf(Points, [&](const FVector2d* PointP) 
 	{
 		double DeltaU = PointB.X - PointP->X;
 		double DeltaV = PointB.Y - PointP->Y;
@@ -447,10 +447,10 @@ inline bool ArePointsInsideSectorABC(const FPoint2D& PointA, const FPoint2D& Poi
 	});
 }
 
-inline FPoint2D SlopeToVector(const double Slope)
+inline FVector2d SlopeToVector(const double Slope)
 {
 	int32 SlopeStep = (int32)(Slope);
-	FPoint2D Vector;
+	FVector2d Vector = FVector2d::ZeroVector;
 	switch (SlopeStep)
 	{
 	case 0:
diff --git a/Engine/Source/Runtime/Datasmith/CADKernel/Base/Public/Mesh/Criteria/Criterion.h b/Engine/Source/Runtime/Datasmith/CADKernel/Base/Public/Mesh/Criteria/Criterion.h
--- a/Engine/Source/Runtime/Datasmith/CADKernel/Base/Public/Mesh/Criteria/Criterion.h
+++ b/Engine/Source/Runtime/Datasmith/CADKernel/Base/Public/Mesh/Criteria/Criterion.h
@@ -72,10 +72,10 @@ public:
 
 	virtual double Value() const = 0;
 
-	static double EvaluateSag(const FPoint& PointPoint, const FPoint& PointNext, const FPoint& PointMiddle, double& Length)
+	static double EvaluateSag(const FVector& PointPoint, const FVector& PointNext, const FVector& PointMiddle, double& Length)
 	{
-		FPoint ChordVec = PointNext - PointPoint;
-		FPoint MiddleVec = PointMiddle - PointPoint;
+		FVector ChordVec = PointNext - PointPoint;
+		FVector MiddleVec = PointMiddle - PointPoint;
 
 		double Sag = 0.0;
 		double NormSqrVec = ChordVec | ChordVec;
@@ -83,7 +83,7 @@ public:
 
 		if (NormSqrVec > DOUBLE_SMALL_NUMBER)
 		{
-			FPoint VecSag = ChordVec ^ MiddleVec;
+			FVector VecSag = ChordVec ^ MiddleVec;
 			double NormSqrSag = VecSag | VecSag;
 			Sag = NormSqrSag / NormSqrVec;
 
