Replaced FPoint and FPoint2D by FVector and  FVector2d and remove unused FFPoint and FPointH classes - Final

diff --git a/Engine/Source/Runtime/Datasmith/CADKernel/Base/Public/Mesh/Structure/EdgeSegment.h b/Engine/Source/Runtime/Datasmith/CADKernel/Base/Public/Mesh/Structure/EdgeSegment.h
--- a/Engine/Source/Runtime/Datasmith/CADKernel/Base/Public/Mesh/Structure/EdgeSegment.h
+++ b/Engine/Source/Runtime/Datasmith/CADKernel/Base/Public/Mesh/Structure/EdgeSegment.h
@@ -23,7 +23,7 @@ class FEdgeSegment : public FHaveStates
 private:
 	FTopologicalEdge* Edge;
 	double Coordinates[2];
-	FPoint2D USSPoints[2]; // in Uniform Scaled space
+	FVector2d USSPoints[2]; // in Uniform Scaled space
 
 	FEdgeSegment* NextSegment;
 	FEdgeSegment* PreviousSegment;
@@ -59,7 +59,7 @@ public:
 
 	virtual ~FEdgeSegment() = default;
 
-	void SetBoundarySegment(bool bInIsInnerLoop, FTopologicalEdge* InEdge, double InStartU, double InEndU, const FPoint2D& InStartPoint, const FPoint2D& InEndPoint)
+	void SetBoundarySegment(bool bInIsInnerLoop, FTopologicalEdge* InEdge, double InStartU, double InEndU, const FVector2d& InStartPoint, const FVector2d& InEndPoint)
 	{
 		if (bInIsInnerLoop)
 		{
@@ -76,7 +76,7 @@ public:
 		ClosedSegment = nullptr;
 
 		SquareDistanceToClosedSegment = HUGE_VAL;
-		Length = FPoint2D::Distance(USSPoints[ELimit::Start], USSPoints[ELimit::End]);
+		Length = FVector2d::Distance(USSPoints[ELimit::Start], USSPoints[ELimit::End]);
 
 		Id = ++LastId;
 		ChainIndex = Ident::Undefined;
@@ -152,7 +152,7 @@ public:
 	/**
 	 * @return Center of the segment in uniform scaled space
 	 */
-	FPoint2D GetCenter() const
+	FVector2d GetCenter() const
 	{
 		return (USSPoints[ELimit::Start] + USSPoints[ELimit::End]) * 0.5;
 	}
@@ -160,7 +160,7 @@ public:
 	/**
 	 * @return Point a parameter U of the segment in uniform scaled space
 	 */
-	FPoint2D ComputeEdgePoint(double EdgeParamU) const
+	FVector2d ComputeEdgePoint(double EdgeParamU) const
 	{
 		double SegmentParamS = (EdgeParamU - Coordinates[ELimit::Start]) / (Coordinates[ELimit::End] - Coordinates[ELimit::Start]);
 		return USSPoints[ELimit::Start] + (USSPoints[ELimit::End] - USSPoints[ELimit::Start]) * SegmentParamS;
@@ -169,7 +169,7 @@ public:
 	/**
 	 * @return the extremity Point of the segment in uniform scaled space
 	 */
-	const FPoint2D& GetExtemity(const ELimit Limit) const
+	const FVector2d& GetExtemity(const ELimit Limit) const
 	{
 		return USSPoints[Limit];
 	}
@@ -205,7 +205,7 @@ public:
 	/**
 	 * Compute the slope of the Segment defined by the two input points according to this.
 	 */
-	double ComputeUnorientedSlopeOf(const FPoint2D& Middle, const FPoint2D& Projection) const 
+	double ComputeUnorientedSlopeOf(const FVector2d& Middle, const FVector2d& Projection) const 
 	{
 		double ReferenceSlope = ComputeSlope(USSPoints[ELimit::Start], USSPoints[ELimit::End]);
 		return ComputeUnorientedSlope(Projection, Middle, ReferenceSlope);
@@ -214,7 +214,7 @@ public:
 	/**
 	 * Compute the slope of the Segment defined by the two input points according to this.
 	 */
-	double ComputeOrientedSlopeOf(const FPoint2D& Middle, const FPoint2D& Projection) const
+	double ComputeOrientedSlopeOf(const FVector2d& Middle, const FVector2d& Projection) const
 	{
 		double ReferenceSlope = ComputeSlope(USSPoints[ELimit::Start], USSPoints[ELimit::End]);
 		return ComputeOrientedSlope(Projection, Middle, ReferenceSlope);
@@ -279,9 +279,9 @@ public:
 		return FMath::Abs(Coordinates[ELimit::End] - Coordinates[ELimit::Start]) * DeltaLength / Length;
 	}
 
-	FPoint2D ProjectPoint(const FPoint2D& PointToProject, double& SegmentU) const
+	FVector2d ProjectPoint(const FVector2d& PointToProject, double& SegmentU) const
 	{
-		return ProjectPointOnSegment(PointToProject, USSPoints[ELimit::Start], USSPoints[ELimit::End], SegmentU, true);
+		return ProjectPointOnSegment<FVector2d>(PointToProject, USSPoints[ELimit::Start], USSPoints[ELimit::End], SegmentU, true);
 	}
 
 private:
diff --git a/Engine/Source/Runtime/Datasmith/CADKernel/Base/Public/Mesh/Structure/FaceMesh.h b/Engine/Source/Runtime/Datasmith/CADKernel/Base/Public/Mesh/Structure/FaceMesh.h
--- a/Engine/Source/Runtime/Datasmith/CADKernel/Base/Public/Mesh/Structure/FaceMesh.h
+++ b/Engine/Source/Runtime/Datasmith/CADKernel/Base/Public/Mesh/Structure/FaceMesh.h
@@ -57,7 +57,7 @@ public:
 
 	void InverseOrientation();
 
-	void GetNodeIdToCoordinates(TMap<int32, const FPoint*>& NodeIdToCoordinates) const;
+	void GetNodeIdToCoordinates(TMap<int32, const FVector*>& NodeIdToCoordinates) const;
 
 };
 }
diff --git a/Engine/Source/Runtime/Datasmith/CADKernel/Base/Public/Mesh/Structure/Mesh.h b/Engine/Source/Runtime/Datasmith/CADKernel/Base/Public/Mesh/Structure/Mesh.h
--- a/Engine/Source/Runtime/Datasmith/CADKernel/Base/Public/Mesh/Structure/Mesh.h
+++ b/Engine/Source/Runtime/Datasmith/CADKernel/Base/Public/Mesh/Structure/Mesh.h
@@ -21,7 +21,7 @@ protected:
 	int32 StartNodeId;
 	int32 LastNodeIndex;
 
-	TArray<FPoint> NodeCoordinates;
+	TArray<FVector> NodeCoordinates;
 	int32 MeshModelIndex;
 
 public:
@@ -47,12 +47,12 @@ public:
 		return EEntity::Mesh;
 	}
 
-	TArray<FPoint>& GetNodeCoordinates()
+	TArray<FVector>& GetNodeCoordinates()
 	{
 		return NodeCoordinates;
 	}
 
-	const TArray<FPoint>& GetNodeCoordinates() const
+	const TArray<FVector>& GetNodeCoordinates() const
 	{
 		return NodeCoordinates;
 	}
diff --git a/Engine/Source/Runtime/Datasmith/CADKernel/Base/Public/Mesh/Structure/ModelMesh.h b/Engine/Source/Runtime/Datasmith/CADKernel/Base/Public/Mesh/Structure/ModelMesh.h
--- a/Engine/Source/Runtime/Datasmith/CADKernel/Base/Public/Mesh/Structure/ModelMesh.h
+++ b/Engine/Source/Runtime/Datasmith/CADKernel/Base/Public/Mesh/Structure/ModelMesh.h
@@ -23,7 +23,7 @@ class CADKERNEL_API FModelMesh : public FEntityGeom
 
 	TArray<TSharedPtr<FCriterion>> Criteria;
 
-	TArray<TArray<FPoint>*> GlobalPointCloud;
+	TArray<TArray<FVector>*> GlobalPointCloud;
 	FIdent LastIdUsed;
 
 	TArray<FVertexMesh*> VertexMeshes;
@@ -129,7 +129,7 @@ public:
 		FaceMeshes.Add(&Mesh);
 	}
 
-	void RegisterCoordinates(TArray<FPoint>& Coordinates, int32& OutStartVertexId, int32& OutIndex)
+	void RegisterCoordinates(TArray<FVector>& Coordinates, int32& OutStartVertexId, int32& OutIndex)
 	{
 		OutIndex = (int32)GlobalPointCloud.Num();
 		OutStartVertexId = LastIdUsed;
@@ -144,7 +144,7 @@ public:
 
 	const FVertexMesh* GetMeshOfVertexNodeId(const int32 Id) const;
 
-	void GetNodeCoordinates(TArray<FPoint>& NodeCoordinates) const;
+	void GetNodeCoordinates(TArray<FVector>& NodeCoordinates) const;
 	void GetNodeCoordinates(TArray<FVector3f>& NodeCoordinates) const;
 
 	const TArray<FMesh*>& GetMeshes() const;
diff --git a/Engine/Source/Runtime/Datasmith/CADKernel/Base/Public/Mesh/Structure/ThinZone2D.h b/Engine/Source/Runtime/Datasmith/CADKernel/Base/Public/Mesh/Structure/ThinZone2D.h
--- a/Engine/Source/Runtime/Datasmith/CADKernel/Base/Public/Mesh/Structure/ThinZone2D.h
+++ b/Engine/Source/Runtime/Datasmith/CADKernel/Base/Public/Mesh/Structure/ThinZone2D.h
@@ -34,7 +34,7 @@ namespace UE::CADKernel
 	class FTopologicalFace;
 
 	using FReserveContainerFunc = TFunction<void(int32)>;
-	using FAddMeshNodeFunc = TFunction<void(const int32, const FPoint2D&, const double, const FEdgeSegment&, const FPairOfIndex&)>;
+	using FAddMeshNodeFunc = TFunction<void(const int32, const FVector2d&, const double, const FEdgeSegment&, const FPairOfIndex&)>;
 
 	class FThinZoneSide : public FHaveStates
 	{
diff --git a/Engine/Source/Runtime/Datasmith/CADKernel/Base/Public/Topo/TopologicalEdge.h b/Engine/Source/Runtime/Datasmith/CADKernel/Base/Public/Topo/TopologicalEdge.h
--- a/Engine/Source/Runtime/Datasmith/CADKernel/Base/Public/Topo/TopologicalEdge.h
+++ b/Engine/Source/Runtime/Datasmith/CADKernel/Base/Public/Topo/TopologicalEdge.h
@@ -125,7 +125,7 @@ private:
 	FTopologicalEdge(const TSharedRef<FRestrictionCurve>& InCurve, const TSharedRef<FTopologicalVertex>& InVertex1, const TSharedRef<FTopologicalVertex>& InVertex2, const FLinearBoundary& InBoundary);
 	FTopologicalEdge(const TSharedRef<FRestrictionCurve>& InCurve, const TSharedRef<FTopologicalVertex>& InVertex1, const TSharedRef<FTopologicalVertex>& InVertex2);
 	FTopologicalEdge(const TSharedRef<FRestrictionCurve>& InCurve, const FLinearBoundary& InBoundary);
-	FTopologicalEdge(const TSharedRef<FSurface>& InSurface, const FPoint2D& InCoordinateVertex1, const TSharedRef<FTopologicalVertex>& InVertex1, const FPoint2D& InCoordinateVertex2, const TSharedRef<FTopologicalVertex>& InVertex2);
+	FTopologicalEdge(const TSharedRef<FSurface>& InSurface, const FVector2d& InCoordinateVertex1, const TSharedRef<FTopologicalVertex>& InVertex1, const FVector2d& InCoordinateVertex2, const TSharedRef<FTopologicalVertex>& InVertex2);
 
 	FTopologicalEdge() = default;
 
@@ -157,7 +157,7 @@ public:
 	/**
 	 * Build an edge to connect two vertices carried by a 2d segment
 	 */
-	static TSharedPtr<FTopologicalEdge> Make(const TSharedRef<FSurface>& InSurface, const FPoint2D& InCoordinateVertex1, const TSharedRef<FTopologicalVertex>& InVertex1, const FPoint2D& InCoordinateVertex2, const TSharedRef<FTopologicalVertex>& InVertex2);
+	static TSharedPtr<FTopologicalEdge> Make(const TSharedRef<FSurface>& InSurface, const FVector2d& InCoordinateVertex1, const TSharedRef<FTopologicalVertex>& InVertex1, const FVector2d& InCoordinateVertex2, const TSharedRef<FTopologicalVertex>& InVertex2);
 
 	/**
 	 * To check the build edge before returning it or return TSharedPtr<FTopologicalEdge>() if the edge is not valid
@@ -423,8 +423,8 @@ public:
 	void SetStartVertex(const double NewCoordinate);
 	void SetEndVertex(const double NewCoordinate);
 
-	void SetStartVertex(const double NewCoordinate, const FPoint& NewPoint3D);
-	void SetEndVertex(const double NewCoordinate, const FPoint& NewPoint3D);
+	void SetStartVertex(const double NewCoordinate, const FVector& NewPoint3D);
+	void SetEndVertex(const double NewCoordinate, const FVector& NewPoint3D);
 
 	// ======   Boundary Function   ======
 
@@ -446,7 +446,7 @@ public:
 	/**
 	 * @return the 3d coordinate of the start vertex (the barycenter of the twin vertices)
 	 */
-	FPoint GetStartBarycenter()
+	FVector GetStartBarycenter()
 	{
 		return StartVertex->GetBarycenter();
 	}
@@ -454,7 +454,7 @@ public:
 	/**
 	 * @return the 3d coordinate of the end vertex (the barycenter of the twin vertices)
 	 */
-	FPoint GetEndBarycenter()
+	FVector GetEndBarycenter()
 	{
 		return EndVertex->GetBarycenter();
 	}
@@ -462,7 +462,7 @@ public:
 	/**
 	 * @return the 3d coordinate of the start vertex (prefer GetStartBarycenter)
 	 */
-	FPoint GetStartCoordinate()
+	FVector GetStartCoordinate()
 	{
 		return StartVertex->GetCoordinates();
 	}
@@ -470,13 +470,13 @@ public:
 	/**
 	 * @return the 3d coordinate of the end vertex (prefer GetEndBarycenter)
 	 */
-	FPoint GetEndCoordinate()
+	FVector GetEndCoordinate()
 	{
 		return EndVertex->GetCoordinates();
 	}
 
-	void GetTangentsAtExtremities(FPoint& StartTangent, FPoint& EndTangent, bool bForward) const;
-	void GetTangentsAtExtremities(FPoint& StartTangent, FPoint& EndTangent, EOrientation Orientation) const
+	void GetTangentsAtExtremities(FVector& StartTangent, FVector& EndTangent, bool bForward) const;
+	void GetTangentsAtExtremities(FVector& StartTangent, FVector& EndTangent, EOrientation Orientation) const
 	{
 		GetTangentsAtExtremities(StartTangent, EndTangent, Orientation == EOrientation::Front);
 	}
@@ -670,7 +670,7 @@ public:
 	/**
 	 * Approximation of 3D points compute with carrier surface 3D polyline
 	 */
-	void ApproximatePoints(const TArray<double>& InCoordinates, TArray<FPoint>& OutPoints) const
+	void ApproximatePoints(const TArray<double>& InCoordinates, TArray<FVector>& OutPoints) const
 	{
 		Curve->Approximate3DPoints(InCoordinates, OutPoints);
 	}
@@ -678,7 +678,7 @@ public:
 	/**
 	 * Approximation of 2D point defined by its coordinate compute with carrier surface 2D polyline
 	 */
-	FPoint Approximate2DPoint(const double InCoordinate) const
+	FVector2d Approximate2DPoint(const double InCoordinate) const
 	{
 		return Curve->Approximate2DPoint(InCoordinate);
 	}
@@ -686,7 +686,7 @@ public:
 	/**
 	 * Approximation of 2D points defined by its coordinates compute with carrier surface 2D polyline
 	 */
-	void Approximate2DPoints(const TArray<double>& InCoordinates, TArray<FPoint2D>& OutPoints) const
+	void Approximate2DPoints(const TArray<double>& InCoordinates, TArray<FVector2d>& OutPoints) const
 	{
 		Curve->Approximate2DPoints(InCoordinates, OutPoints);
 	}
@@ -699,12 +699,12 @@ public:
 		Curve->ApproximatePolyline(Polyline);
 	}
 
-	FPoint GetTangentAt(const double InCoordinate) const
+	FVector GetTangentAt(const double InCoordinate) const
 	{
 		return Curve->GetTangentAt(InCoordinate);
 	}
 
-	FPoint2D GetTangent2DAt(const double InCoordinate) const
+	FVector2d GetTangent2DAt(const double InCoordinate) const
 	{
 		return Curve->GetTangent2DAt(InCoordinate);
 	}
@@ -712,8 +712,8 @@ public:
 	/**
 	 * Return the tangent at the input vertex
 	 */
-	FPoint GetTangentAt(const FTopologicalVertex& InVertex);
-	FPoint2D GetTangent2DAt(const FTopologicalVertex& InVertex);
+	FVector GetTangentAt(const FTopologicalVertex& InVertex);
+	FVector2d GetTangent2DAt(const FTopologicalVertex& InVertex);
 
 	/**
 	 * Project Point (2D or 3D) on the polyline (2D or 3D) and return the coordinate of the projected point
@@ -727,7 +727,7 @@ public:
 	/**
 	 * Project a set of points on the 3D polyline and return the coordinate of the projected point
 	 */
-	void ProjectPoints(const TArray<FPoint>& InPointsToProject, TArray<double>& OutProjectedPointCoords, TArray<FPoint>& OutProjectedPoints) const
+	void ProjectPoints(const TArray<FVector>& InPointsToProject, TArray<double>& OutProjectedPointCoords, TArray<FVector>& OutProjectedPoints) const
 	{
 		Curve->ProjectPoints(Boundary, InPointsToProject, OutProjectedPointCoords, OutProjectedPoints);
 	}
@@ -735,7 +735,7 @@ public:
 	/**
 	 * Project a set of points of a twin edge on the 3D polyline and return the coordinate of the projected point
 	 */
-	void ProjectTwinEdgePoints(const TArray<FPoint>& InPointsToProject, bool bSameOrientation, TArray<double>& OutProjectedPointCoords) const
+	void ProjectTwinEdgePoints(const TArray<FVector>& InPointsToProject, bool bSameOrientation, TArray<double>& OutProjectedPointCoords) const
 	{
 		const double ToleranceOfProjection = Length3D * 0.1;
 		Curve->ProjectTwinCurvePoints(Boundary, InPointsToProject, bSameOrientation, OutProjectedPointCoords, ToleranceOfProjection);
@@ -744,7 +744,7 @@ public:
 	/**
 	 * Compute 2D points of the edge coincident the points of the twin edge defined by their coordinates
 	 */
-	void ProjectTwinEdgePointsOn2DCurve(const TSharedRef<FTopologicalEdge>& InTwinEdge, const TArray<double>& InTwinEdgePointCoords, TArray<FPoint2D>& OutPoints2D);
+	void ProjectTwinEdgePointsOn2DCurve(const TSharedRef<FTopologicalEdge>& InTwinEdge, const TArray<double>& InTwinEdgePointCoords, TArray<FVector2d>& OutPoints2D);
 
 	void ComputeIntersectionsWithIsos(const TArray<double>& InIsoCoordinates, const EIso InTypeIso, const FSurfacicTolerance& ToleranceIso, TArray<double>& OutIntersection) const
 	{
@@ -778,7 +778,7 @@ public:
 		Curve->GetExtremities(Boundary, Extremities);
 	}
 
-	void Offset2D(const FPoint2D& OffsetDirection);
+	void Offset2D(const FVector2d& OffsetDirection);
 
 	// ======   Geometrical Functions   ======
 
@@ -789,13 +789,13 @@ public:
 	 * @param bKeepStartVertexConnectivity: if true the new edge is connected to endVertex, otherwise the new edge is connected to start vertex
 	 * @param OutNewEdge, the second edge
 	 */
-	FTopologicalVertex* SplitAt(double SplittingCoordinate, const FPoint& NewVertexCoordinate, bool bKeepStartVertexConnectivity, TSharedPtr<FTopologicalEdge>& OutNewEdge);
+	FTopologicalVertex* SplitAt(double SplittingCoordinate, const FVector& NewVertexCoordinate, bool bKeepStartVertexConnectivity, TSharedPtr<FTopologicalEdge>& OutNewEdge);
 
 	/**
 	 * Extend the Edge to the NewVertex.
 	 * NewVertex become the new extremity of the edge. The old vertex is unconnected of the edge.
 	 */
-	bool ExtendTo(bool bStartExtremity, const FPoint2D& NewExtremityCoordinate, TSharedRef<FTopologicalVertex>& NewVertex);
+	bool ExtendTo(bool bStartExtremity, const FVector2d& NewExtremityCoordinate, TSharedRef<FTopologicalVertex>& NewVertex);
 
 	bool IsSharpEdge() const;
 
diff --git a/Engine/Source/Runtime/Datasmith/CADKernel/Base/Public/Topo/TopologicalFace.h b/Engine/Source/Runtime/Datasmith/CADKernel/Base/Public/Topo/TopologicalFace.h
--- a/Engine/Source/Runtime/Datasmith/CADKernel/Base/Public/Topo/TopologicalFace.h
+++ b/Engine/Source/Runtime/Datasmith/CADKernel/Base/Public/Topo/TopologicalFace.h
@@ -230,7 +230,7 @@ public:
 	 * Get a sampling of each loop of the face
 	 * @param OutLoopSamplings an array of 2d points
 	 */
-	const void Get2DLoopSampling(TArray<TArray<FPoint2D>>& OutLoopSamplings) const;
+	const void Get2DLoopSampling(TArray<TArray<FVector2d>>& OutLoopSamplings) const;
 
 	// ======   Loop edge Functions   ======
 
@@ -685,12 +685,12 @@ struct FFaceSubset
 
 struct FBBoxWithNormal
 {
-	FPoint Max;
-	FPoint Min;
-	FPoint MaxPoints[3];
-	FPoint MinPoints[3];
-	FPoint2D MaxCoordinates[3];
-	FPoint2D MinCoordinates[3];
+	FVector Max;
+	FVector Min;
+	FVector MaxPoints[3];
+	FVector MinPoints[3];
+	FVector2d MaxCoordinates[3];
+	FVector2d MinCoordinates[3];
 	FVector MaxPointNormals[3];
 	FVector MinPointNormals[3];
 	bool MaxNormalNeedUpdate[3] = {true, true, true};
@@ -758,7 +758,7 @@ struct FBBoxWithNormal
 			{
 				Max[Index] = PolylineBBox.Max[Index];
 				MaxPoints[Index] = PolylineBBox.MaxPoints[Index];
-				MaxCoordinates[Index] = (IsoType == IsoU) ? FPoint2D(IsoCoordinate, PolylineBBox.CoordinateOfMaxPoint[Index]) : FPoint2D(PolylineBBox.CoordinateOfMaxPoint[Index], IsoCoordinate);
+				MaxCoordinates[Index] = (IsoType == IsoU) ? FVector2d(IsoCoordinate, PolylineBBox.CoordinateOfMaxPoint[Index]) : FVector2d(PolylineBBox.CoordinateOfMaxPoint[Index], IsoCoordinate);
 				MaxNormalNeedUpdate[Index] = true;
 			}
 
@@ -766,13 +766,13 @@ struct FBBoxWithNormal
 			{
 				Min[Index] = PolylineBBox.Min[Index];
 				MinPoints[Index] = PolylineBBox.MinPoints[Index];
-				MinCoordinates[Index] = (IsoType == IsoU) ? FPoint2D(IsoCoordinate, PolylineBBox.CoordinateOfMinPoint[Index]) : FPoint2D(PolylineBBox.CoordinateOfMinPoint[Index], IsoCoordinate);
+				MinCoordinates[Index] = (IsoType == IsoU) ? FVector2d(IsoCoordinate, PolylineBBox.CoordinateOfMinPoint[Index]) : FVector2d(PolylineBBox.CoordinateOfMinPoint[Index], IsoCoordinate);
 				MinNormalNeedUpdate[Index] = true;
 			}
 		}
 	}
 
-	void Update(const FPoint& Point, const FPoint2D InPoint2D)
+	void Update(const FVector& Point, const FVector2d InPoint2D)
 	{
 		for (int32 Index = 0; Index < 3; ++Index)
 		{
diff --git a/Engine/Source/Runtime/Datasmith/CADKernel/Base/Public/Topo/TopologicalLoop.h b/Engine/Source/Runtime/Datasmith/CADKernel/Base/Public/Topo/TopologicalLoop.h
--- a/Engine/Source/Runtime/Datasmith/CADKernel/Base/Public/Topo/TopologicalLoop.h
+++ b/Engine/Source/Runtime/Datasmith/CADKernel/Base/Public/Topo/TopologicalLoop.h
@@ -252,7 +252,7 @@ public:
 		return -1;
 	}
 
-	void Get2DSampling(TArray<FPoint2D>& LoopSampling) const;
+	void Get2DSampling(TArray<FVector2d>& LoopSampling) const;
 
 	/**
 	 * The idea is to remove degenerated edges of the loop i.e. where the surface is degenerated
@@ -261,7 +261,7 @@ public:
 	 * - so where the sampling could be in self-intersecting
 	 * @return false if the loop is degenerated
 	 */
-	bool Get2DSamplingWithoutDegeneratedEdges(TArray<FPoint2D>& LoopSampling) const;
+	bool Get2DSamplingWithoutDegeneratedEdges(TArray<FVector2d>& LoopSampling) const;
 
 	void FindSurfaceCorners(TArray<TSharedPtr<FTopologicalVertex>>& OutCorners, TArray<int32>& OutStartSideIndex) const;
 	void FindBreaks(TArray<TSharedPtr<FTopologicalVertex>>& Ruptures, TArray<int32>& OutStartSideIndex, TArray<double>& RuptureValues) const;
diff --git a/Engine/Source/Runtime/Datasmith/CADKernel/Base/Public/Topo/TopologicalVertex.h b/Engine/Source/Runtime/Datasmith/CADKernel/Base/Public/Topo/TopologicalVertex.h
--- a/Engine/Source/Runtime/Datasmith/CADKernel/Base/Public/Topo/TopologicalVertex.h
+++ b/Engine/Source/Runtime/Datasmith/CADKernel/Base/Public/Topo/TopologicalVertex.h
@@ -25,22 +25,22 @@ class FVertexLink : public TTopologicalLink<FTopologicalVertex>
 	friend class FTopologicalVertex;
 
 protected:
-	FPoint Barycenter;
+	FVector Barycenter;
 
-	void SetBarycenter(const FPoint& Point)
+	void SetBarycenter(const FVector& Point)
 	{
 		Barycenter = Point;
 	}
 
 public:
 	FVertexLink()
-		: Barycenter(FPoint::ZeroVector)
+		: Barycenter(FVector::ZeroVector)
 	{
 	}
 
 	FVertexLink(FTopologicalVertex& Entity)
 		: TTopologicalLink<FTopologicalVertex>(Entity)
-		, Barycenter(FPoint::ZeroVector)
+		, Barycenter(FVector::ZeroVector)
 	{
 	}
 
@@ -54,7 +54,7 @@ public:
 	CADKERNEL_API virtual FInfoEntity& GetInfo(FInfoEntity& Info) const override;
 #endif
 
-	const FPoint& GetBarycenter() const
+	const FVector& GetBarycenter() const
 	{
 		return Barycenter;
 	}
@@ -87,10 +87,10 @@ class FTopologicalVertex : public TLinkable<FTopologicalVertex, FVertexLink>
 protected:
 
 	TArray<FTopologicalEdge*> ConnectedEdges;
-	FPoint Coordinates;
+	FVector Coordinates;
 	TSharedPtr<FVertexMesh> Mesh;
 
-	FTopologicalVertex(const FPoint& InCoordinates)
+	FTopologicalVertex(const FVector& InCoordinates)
 		: Coordinates(InCoordinates)
 	{
 	}
@@ -104,7 +104,7 @@ public:
 		FTopologicalVertex::Empty();
 	}
 
-	static TSharedRef<FTopologicalVertex> Make(const FPoint& InCoordinate)
+	static TSharedRef<FTopologicalVertex> Make(const FVector& InCoordinate)
 	{
 		TSharedRef<FTopologicalVertex> Vertex = FEntity::MakeShared<FTopologicalVertex>(InCoordinate);
 		Vertex->Finalize();
@@ -137,7 +137,7 @@ public:
 	/**
 	 * @return the 3d coordinate of the barycenter of the twin vertices
 	 */
-	inline const FPoint& GetBarycenter() const
+	inline const FVector& GetBarycenter() const
 	{
 		if (TopologicalLink.IsValid() && TopologicalLink->GetTwinEntityNum() > 1)
 		{
@@ -149,17 +149,17 @@ public:
 	/**
 	 * @return the 3d coordinates of the vertex (prefere GetBarycenter())
 	 */
-	const FPoint& GetCoordinates() const
+	const FVector& GetCoordinates() const
 	{
 		return Coordinates;
 	}
 
-	void SetCoordinates(const FPoint& NewCoordinates)
+	void SetCoordinates(const FVector& NewCoordinates)
 	{
 		if (GetLink()->GetTwinEntityNum() > 1)
 		{
 			// Update barycenter
-			FPoint BaryCenter = GetLink()->GetBarycenter() * (double)GetLink()->GetTwinEntityNum();
+			FVector BaryCenter = GetLink()->GetBarycenter() * (double)GetLink()->GetTwinEntityNum();
 			BaryCenter -= Coordinates;
 			BaryCenter += NewCoordinates;
 			BaryCenter /= (double)GetLink()->GetTwinEntityNum();
@@ -187,7 +187,7 @@ public:
 		return FVector::DistSquared(GetLink()->GetBarycenter(), OtherVertex.GetLink()->GetBarycenter());
 	}
 
-	double SquareDistance(const FPoint& Point) const
+	double SquareDistance(const FVector& Point) const
 	{
 		return FVector::DistSquared(Coordinates, Point);
 	}
diff --git a/Engine/Source/Runtime/Datasmith/CADKernel/Base/Public/UI/Display.h b/Engine/Source/Runtime/Datasmith/CADKernel/Base/Public/UI/Display.h
--- a/Engine/Source/Runtime/Datasmith/CADKernel/Base/Public/UI/Display.h
+++ b/Engine/Source/Runtime/Datasmith/CADKernel/Base/Public/UI/Display.h
@@ -127,7 +127,7 @@ void DrawPoint(const TPoint& InPoint, EVisuProperty Property = EVisuProperty::Bl
 /**
  * Draw a mesh element i.e. Element of dimension 1 is an edge, of dimension 2 is a triangle or quadrangle, of dimension 3 is tetrahedron, pyramid, hexahedron, ...
  */
-CADKERNEL_API void DrawElement(int32 Dimension, TArray<FPoint>& Points, EVisuProperty Property = EVisuProperty::Element);
+CADKERNEL_API void DrawElement(int32 Dimension, TArray<FVector>& Points, EVisuProperty Property = EVisuProperty::Element);
 
 template<typename TPoint>
 void Draw(const TArray<TPoint>& Points, EVisuProperty Property = EVisuProperty::BlueCurve)
@@ -139,7 +139,7 @@ void Draw(const TArray<TPoint>& Points, EVisuProperty Property = EVisuProperty::
 
 CADKERNEL_API void DrawMesh(const TSharedPtr<FMesh>& mesh);
 
-CADKERNEL_API void DisplayEdgeCriteriaGrid(int32 EdgeId, const TArray<FPoint>& Points3D);
+CADKERNEL_API void DisplayEdgeCriteriaGrid(int32 EdgeId, const TArray<FVector>& Points3D);
 
 template<typename TPoint>
 void DisplayPoint(const TPoint& Point, FIdent Ident)
@@ -167,14 +167,14 @@ void DisplayPoint(const TPoint& Point, EVisuProperty Property, FIdent Ident)
 #endif
 }
 
-static void DisplayPoint2DWithScale(const FPoint2D& Point, EVisuProperty Property = EVisuProperty::BluePoint)
+static void DisplayPoint2DWithScale(const FVector2d& Point, EVisuProperty Property = EVisuProperty::BluePoint)
 {
 #ifdef CADKERNEL_DEBUG
 	DrawPoint(Point * DisplayScale, Property);
 #endif
 }
 
-static void DisplayPoint2DWithScale(const FPoint2D& Point, EVisuProperty Property, FIdent Ident)
+static void DisplayPoint2DWithScale(const FVector2d& Point, EVisuProperty Property, FIdent Ident)
 {
 #ifdef CADKERNEL_DEBUG
 	F3DDebugSegment G(Ident);
@@ -287,11 +287,11 @@ void DisplayPolyline(const TArray<TPoint>& Points, EVisuProperty Property)
 #endif
 }
 
-static void DisplayPolylineWithScale(const TArray<FPoint2D>& Points, EVisuProperty Property)
+static void DisplayPolylineWithScale(const TArray<FVector2d>& Points, EVisuProperty Property)
 {
 #ifdef CADKERNEL_DEBUG
-	TArray<FPoint2D> PointsWithScale = Points;
-	for (FPoint2D& Point : PointsWithScale)
+	TArray<FVector2d> PointsWithScale = Points;
+	for (FVector2d& Point : PointsWithScale)
 	{
 		Point *= DisplayScale;
 	}
@@ -315,7 +315,7 @@ void DisplayOrientedPolyline(const TArray<TPoint>& Points, EVisuProperty Propert
 #endif
 }
 
-CADKERNEL_API void DrawQuadripode(double Height, double Base, FPoint& Centre, FPoint& Direction, EVisuProperty Property);
+CADKERNEL_API void DrawQuadripode(double Height, double Base, FVector& Centre, FVector& Direction, EVisuProperty Property);
 
 CADKERNEL_API void Draw(const FTopologicalEdge& Edge, EVisuProperty Property = EVisuProperty::BlueCurve);
 CADKERNEL_API void Draw(const FTopologicalFace& Face);
@@ -331,7 +331,7 @@ template<typename TPoint>
 void DrawSegment(const TPoint& Point1, const TPoint& Point2, EVisuProperty Property = EVisuProperty::Element)
 {
 #ifdef CADKERNEL_DEBUG
-	TArray<FPoint> Points;
+	TArray<FVector> Points;
 	Points.Add(Point1);
 	Points.Add(Point2);
 	Draw(Points, Property);
@@ -346,8 +346,8 @@ void DrawSegmentOrientation(const TPoint& Point1, const TPoint& Point2, EVisuPro
 	double Height = Length / 10.0;
 	double Base = Height / 2;
 
-	FPoint Middle = (Point1 + Point2) / 2.;
-	FPoint Tangent = Point2 - Point1;
+	FVector Middle = (Point1 + Point2) / 2.;
+	FVector Tangent = Point2 - Point1;
 	DrawQuadripode(Height, Base, Middle, Tangent, Property);
 #endif
 }
diff --git a/Engine/Source/Runtime/Datasmith/CADKernel/Base/Public/UI/Visu.h b/Engine/Source/Runtime/Datasmith/CADKernel/Base/Public/UI/Visu.h
--- a/Engine/Source/Runtime/Datasmith/CADKernel/Base/Public/UI/Visu.h
+++ b/Engine/Source/Runtime/Datasmith/CADKernel/Base/Public/UI/Visu.h
@@ -151,19 +151,19 @@ public:
 	virtual void UpdateViewer()
 	{}
 
-	virtual void DrawPoint(const FPoint& Point, EVisuProperty InProperty = EVisuProperty::BluePoint)
+	virtual void DrawPoint(const FVector& Point, EVisuProperty InProperty = EVisuProperty::BluePoint)
 	{}
 
-	virtual void DrawPoint(const FPoint2D& Point, EVisuProperty InProperty = EVisuProperty::BluePoint)
+	virtual void DrawPoint(const FVector2d& Point, EVisuProperty InProperty = EVisuProperty::BluePoint)
 	{}
 
-	virtual void DrawElement(int32 Dimension, TArray<FPoint>& Points, EVisuProperty InProperty = EVisuProperty::Element)
+	virtual void DrawElement(int32 Dimension, TArray<FVector>& Points, EVisuProperty InProperty = EVisuProperty::Element)
 	{}
 
-	virtual void DrawPolyline(const TArray<FPoint>& Points, EVisuProperty InProperty = EVisuProperty::BlueCurve)
+	virtual void DrawPolyline(const TArray<FVector>& Points, EVisuProperty InProperty = EVisuProperty::BlueCurve)
 	{}
 
-	virtual void DrawPolyline(const TArray<FPoint2D>& Points, EVisuProperty InProperty = EVisuProperty::BlueCurve)
+	virtual void DrawPolyline(const TArray<FVector2d>& Points, EVisuProperty InProperty = EVisuProperty::BlueCurve)
 	{}
 
 	virtual void DrawMesh(FIdent MeshId)
diff --git a/Engine/Source/Runtime/Datasmith/CADKernel/Base/Public/Utils/StringUtil.h b/Engine/Source/Runtime/Datasmith/CADKernel/Base/Public/Utils/StringUtil.h
--- a/Engine/Source/Runtime/Datasmith/CADKernel/Base/Public/Utils/StringUtil.h
+++ b/Engine/Source/Runtime/Datasmith/CADKernel/Base/Public/Utils/StringUtil.h
@@ -30,9 +30,9 @@ inline std::string ToString(const ValueType& v)
 	return stream.str();
 }
 
-inline std::string ToString(const FPoint& p)
+inline std::string ToString(const FVector& p)
 {
-	std::string Message = "dm.FPoint(";
+	std::string Message = "dm.FVector(";
 	Message += Utils::ToString(p.X) + "," + Utils::ToString(p.Y) + "," + Utils::ToString(p.Z) + ")";
 	return Message;
 }
@@ -62,9 +62,9 @@ inline FString ToFString(const bool& Value)
 	return FString::Printf(TEXT("%d"), Value);
 }
 
-inline FString ToFString(const FPoint& p)
+inline FString ToFString(const FVector& p)
 {
-	return FString::Printf(TEXT("dm.FPoint(%f, %f, %f)"), p.X, p.Y, p.Z);
+	return FString::Printf(TEXT("dm.FVector(%f, %f, %f)"), p.X, p.Y, p.Z);
 }
 
 void Explode(const FString& FullString, const TCHAR* Separator, TArray<FString>& StringArray);
