Fixed wrong tessellation of specific file
Root cause: Buggy method MergeUnconnectedSuccessiveEdges which is trying to aggregate end-to-end curves with different types.
Solution: Remove the call. It was an unnecessary optimization which was impacting the rest of the sewing logic.

#rnx

diff --git a/Engine/Source/Runtime/Datasmith/CADKernel/Base/Private/Topo/Topomaker.cpp b/Engine/Source/Runtime/Datasmith/CADKernel/Base/Private/Topo/Topomaker.cpp
--- a/Engine/Source/Runtime/Datasmith/CADKernel/Base/Private/Topo/Topomaker.cpp
+++ b/Engine/Source/Runtime/Datasmith/CADKernel/Base/Private/Topo/Topomaker.cpp
@@ -663,10 +663,9 @@ void FTopomaker::Sew()
 	CheckSelfConnectedEdge(EdgeLengthTolerance, BorderVertices);
 
 	// basic case: merge pair of edges connected together at both extremities i.e. pair of edges with same extremity active vertices.
+	// #cadkernel_check: Commenting the line below prevent the creation of the mesh in UE-228246. Why?
 	TopomakerTools::MergeCoincidentEdges(BorderVertices, EdgeLengthTolerance);
 
-	//MergeUnconnectedSuccessiveEdges();
-
 	// advance case: two partially coincident edges connected at one extremity i.e. coincident along the shortest edge
 	BorderVertices.Reset();
 	GetBorderVertices(BorderVertices);
