CADKernelEngine module addition : WIP

#rnx

diff --git a/Engine/Source/Runtime/Core/Public/UObject/EnterpriseObjectVersion.h b/Engine/Source/Runtime/Core/Public/UObject/EnterpriseObjectVersion.h
--- a/Engine/Source/Runtime/Core/Public/UObject/EnterpriseObjectVersion.h
+++ b/Engine/Source/Runtime/Core/Public/UObject/EnterpriseObjectVersion.h
@@ -42,6 +42,9 @@ struct FEnterpriseObjectVersion
 		// Change CoreTechParametricSurface serialization to be more efficient
 		CoreTechParametricSurfaceOptim,
 
+		// Replacement of UDatasmithParametricSurfaceData by UParametricSurfaceData
+		AddedParametricSurfaceData,
+
 		// -----<new versions can be added above this line>-------------------------------------------------
 		VersionPlusOne,
 		LatestVersion = VersionPlusOne - 1
diff --git a/Engine/Source/Runtime/Datasmith/CADKernel/Engine/CADKernelEngine.Build.cs b/Engine/Source/Runtime/Datasmith/CADKernel/Engine/CADKernelEngine.Build.cs
new file mode 100644
--- /dev/null
+++ b/Engine/Source/Runtime/Datasmith/CADKernel/Engine/CADKernelEngine.Build.cs
@@ -0,0 +1,44 @@
+// Copyright Epic Games, Inc. All Rights Reserved.
+using System.IO;
+
+namespace UnrealBuildTool.Rules
+{
+	public class CADKernelEngine : ModuleRules
+	{
+		public CADKernelEngine(ReadOnlyTargetRules Target)
+			: base(Target)
+		{
+			bUseUnity = false; // TechSoftLibrary.cpp cannot be compiled with the other cpp files.
+			OptimizeCode = CodeOptimization.Never;
+
+			PublicDependencyModuleNames.AddRange(
+				new string[]
+				{
+					"Core",
+				}
+			);
+			PublicDependencyModuleNames.Add("CoreUObject");
+
+			PrivateDependencyModuleNames.AddRange(
+				new string[]
+				{
+					"CADKernel",
+					"GeometryCore",
+					"Json",
+					"MeshConversion",
+					"MeshDescription",
+					"StaticMeshDescription",
+				}
+			);
+
+			// TechSoft related API needs the TechSoft headers to be accessible
+			bool bHasTechSoft = System.Type.GetType("TechSoft") != null;
+
+			if ((Target.Platform == UnrealTargetPlatform.Win64 || Target.Platform == UnrealTargetPlatform.Linux) && bHasTechSoft)
+			{
+				// #cad_import: Verifiy that this is correct
+				PublicDependencyModuleNames.Add("TechSoft");
+			}
+		}
+	}
+}
\ No newline at end of file
diff --git a/Engine/Source/Runtime/Datasmith/CADKernel/Engine/Private/CADKernelEngine.cpp b/Engine/Source/Runtime/Datasmith/CADKernel/Engine/Private/CADKernelEngine.cpp
new file mode 100644
--- /dev/null
+++ b/Engine/Source/Runtime/Datasmith/CADKernel/Engine/Private/CADKernelEngine.cpp
@@ -0,0 +1,101 @@
+// Copyright Epic Games, Inc. All Rights Reserved.
+
+#include "CADKernelEngine.h"
+
+#include "Topo/Body.h"
+#include "Topo/Model.h"
+#include "Topo/Shell.h"
+#include "Topo/TopologicalFace.h"
+
+#include "CADKernelEngineLog.h"
+
+#include "HAL/IConsoleManager.h"
+
+DEFINE_LOG_CATEGORY(LogCADKernelEngine);
+
+namespace UE::CADKernel
+{
+	static bool GUseEngine = false;
+	FAutoConsoleVariableRef GCADKernelDebugUseEngine(
+		TEXT("CADKernel.Debug.UseEngine"),
+		GUseEngine,
+		TEXT(""),
+		ECVF_Default);
+
+	bool FCADKernelUtilities::Save(TSharedPtr<UE::CADKernel::FModel>& Model, const FString& FilePath)
+	{
+		return false;
+	}
+
+	bool FCADKernelUtilities::Tessellate(TSharedPtr<UE::CADKernel::FModel>& Model, const FTessellationContext& Context, FMeshDescription& Mesh)
+	{
+		return false;
+	}
+
+	bool FCADKernelUtilities::Tessellate(TSharedPtr<UE::CADKernel::FModel>& Model, const FTessellationContext& Context, UE::Geometry::FDynamicMesh3& Mesh)
+	{
+		return false;
+	}
+
+	bool FCADKernelUtilities::Tessellate(const UParametricSurfaceData& SurfaceData, UE::Geometry::FDynamicMesh3& Mesh)
+	{
+#if 0
+		FMeshTranslationContext Context(SurfaceData);
+
+		if (!Context.TessellationSettings.bUseCADKernel)
+		{
+			return false;
+		}
+
+		return Private::TessellateWithCADKernel(SurfaceData.GetRawData(false), Context, Mesh);
+#else
+		return false;
+#endif
+	}
+
+	bool FCADKernelUtilities::Tessellate(const UParametricSurfaceData& SurfaceData, const FCADKernelRetessellationSettings& Settings, UE::Geometry::FDynamicMesh3& Mesh)
+	{
+#if 0
+		FMeshTranslationContext Context(SurfaceData);
+
+		if (!Context.TessellationSettings.bUseCADKernel)
+		{
+			return false;
+		}
+
+		return Private::TessellateWithCADKernel(SurfaceData.GetRawData(false), Context, Mesh);
+#else
+		return false;
+#endif
+	}
+
+	bool FCADKernelUtilities::Tessellate(const UParametricSurfaceData& SurfaceData, FMeshDescription& Mesh)
+	{
+#if 0
+		FMeshTranslationContext Context(SurfaceData);
+
+		if (!Context.TessellationSettings.bUseCADKernel)
+		{
+			return false;
+		}
+
+		return Private::TessellateWithCADKernel(SurfaceData.GetRawData(false), Context, Mesh);
+#else
+		return false;
+#endif 
+	}
+
+	bool FCADKernelUtilities::Tessellate(const UParametricSurfaceData& SurfaceData, const FCADKernelRetessellationSettings& Settings, FMeshDescription& Mesh)
+	{
+#if 0
+		if (!Context.TessellationSettings.bUseCADKernel)
+		{
+			return false;
+		}
+
+		return Private::TessellateWithCADKernel(SurfaceData.GetRawData(false), Context, Mesh);
+#else
+		return false;
+#endif
+	}
+}
\ No newline at end of file
diff --git a/Engine/Source/Runtime/Datasmith/CADKernel/Engine/Private/CADKernelEngineLog.h b/Engine/Source/Runtime/Datasmith/CADKernel/Engine/Private/CADKernelEngineLog.h
new file mode 100644
--- /dev/null
+++ b/Engine/Source/Runtime/Datasmith/CADKernel/Engine/Private/CADKernelEngineLog.h
@@ -0,0 +1,7 @@
+// Copyright Epic Games, Inc. All Rights Reserved.
+#pragma once
+
+#include "Logging/LogMacros.h"
+
+DECLARE_LOG_CATEGORY_EXTERN(LogCADKernelEngine, Log, All)
+
diff --git a/Engine/Source/Runtime/Datasmith/CADKernel/Engine/Private/CADKernelEnginePrivate.cpp b/Engine/Source/Runtime/Datasmith/CADKernel/Engine/Private/CADKernelEnginePrivate.cpp
new file mode 100644
--- /dev/null
+++ b/Engine/Source/Runtime/Datasmith/CADKernel/Engine/Private/CADKernelEnginePrivate.cpp
@@ -0,0 +1,283 @@
+// Copyright Epic Games, Inc. All Rights Reserved.
+
+#include "CADKernelEnginePrivate.h"
+
+#include "CADKernelEngine.h"
+#include "MeshUtilities.h"
+
+#include "Geo/Curves/CurveUtilities.h"
+#include "Mesh/Criteria/Criterion.h"
+#include "Mesh/Meshers/Mesher.h"
+#include "Mesh/Structure/FaceMesh.h"
+#include "Mesh/Structure/ModelMesh.h"
+#include "Topo/Model.h"
+#include "Topo/TopologicalEdge.h"
+#include "Topo/TopologicalFace.h"
+#include "Topo/TopologicalFaceUtilities.h"
+#include "Topo/TopologicalLoop.h"
+
+#include "DynamicMesh/DynamicMesh3.h"
+#include "DynamicMesh/DynamicMeshAttributeSet.h"
+
+#include "Algo/Reverse.h"
+#include "MeshDescription.h"
+
+namespace UE::CADKernel::Private
+{
+	using namespace UE::CADKernel;
+	using namespace UE::CADKernel::MeshUtilities;
+	using namespace UE::Geometry;
+
+	void CheckPatchIds(FModelMesh& ModelMesh)
+	{
+		TSet<int32> PatchIds;
+
+		for (FFaceMesh* FaceMesh : ModelMesh.GetFaceMeshes())
+		{
+			FTopologicalFace& Face = static_cast<FTopologicalFace&>(FaceMesh->GetGeometricEntity());
+			PatchIds.Add(Face.GetPatchId());
+		}
+
+		printf(">>> Found %d", PatchIds.Num());
+	}
+
+	void FillPolygroupLayer(FDynamicMesh3& Mesh, TMap<int32, int32>& OutMapping, const FModelMesh& ModelMesh)
+	{
+		if (!ensure(Mesh.HasTriangleGroups()))
+		{
+			return;
+		}
+
+		static TSet<int32> LastPatchIdSet;
+		TSet<int32> PatchIdSet;
+		int32 TriangleCount = 0;
+
+		for (FFaceMesh* FaceMesh : ModelMesh.GetFaceMeshes())
+		{
+			FTopologicalFace& Face = static_cast<FTopologicalFace&>(FaceMesh->GetGeometricEntity());
+			TriangleCount += FaceMesh->TrianglesVerticesIndex.Num() / 3;
+		}
+
+		if (LastPatchIdSet.Num() > 0)
+		{
+			for (int32 PatchId : PatchIdSet)
+			{
+				ensure(LastPatchIdSet.Contains(PatchId));
+			}
+		}
+
+		LastPatchIdSet = PatchIdSet;
+
+		const int32 Mesh3TriangleCount = Mesh.TriangleCount();
+		ensure(Mesh3TriangleCount == TriangleCount);
+
+		TSet<int32> Missings;
+		for (int TriID = 0; TriID < Mesh.TriangleCount(); ++TriID)
+		{
+			if (!PatchIdSet.Contains(Mesh.GetTriangleGroup(TriID)))
+			{
+				Missings.Add(Mesh.GetTriangleGroup(TriID));
+			}
+		}
+
+		ensure(Missings.Num() == 0);
+
+		Mesh.EnableAttributes();
+		Mesh.Attributes()->SetNumPolygroupLayers(PatchIdSet.Num());
+
+		TMap<int32, FDynamicMeshPolygroupAttribute*> LayerMapping;
+		LayerMapping.Reserve(PatchIdSet.Num());
+
+		int32 Index = -1;
+		for (int32 PatchId : PatchIdSet)
+		{
+			FDynamicMeshPolygroupAttribute* Layer = Mesh.Attributes()->GetPolygroupLayer(++Index);
+			Layer->SetName(*FString::Printf(TEXT("Face #%d"), PatchId));
+
+			LayerMapping.Add(PatchId, Layer);
+			OutMapping.Add(PatchId, Index);
+
+			for (int TriID = 0; TriID < Mesh.TriangleCount(); ++TriID)
+			{
+				Layer->SetValue(TriID, -1);
+			}
+		}
+
+		for (int TriID = 0; TriID < Mesh.TriangleCount(); ++TriID)
+		{
+			const int32 PatchId = Mesh.GetTriangleGroup(TriID);
+
+			FDynamicMeshPolygroupAttribute** LayerPtr = LayerMapping.Find(PatchId);
+			if (ensure(LayerPtr))
+			{
+				(*LayerPtr)->SetValue(TriID, PatchId);
+			}
+		}
+	}
+
+	void DefineMeshCriteria(FModelMesh& MeshModel, const FCADKernelTessellationSettings& TessellationSettings)
+	{
+		using namespace UE::CADKernel;
+
+		constexpr bool bInMillimeter = true;
+
+		{
+			TSharedPtr<FCriterion> CurvatureCriterion = FCriterion::CreateCriterion(ECriterion::CADCurvature);
+			MeshModel.AddCriterion(CurvatureCriterion);
+
+			const double GeometricTolerance = TessellationSettings.GetGeometricTolerance(bInMillimeter);
+			TSharedPtr<FCriterion> MinSizeCriterion = FCriterion::CreateCriterion(ECriterion::MinSize, 2. * GeometricTolerance);
+			MeshModel.AddCriterion(MinSizeCriterion);
+		}
+
+		if (TessellationSettings.GetMaxEdgeLength(bInMillimeter) > SMALL_NUMBER)
+		{
+			TSharedPtr<FCriterion> MaxSizeCriterion = FCriterion::CreateCriterion(ECriterion::MaxSize, TessellationSettings.GetMaxEdgeLength(bInMillimeter));
+			MeshModel.AddCriterion(MaxSizeCriterion);
+		}
+
+		if (TessellationSettings.GetChordTolerance(bInMillimeter) > SMALL_NUMBER)
+		{
+			TSharedPtr<FCriterion> ChordCriterion = FCriterion::CreateCriterion(ECriterion::Sag, TessellationSettings.GetChordTolerance(bInMillimeter));
+			MeshModel.AddCriterion(ChordCriterion);
+		}
+
+		if (TessellationSettings.NormalTolerance > SMALL_NUMBER)
+		{
+			TSharedPtr<FCriterion> MaxNormalAngleCriterion = FCriterion::CreateCriterion(ECriterion::Angle, TessellationSettings.NormalTolerance);
+			MeshModel.AddCriterion(MaxNormalAngleCriterion);
+		}
+	}
+
+	bool TessellateWithCADKernel(const TArray<uint8>& RawData, const FTessellationContext& Context, FDynamicMesh3& Mesh)
+	{
+		TSharedPtr<FMeshWrapperAbstract> MeshWrapper = FMeshWrapperAbstract::MakeWrapper(Context, Mesh);
+		return Tessellate(RawData, Context, *MeshWrapper);
+	}
+
+	bool TessellateWithCADKernel(const TArray<uint8>& RawData, const FTessellationContext& Context, FMeshDescription& Mesh)
+	{
+		TSharedPtr<FMeshWrapperAbstract> MeshWrapper = FMeshWrapperAbstract::MakeWrapper(Context, Mesh);
+		return Tessellate(RawData, Context, *MeshWrapper);
+	}
+	
+	bool Tessellate(const TArray<uint8>& RawData, const UE::CADKernel::FTessellationContext& Context, FMeshWrapperAbstract& MeshWrapper)
+	{
+		TSharedRef<FSession> Session = MakeShared<FSession>(Context.TessellationSettings.GetGeometricTolerance(true));
+		Session->AddDatabase(RawData);
+
+		FModel& Model = Session->GetModel();
+		if (Model.GetBodies().IsEmpty())
+		{
+			return false;
+		}
+
+		// Tessellate the model
+		TSharedRef<FModelMesh> ModelMesh = FEntity::MakeShared<FModelMesh>();
+
+		const double GeometriTolerance = Context.TessellationSettings.GetGeometricTolerance(true);
+		FMesher Mesher(*ModelMesh, GeometriTolerance, false/*bActivateThinZoneMeshing*/);
+
+		DefineMeshCriteria(*ModelMesh, Context.TessellationSettings);
+		Mesher.MeshEntity(Model);
+
+		// #cad_import: Extract mesh data from FModelMesh using MeshWrapper
+
+		return true;
+	}
+
+	bool GetFaceTrimmingCurves(const FModel& Model, const FTopologicalFace& Face, TArray<TArray<TArray<FVector>>>& CurvesOut)
+	{
+		using namespace UE::CADKernel;
+
+		const int32 CurvesCount = CurvesOut.Num();
+
+		CurvesOut.Reserve(CurvesCount + Face.GetLoops().Num());
+
+		for (const TSharedPtr<FTopologicalLoop>& Loop : Face.GetLoops())
+		{
+			TArray<TArray<FVector>> Polylines;
+			Polylines.Reserve(Loop->GetEdges().Num());
+
+			for (const FOrientedEdge& Edge : Loop->GetEdges())
+			{
+				TArray<FVector> Polyline;
+
+				Polyline = CurveUtilities::GetPoles(*Edge.Entity->GetCurve()->Get2DCurve());
+
+				
+				if (Polyline.Num() > 1)
+				{
+					if (Edge.Direction == EOrientation::Back)
+					{
+						if (Polyline.Num() == 2)
+						{
+							FVector Reserve = Polyline[0];
+							Polyline[0] = MoveTemp(Polyline[1]);
+							Polyline[1] = MoveTemp(Reserve);
+						}
+						else
+						{
+							Algo::Reverse(Polyline);
+						}
+					}
+
+					Polylines.Emplace(MoveTemp(Polyline));
+				}
+			}
+
+			if (Polylines.Num() > 1)
+			{
+				CurvesOut.Add(MoveTemp(Polylines));
+			}
+		}
+
+		return CurvesOut.Num() > CurvesCount;
+	}
+
+	bool GetFaceTrimming2DPolylines(const FModel& Model, const FTopologicalFace& Face, TArray<TArray<FVector2d>>& PolylinesOut)
+	{
+		using namespace UE::CADKernel;
+
+		const int32 PolylinesCount = PolylinesOut.Num();
+
+		PolylinesOut.Reserve(PolylinesCount + Face.GetLoops().Num());
+
+		for (const TSharedPtr<FTopologicalLoop>& Loop : Face.GetLoops())
+		{
+			TArray<FVector2d> Polyline;
+
+			Polyline = TopologicalFaceUtilities::Get2DPolyline(*Loop);
+
+			if (Polyline.Num() > 1)
+			{
+				PolylinesOut.Add(MoveTemp(Polyline));
+			}
+		}
+
+		return PolylinesOut.Num() > PolylinesCount;
+	}
+
+	bool GetFaceTrimming3DPolylines(const FModel& Model, const FTopologicalFace& Face, TArray<TArray<FVector>>& PolylinesOut)
+	{
+		using namespace UE::CADKernel;
+
+		const int32 PolylinesCount = PolylinesOut.Num();
+
+		PolylinesOut.Reserve(PolylinesCount + Face.GetLoops().Num());
+
+		for (const TSharedPtr<FTopologicalLoop>& Loop : Face.GetLoops())
+		{
+			TArray<FVector> Polyline;
+
+			Polyline = TopologicalFaceUtilities::Get3DPolyline(*Loop);
+
+			if (Polyline.Num() > 1)
+			{
+				PolylinesOut.Add(MoveTemp(Polyline));
+			}
+		}
+
+		return PolylinesOut.Num() > PolylinesCount;
+	}
+}
diff --git a/Engine/Source/Runtime/Datasmith/CADKernel/Engine/Private/CADKernelEnginePrivate.h b/Engine/Source/Runtime/Datasmith/CADKernel/Engine/Private/CADKernelEnginePrivate.h
new file mode 100644
--- /dev/null
+++ b/Engine/Source/Runtime/Datasmith/CADKernel/Engine/Private/CADKernelEnginePrivate.h
@@ -0,0 +1,39 @@
+// Copyright Epic Games, Inc. All Rights Reserved.
+#pragma once
+
+#include "CoreMinimal.h"
+
+#include "MeshUtilities.h"
+
+namespace UE::Geometry
+{
+	class FDynamicMesh3;
+}
+
+namespace UE::CADKernel
+{
+	class FModel;
+	class FFaceMesh;
+	class FModelMesh;
+	class FTopologicalFace;
+
+	struct FTessellationContext;
+
+	namespace MeshUtilities
+	{
+		class FMeshWrapperAbstract;
+	}
+
+	namespace Private
+	{
+		bool Tessellate(const TArray<uint8>& RawData, const UE::CADKernel::FTessellationContext& Context, UE::CADKernel::MeshUtilities::FMeshWrapperAbstract& MeshWrapper);
+		bool TessellateWithCADKernel(const TArray<uint8>& RawData, const UE::CADKernel::FTessellationContext& Context, UE::Geometry::FDynamicMesh3& Mesh);
+		bool TessellateWithCADKernel(const TArray<uint8>& RawData, const UE::CADKernel::FTessellationContext& Context, FMeshDescription& Mesh);
+
+		// For future use
+		bool GetFaceTrimmingCurves(const UE::CADKernel::FModel& Model, const UE::CADKernel::FTopologicalFace& Face, TArray<TArray<TArray<FVector>>>& CurvesOut);
+		bool GetFaceTrimming2DPolylines(const UE::CADKernel::FModel& Model, const UE::CADKernel::FTopologicalFace& Face, TArray<TArray<FVector2d>>& PolylinesOut);
+		bool GetFaceTrimming3DPolylines(const UE::CADKernel::FModel& Model, const UE::CADKernel::FTopologicalFace& Face, TArray<TArray<FVector>>& PolylinesOut);
+	}
+}
+
diff --git a/Engine/Source/Runtime/Datasmith/CADKernel/Engine/Private/DynamicMeshWrapper.cpp b/Engine/Source/Runtime/Datasmith/CADKernel/Engine/Private/DynamicMeshWrapper.cpp
new file mode 100644
--- /dev/null
+++ b/Engine/Source/Runtime/Datasmith/CADKernel/Engine/Private/DynamicMeshWrapper.cpp
@@ -0,0 +1,430 @@
+// Copyright Epic Games, Inc. All Rights Reserved.
+
+#include "DynamicMeshWrapper.h"
+
+#include "CADKernelEngine.h"
+
+#include "DynamicMesh/DynamicMesh3.h"
+
+#include "Algo/Reverse.h"
+
+namespace UE::CADKernel::MeshUtilities
+{
+	using namespace UE::CADKernel;
+	using namespace UE::Geometry;
+	
+	TSharedPtr<FMeshWrapperAbstract> FMeshWrapperAbstract::MakeWrapper(const FMeshExtractionContext& Context, FDynamicMesh3& Mesh)
+	{
+		return MakeShared<FDynamicMeshWrapper>(Context, Mesh);
+	}
+
+	FDynamicMeshWrapper::FDynamicMeshWrapper(const FMeshExtractionContext& InContext, FDynamicMesh3& InMesh)
+		: FMeshWrapperAbstract(InContext)
+		, MeshOut(InMesh)
+	{
+	}
+
+	bool FDynamicMeshWrapper::AddNewVertices(TArray<FVector>&& VertexArray)
+	{
+		if (bAreVerticesSet)
+		{
+			return false;
+		}
+
+		VertexIdOffset = VertexMapping.Num();
+
+		const int32 VertexCount = VertexArray.Num();
+		VertexMapping.Reserve(VertexIdOffset + VertexCount);
+		VertIDMap.SetNumUninitialized(VertexIdOffset + VertexCount);
+
+		ConvertVectorArray(Context.ModelParams.ModelCoordSys, VertexArray);
+
+		for (const FVector& Vertex : VertexArray)
+		{
+			const int32 NewVertIdx = MeshOut.AppendVertex(Vertex);
+			VertIDMap[NewVertIdx] = NewVertIdx;
+			VertexMapping.Add(NewVertIdx);
+		}
+
+		bNewVerticesAdded = true;
+
+		return true;
+	}
+
+	bool FDynamicMeshWrapper::SetVertices(TArray<FVector>&& InVertices)
+	{
+		VertexIdOffset = VertexMapping.Num();
+		ensure(VertexIdOffset == 0);
+
+		TArray<FVector> Vertices = MoveTemp(InVertices);;
+		AddNewVertices(MoveTemp(Vertices));
+
+		bAreVerticesSet = true;
+
+		return true;
+	}
+
+	void FDynamicMeshWrapper::FinalizeMesh()
+	{
+		AddTriangles();
+		AddSymmetry();
+	}
+
+	bool FDynamicMeshWrapper::ReserveNewTriangles(int32 TriangleCount)
+	{
+		Normals.Reserve(Normals.Num() + TriangleCount * 3);
+		TexCoords.Reserve(TexCoords.Num() + TriangleCount * 3);
+
+		return true;
+	}
+
+	bool  FDynamicMeshWrapper::AddTriangle(int32 GroupID, uint32 MaterialID, const FArray3i& VertexIndices, const TArrayView<FVector3f>& InNormals, const TArrayView<FVector2f>& InTexCoords)
+	{
+		ensure(InNormals.Num() == 3 && InTexCoords.Num() == 3);
+
+		const bool bNeedSwapOrientation = Context.MeshParams.bNeedSwapOrientation;
+		const FArray3i& Orientation = bNeedSwapOrientation ? CounterClockwise : Clockwise;
+
+		LastNormalIndex = Normals.Num();
+		Normals.Append(InNormals);
+		TexCoords.Append(InTexCoords);
+
+		GroupIDSet.Add(GroupID);
+		if (!MaterialMapping.Contains(MaterialID))
+		{
+			MaterialMapping.Add(MaterialID, ++MaterialIDCount);
+		}
+
+		FIndex3i VertexIDs{
+			VertexMapping[VertexIdOffset + VertexIndices[Orientation[0]]],
+			VertexMapping[VertexIdOffset + VertexIndices[Orientation[1]]],
+			VertexMapping[VertexIdOffset + VertexIndices[Orientation[2]]]
+		};
+
+
+		int NewTriangleID = AppendTriangle(VertexIDs, GroupID);
+
+		//-- already seen this triangle for some reason.. or the MeshDecription had a degenerate edge
+		if (NewTriangleID == FDynamicMesh3::InvalidID)
+		{
+			return false;;
+		}
+
+		FArray3i NormalIndices{
+			LastNormalIndex + LastNormalIndex + Orientation[0],
+			LastNormalIndex + LastNormalIndex + Orientation[1],
+			LastNormalIndex + LastNormalIndex + Orientation[2]
+		};
+
+		FArray3i TexCoordIndices{
+			LastNormalIndex + LastNormalIndex + Orientation[0],
+			LastNormalIndex + LastNormalIndex + Orientation[1],
+			LastNormalIndex + LastNormalIndex + Orientation[2]
+		};
+
+		TriangleDataSet.Emplace(GroupID, MaterialID, NormalIndices, TexCoordIndices);
+
+		return true;
+	}
+
+	void FDynamicMeshWrapper::AddSymmetry()
+	{
+		const int32 VertexCount = MeshOut.VertexCount();
+		const int32 TriangleCount = MeshOut.TriangleCount();
+
+		FMatrix44f SymmetricMatrix = (FMatrix44f)GetSymmetricMatrix(Context.MeshParams.SymmetricOrigin, Context.MeshParams.SymmetricNormal);
+		TArray<int32> SymmetricVertexIds;
+
+		SymmetricVertexIds.SetNum(VertexCount);
+
+		for (int32 Index = 0; Index < VertexCount; ++Index)
+		{
+			FVector4f SymmetricPosition = FVector4f(SymmetricMatrix.TransformPosition((FVector3f)MeshOut.GetVertex(Index)));
+			int32 NewVertIdx = MeshOut.AppendVertex(FVector3d(SymmetricPosition));
+			SymmetricVertexIds[Index] = NewVertIdx;
+		}
+
+		int32 NewTriID;
+		for (int32 TriID = 0; TriID < TriangleCount; ++TriID)
+		{
+			const int32 GroupID = MeshOut.GetTriangleGroup(TriID);
+			{
+				const FIndex3i& VertexIDs = MeshOut.GetTriangle(TriID);
+				FIndex3i NewVertexIDs{
+					SymmetricVertexIds[VertexIDs[2]],
+					SymmetricVertexIds[VertexIDs[1]],
+					SymmetricVertexIds[VertexIDs[0]],
+				};
+				NewTriID = MeshOut.AppendTriangle(NewVertexIDs, GroupID);
+				ensure(NewTriID != FDynamicMesh3::DuplicateTriangleID && NewTriID != FDynamicMesh3::InvalidID && NewTriID != FDynamicMesh3::NonManifoldID);
+			}
+
+			{
+				const FIndex3i& VertexIDs = UVOverlay->GetTriangle(TriID);
+				FIndex3i NewVertexIDs;
+
+				NewVertexIDs[0] = UVOverlay->AppendElement(UVOverlay->GetElement(VertexIDs[2]));
+				NewVertexIDs[1] = UVOverlay->AppendElement(UVOverlay->GetElement(VertexIDs[1]));
+				NewVertexIDs[2] = UVOverlay->AppendElement(UVOverlay->GetElement(VertexIDs[0]));
+
+				// set the triangle in the overlay
+				UVOverlay->SetTriangle(NewTriID, NewVertexIDs);
+			}
+
+			{
+				const FIndex3i& VertexIDs = NormalOverlay->GetTriangle(TriID);
+				FVector3f TrianglesNormals[3]{
+					SymmetricMatrix.TransformVector(NormalOverlay->GetElement(VertexIDs[2])),
+					SymmetricMatrix.TransformVector(NormalOverlay->GetElement(VertexIDs[1])),
+					SymmetricMatrix.TransformVector(NormalOverlay->GetElement(VertexIDs[0])),
+				};
+
+				FIndex3i NewVertexIDs;
+				NewVertexIDs[0] = NormalOverlay->AppendElement(&TrianglesNormals[0].X);
+				NewVertexIDs[1] = NormalOverlay->AppendElement(&TrianglesNormals[1].X);
+				NewVertexIDs[2] = NormalOverlay->AppendElement(&TrianglesNormals[2].X);
+
+				NormalOverlay->SetTriangle(NewTriID, NewVertexIDs);
+			}
+
+			TangentOverlay->SetTriangle(NewTriID, { TangentOverlayID , TangentOverlayID, TangentOverlayID });
+			BiTangentOverlay->SetTriangle(NewTriID, { BiTangentOverlayID , BiTangentOverlayID, BiTangentOverlayID });
+			ColorOverlay->SetTriangle(NewTriID, { ColorOverlayID , ColorOverlayID, ColorOverlayID });
+
+			// Add the triangle as a polygon to the mesh description
+			int32 MaterialIndex = MaterialIDAttrib->GetValue(TriID);
+			MaterialIDAttrib->SetValue(NewTriID, &MaterialIndex);
+
+			// Set patch id attribute
+			MeshOut.SetTriangleGroup(NewTriID, GroupID);
+		}
+	}
+
+	void FDynamicMeshWrapper::InitializeAttributes()
+	{
+		int32 TriangleCount = MeshOut.TriangleCount();
+
+		// Enable attributes to extract normals, tangents, etc...
+		MeshOut.EnableAttributes();
+		FDynamicMeshAttributeSet* Attributes = MeshOut.Attributes();
+		ensure(Attributes);
+
+
+		// Normals
+		constexpr int32 NormalLayerCount = 3;
+		Attributes->SetNumNormalLayers(NormalLayerCount);
+		for (int32 Index = 0; Index < NormalLayerCount; ++Index)
+		{
+			Attributes->GetNormalLayer(Index)->InitializeTriangles(TriangleCount);
+		}
+
+		NormalOverlay = Attributes->PrimaryNormals();
+		TangentOverlay = Attributes->PrimaryTangents();
+		BiTangentOverlay = Attributes->PrimaryBiTangents();
+
+		static FVector3f ZeroVector(ForceInitToZero);
+		TangentOverlayID = TangentOverlay->AppendElement(&ZeroVector.X);
+		BiTangentOverlayID = BiTangentOverlay->AppendElement(&ZeroVector.X);
+
+		// Only one UV channel 
+		Attributes->SetNumUVLayers(1);
+		Attributes->GetUVLayer(0)->InitializeTriangles(TriangleCount);
+		UVOverlay = MeshOut.Attributes()->GetUVLayer(0);
+
+		Attributes->EnablePrimaryColors();
+		ColorOverlay = Attributes->PrimaryColors();
+		ColorOverlayID = ColorOverlay->AppendElement(&ZeroVector.X);
+
+		// always enable Material ID if there are any attributes
+		Attributes->EnableMaterialID();
+		MaterialIDAttrib = Attributes->GetMaterialID();
+
+		//
+		MeshOut.EnableTriangleGroups();
+		Attributes->SetNumPolygroupLayers(GroupIDSet.Num());
+		LayerMapping.Reserve(GroupIDSet.Num());
+
+		constexpr int32 InvalidGroupID = INDEX_NONE;
+		int32 Index = -1;
+		for (int32 GroupID : GroupIDSet)
+		{
+			FDynamicMeshPolygroupAttribute* Layer = Attributes->GetPolygroupLayer(++Index);
+			Layer->SetName(*FString::Printf(TEXT("Face #%d"), GroupID));
+
+			LayerMapping.Add(GroupID, Layer);
+		}
+	}
+
+	void FDynamicMeshWrapper::AddTriangles()
+	{
+		if (TriangleDataSet.IsEmpty())
+		{
+			return;
+		}
+
+		InitializeAttributes();
+
+		FIndex3i TriVertexIDs;
+
+		auto TriangleDataIter = TriangleDataSet.CreateIterator();
+		for (int TriID : MeshOut.TriangleIndicesItr())
+		{
+			FIndex3i Triangle = MeshOut.GetTriangle(TriID);
+
+			{
+				TriVertexIDs[0] = UVOverlay->AppendElement(FVector2f(TexCoords[TriangleDataIter->TexCoordIndices[0]] * ScaleUV));
+				TriVertexIDs[1] = UVOverlay->AppendElement(FVector2f(TexCoords[TriangleDataIter->TexCoordIndices[1]] * ScaleUV));
+				TriVertexIDs[2] = UVOverlay->AppendElement(FVector2f(TexCoords[TriangleDataIter->TexCoordIndices[2]] * ScaleUV));
+
+				// set the triangle in the overlay
+				UVOverlay->SetTriangle(TriID, TriVertexIDs);
+
+				
+			}
+
+			// #cad_import: FNormalWelder????
+			{
+				FVector3f TrianglesNormals[3]{
+					Normals[TriangleDataIter->NormalIndices[0]],
+					Normals[TriangleDataIter->NormalIndices[0]],
+					Normals[TriangleDataIter->NormalIndices[0]],
+				};
+
+				TriVertexIDs[0] = NormalOverlay->AppendElement(&TrianglesNormals[0].X);
+				TriVertexIDs[1] = NormalOverlay->AppendElement(&TrianglesNormals[1].X);
+				TriVertexIDs[2] = NormalOverlay->AppendElement(&TrianglesNormals[2].X);
+
+				NormalOverlay->SetTriangle(TriID, TriVertexIDs);
+			}
+
+			TangentOverlay->SetTriangle(TriID, { TangentOverlayID , TangentOverlayID, TangentOverlayID });
+			BiTangentOverlay->SetTriangle(TriID, { BiTangentOverlayID , BiTangentOverlayID, BiTangentOverlayID });
+			ColorOverlay->SetTriangle(TriID, { ColorOverlayID , ColorOverlayID, ColorOverlayID });
+
+			// Add the triangle as a polygon to the mesh description
+			MaterialIDAttrib->SetValue(TriID, MaterialMapping.Find(TriangleDataIter->MaterialID));
+
+			// Set patch id attribute
+			MeshOut.SetTriangleGroup(TriID, TriangleDataIter->GroupID);
+
+			++TriangleDataIter;
+		}
+	}
+
+	bool FDynamicMeshWrapper::StartFaceTriangles(int32 TriangleCount, const TArray<FVector3f>& InNormals, const TArray<FVector2f>& InTexCoords)
+	{
+		ensure(InNormals.Num() == InTexCoords.Num() && InNormals.Num()%3 == 0);
+
+		LastNormalIndex = Normals.Num();
+		Normals.Append(InNormals);
+		TexCoords.Append(InTexCoords);
+
+		return true;
+	}
+
+	int32 FDynamicMeshWrapper::AppendTriangle(FIndex3i& VertexIDs, int32 GroupID)
+	{
+		int32 NewTriangleID = MeshOut.AppendTriangle(VertexIDs, GroupID);
+
+		//-- already seen this triangle for some reason.. or the MeshDecription had a degenerate tri
+		if (NewTriangleID == FDynamicMesh3::DuplicateTriangleID || NewTriangleID == FDynamicMesh3::InvalidID)
+		{
+			return FDynamicMesh3::InvalidID;
+		}
+
+		//-- non manifold 
+		// if append failed due to non-manifold, duplicate verts
+		if (NewTriangleID == FDynamicMesh3::NonManifoldID)
+		{
+			int e0 = MeshOut.FindEdge(VertexIDs[0], VertexIDs[1]);
+			int e1 = MeshOut.FindEdge(VertexIDs[1], VertexIDs[2]);
+			int e2 = MeshOut.FindEdge(VertexIDs[2], VertexIDs[0]);
+
+			// determine which verts need to be duplicated
+			bool bDuplicate[3] = { false, false, false };
+			if (e0 != FDynamicMesh3::InvalidID && MeshOut.IsBoundaryEdge(e0) == false)
+			{
+				bDuplicate[0] = true;
+				bDuplicate[1] = true;
+			}
+			if (e1 != FDynamicMesh3::InvalidID && MeshOut.IsBoundaryEdge(e1) == false)
+			{
+				bDuplicate[1] = true;
+				bDuplicate[2] = true;
+			}
+			if (e2 != FDynamicMesh3::InvalidID && MeshOut.IsBoundaryEdge(e2) == false)
+			{
+				bDuplicate[2] = true;
+				bDuplicate[0] = true;
+			}
+			for (int32 Index = 0; Index < 3; ++Index)
+			{
+				if (bDuplicate[Index])
+				{
+					const FVector Position = MeshOut.GetVertex(VertexIDs[Index]);
+					const int32 NewVertIdx = MeshOut.AppendVertex(Position);
+					VertIDMap.SetNumUninitialized(NewVertIdx + 1);
+					VertIDMap[NewVertIdx] = VertexIDs[Index];
+					VertexMapping[VertexIDs[Index]] = NewVertIdx;
+					VertexIDs[Index] = NewVertIdx;
+				}
+			}
+
+			NewTriangleID = MeshOut.AppendTriangle(VertexIDs, GroupID);
+		}
+
+		return NewTriangleID;
+	}
+
+	bool FDynamicMeshWrapper::AddFaceTriangles(const TArray<FFaceTriangle>& FaceTriangles)
+	{
+		const bool bNeedSwapOrientation = Context.MeshParams.bNeedSwapOrientation;
+		const FArray3i& Orientation = bNeedSwapOrientation ? CounterClockwise : Clockwise;
+
+		bool bSrcIsManifold = true;
+		TriangleDataSet.Reserve(TriangleDataSet.Num() + FaceTriangles.Num());
+
+		for (const FFaceTriangle& FaceTriangle : FaceTriangles)
+		{
+			GroupIDSet.Add(FaceTriangle.GroupID);
+			if (!MaterialMapping.Contains(FaceTriangle.MaterialID))
+			{
+				MaterialMapping.Add(FaceTriangle.MaterialID, ++MaterialIDCount);
+			}
+
+			FIndex3i VertexIDs{
+				VertexMapping[VertexIdOffset + FaceTriangle.VertexIndices[Orientation[0]]],
+				VertexMapping[VertexIdOffset + FaceTriangle.VertexIndices[Orientation[1]]],
+				VertexMapping[VertexIdOffset + FaceTriangle.VertexIndices[Orientation[2]]]
+			};
+
+			int NewTriangleID = AppendTriangle(VertexIDs, FaceTriangle.GroupID);
+
+			//-- already seen this triangle for some reason.. or the MeshDecription had a degenerate tri
+			if (NewTriangleID == FDynamicMesh3::InvalidID)
+			{
+				continue;
+			}
+
+			FArray3i NormalIndices{
+				LastNormalIndex + FaceTriangle.Normals[Orientation[0]],
+				LastNormalIndex + FaceTriangle.Normals[Orientation[1]],
+				LastNormalIndex + FaceTriangle.Normals[Orientation[2]]
+			};
+
+			FArray3i TexCoordIndices{
+				LastNormalIndex + FaceTriangle.TexCoords[Orientation[0]],
+				LastNormalIndex + FaceTriangle.TexCoords[Orientation[1]],
+				LastNormalIndex + FaceTriangle.TexCoords[Orientation[2]]
+			};
+
+			TriangleDataSet.Emplace(FaceTriangle.GroupID, FaceTriangle.MaterialID, NormalIndices, TexCoordIndices);
+		}
+
+		return true;
+	}
+
+	void FDynamicMeshWrapper::EndFaceTriangles()
+	{
+	}
+}
diff --git a/Engine/Source/Runtime/Datasmith/CADKernel/Engine/Private/DynamicMeshWrapper.h b/Engine/Source/Runtime/Datasmith/CADKernel/Engine/Private/DynamicMeshWrapper.h
new file mode 100644
--- /dev/null
+++ b/Engine/Source/Runtime/Datasmith/CADKernel/Engine/Private/DynamicMeshWrapper.h
@@ -0,0 +1,82 @@
+// Copyright Epic Games, Inc. All Rights Reserved.
+#pragma once
+
+#include "CoreMinimal.h"
+
+#include "MeshUtilities.h"
+
+#include "DynamicMesh/DynamicMeshAttributeSet.h"
+
+namespace UE::CADKernel::MeshUtilities
+{
+	class FDynamicMeshWrapper : public FMeshWrapperAbstract
+	{
+	public:
+		FDynamicMeshWrapper(const FMeshExtractionContext& InContext, UE::Geometry::FDynamicMesh3& InMesh);
+
+		~FDynamicMeshWrapper() { FinalizeMesh(); }
+		
+		virtual bool ReserveNewTriangles(int32 AdditionalTriangleCount) override;
+		virtual bool SetVertices(TArray<FVector>&& Vertices) override;
+		virtual bool AddNewVertices(TArray<FVector>&& Vertices) override;
+		virtual bool AddTriangle(int32 GroupID, uint32 MaterialID, const FArray3i& VertexIndices, const TArrayView<FVector3f>& Normals, const TArrayView<FVector2f>& TexCoords) override;
+		virtual bool StartFaceTriangles(int32 TriangleCount, const TArray<FVector3f>& Normals, const TArray<FVector2f>& TexCoords) override;
+		virtual bool AddFaceTriangles(const TArray<FFaceTriangle>& FaceTriangles) override;
+		virtual bool AddFaceTriangle(const FFaceTriangle& FaceTriangle) override
+		{
+			return AddFaceTriangles({ FaceTriangle });
+		}
+		virtual void EndFaceTriangles() override;
+
+		virtual void FinalizeMesh() override;
+
+	private:
+		void AddSymmetry();
+		void InitializeAttributes();
+		void AddTriangles();
+		int32 AppendTriangle(UE::Geometry::FIndex3i& VertexIDs, int32 GroupID);
+
+	private:
+		UE::Geometry::FDynamicMesh3& MeshOut;
+
+		TMap<uint32, int32> MaterialMapping;
+		/**
+		* map from DynamicMesh vertex Id to MeshDecription FVertexID.
+		* NB: due to vertex splitting, multiple DynamicMesh vertex ids
+		* may map to the same MeshDescription FVertexID.
+		*  ( a vertex split is a result of reconciling non-manifold MeshDescription vertex )
+		*/
+		TArray<int32> VertIDMap;
+		TMap<int32, UE::Geometry::FDynamicMeshPolygroupAttribute*> LayerMapping;
+
+		UE::Geometry::FDynamicMeshNormalOverlay* NormalOverlay = nullptr;
+		UE::Geometry::FDynamicMeshMaterialAttribute* MaterialIDAttrib = nullptr;
+		UE::Geometry::FDynamicMeshUVOverlay* UVOverlay = nullptr;
+		UE::Geometry::FDynamicMeshColorOverlay* ColorOverlay = nullptr;
+		UE::Geometry::FDynamicMeshNormalOverlay* TangentOverlay = nullptr;
+		UE::Geometry::FDynamicMeshNormalOverlay* BiTangentOverlay = nullptr;
+		int32 TangentOverlayID = INDEX_NONE;
+		int32 BiTangentOverlayID = INDEX_NONE;
+		int32 ColorOverlayID = INDEX_NONE;
+
+		bool bNewVerticesAdded = false;
+		int32 VertexIdOffset = 0;
+		TArray<int32> VertexMapping;
+
+		struct FTriangleData
+		{
+			int32 GroupID;
+			uint32 MaterialID;
+			FArray3i NormalIndices;
+			FArray3i TexCoordIndices;
+		};
+		TArray<FTriangleData> TriangleDataSet;
+		TSet<int32> GroupIDSet;
+		TArray<FVector3f> Normals;
+		TArray<FVector2f> TexCoords;
+		int32 LastNormalIndex = 0;
+		int32 MaterialIDCount = 0;
+
+		bool bIsFinalized = false;
+	};
+}
diff --git a/Engine/Source/Runtime/Datasmith/CADKernel/Engine/Private/MeshDescriptionWrapper.cpp b/Engine/Source/Runtime/Datasmith/CADKernel/Engine/Private/MeshDescriptionWrapper.cpp
new file mode 100644
--- /dev/null
+++ b/Engine/Source/Runtime/Datasmith/CADKernel/Engine/Private/MeshDescriptionWrapper.cpp
@@ -0,0 +1,384 @@
+// Copyright Epic Games, Inc. All Rights Reserved. 
+
+#include "MeshDescriptionWrapper.h"
+
+#include "CADKernelEngine.h"
+#include "CADKernelEngineLog.h"
+#include "MeshUtilities.h"
+
+#include "MeshDescriptionBuilder.h"
+
+namespace UE::CADKernel::MeshUtilities
+{
+	TSharedPtr<FMeshWrapperAbstract> FMeshWrapperAbstract::MakeWrapper(const FMeshExtractionContext& Context, FMeshDescription& Mesh)
+	{
+		return MakeShared<FMeshDescriptionWrapper>(Context, Mesh);
+	}
+
+	FMeshDescriptionWrapper::FMeshDescriptionWrapper(const FMeshExtractionContext& InContext, FMeshDescription& InMesh)
+		: FMeshWrapperAbstract(InContext)
+		, Attributes(InMesh)
+		, Mesh(InMesh)
+	{
+		Attributes.Register();
+		if (Attributes.IsValid())
+		{
+			VertexPositions = Mesh.GetVertexPositions();
+			VertexInstanceNormals = Attributes.GetVertexInstanceNormals();
+			VertexInstanceTangents = Attributes.GetVertexInstanceTangents();
+			VertexInstanceBinormalSigns = Attributes.GetVertexInstanceBinormalSigns();
+			VertexInstanceColors = Attributes.GetVertexInstanceColors();
+			VertexInstanceUVs = Attributes.GetVertexInstanceUVs();
+			PolygonAttributes = Attributes.GetPolygonGroups();
+			PolygonGroupImportedMaterialSlotNames = Attributes.GetPolygonGroupMaterialSlotNames();
+		}
+	}
+
+	void FMeshDescriptionWrapper::FinalizeMesh()
+	{
+		if (bIsFinalized)
+		{
+			return;
+		}
+
+		AddSymmetry();
+
+		if (Context.MeshParams.bNeedSwapOrientation)
+		{
+			for (const FVertexInstanceID VertexInstanceID : Mesh.VertexInstances().GetElementIDs())
+			{
+				VertexInstanceNormals[VertexInstanceID] *= -1.f;
+			}
+		}
+
+		for (const FVertexInstanceID VertexInstanceID : Mesh.VertexInstances().GetElementIDs())
+		{
+			VertexInstanceColors[VertexInstanceID] = FLinearColor::White;
+			VertexInstanceTangents[VertexInstanceID] = FVector3f(ForceInitToZero);
+			VertexInstanceBinormalSigns[VertexInstanceID] = 0.0f;
+		}
+
+		bIsFinalized = true;
+	}
+
+	void FMeshDescriptionWrapper::AddSymmetry()
+	{
+		if (!Context.MeshParams.bIsSymmetric)
+		{
+			return;
+		}
+
+		FMatrix44f SymmetricMatrix = (FMatrix44f)GetSymmetricMatrix(Context.MeshParams.SymmetricOrigin, Context.MeshParams.SymmetricNormal);
+
+		TMap<FVertexID, FVertexID> VertexMapping;
+		VertexMapping.Reserve(VertexPositions.GetNumElements());
+
+		for (const FVertexID VertexID : Mesh.Vertices().GetElementIDs())
+		{
+			const FVector4f SymmetricPosition = FVector4f(SymmetricMatrix.TransformPosition(VertexPositions[VertexID]));
+
+			const FVertexID NewVertexID = Mesh.CreateVertex();
+			VertexPositions[NewVertexID] = FVector3f(SymmetricPosition);
+
+			VertexMapping.Add(VertexID, NewVertexID);
+		}
+
+		FVertexID3 NewVertices;
+		FVertexInstanceID3 NewVertexInstanceIDs;
+		FVertexInstanceID3 VertexInstancesIDs;
+		TArray<FVertexInstanceID> VertexInstancesIDArray(VertexInstancesIDs.ABC, 3);
+
+		for (const FPolygonID PolygonID : Mesh.Polygons().GetElementIDs())
+		{
+			Mesh.GetPolygonVertexInstances(PolygonID, VertexInstancesIDArray);
+
+			NewVertices[0] = VertexMapping[Mesh.GetVertexInstanceVertex(VertexInstancesIDs[2])];
+			NewVertices[1] = VertexMapping[Mesh.GetVertexInstanceVertex(VertexInstancesIDs[1])];
+			NewVertices[2] = VertexMapping[Mesh.GetVertexInstanceVertex(VertexInstancesIDs[0])];
+
+			NewVertexInstanceIDs[0] = Mesh.CreateVertexInstance(NewVertices[0]);
+			NewVertexInstanceIDs[1] = Mesh.CreateVertexInstance(NewVertices[1]);
+			NewVertexInstanceIDs[2] = Mesh.CreateVertexInstance(NewVertices[2]);
+
+			VertexInstanceUVs.Set(NewVertexInstanceIDs[0], 0/*UVChannel*/, VertexInstanceUVs.Get(VertexInstancesIDs[2]));
+			VertexInstanceUVs.Set(NewVertexInstanceIDs[1], 0/*UVChannel*/, VertexInstanceUVs.Get(VertexInstancesIDs[1]));
+			VertexInstanceUVs.Set(NewVertexInstanceIDs[2], 0/*UVChannel*/, VertexInstanceUVs.Get(VertexInstancesIDs[0]));
+
+			VertexInstanceNormals[NewVertexInstanceIDs[0]] = SymmetricMatrix.TransformVector(VertexInstanceNormals[VertexInstancesIDs[2]]);
+			VertexInstanceNormals[NewVertexInstanceIDs[1]] = SymmetricMatrix.TransformVector(VertexInstanceNormals[VertexInstancesIDs[1]]);
+			VertexInstanceNormals[NewVertexInstanceIDs[2]] = SymmetricMatrix.TransformVector(VertexInstanceNormals[VertexInstancesIDs[0]]);
+
+			// Add the triangle as a polygon to the mesh description
+			const FPolygonID NewPolygonID = Mesh.CreatePolygon(Mesh.GetPolygonPolygonGroup(PolygonID), NewVertexInstanceIDs.ABC);
+
+			// Set patch id attribute
+			PolygonAttributes[NewPolygonID] = PolygonAttributes[PolygonID];
+		}
+	}
+
+	// Inspired from FDynamicMesh3::AppendTriangle and its usage
+	bool FMeshDescriptionWrapper::GetVertexInstances(const FArray3i& Vertices, FVertexInstanceID3& VertexInstances)
+	{
+		FVertexID3 Triangle(VertexIDs[Vertices[0]], VertexIDs[Vertices[1]], VertexIDs[Vertices[2]]);
+
+		if (Triangle.A == Triangle.B || Triangle.A == Triangle.C || Triangle.C == Triangle.B)
+		{
+			// Degenerated triangle
+			return false;
+		}
+
+		FEdgeID Edges[3];
+		bool IsBoundary[3]{ true, true, true };
+
+		Edges[0] = FindEdge(Triangle.A, Triangle.B, IsBoundary[0]);
+		Edges[1] = FindEdge(Triangle.B, Triangle.C, IsBoundary[1]);
+		Edges[2] = FindEdge(Triangle.C, Triangle.A, IsBoundary[2]);
+
+		// At least one edge is non-manifold
+		if (IsBoundary[0] == false || IsBoundary[1] == false || IsBoundary[2] == false)
+		{
+			// #cad_debug
+			ensure(false);
+
+			bool bDuplicate[3]{ false, false, false };
+			if (IsBoundary[0] == false)
+			{
+				bDuplicate[0] = true;
+				bDuplicate[1] = true;
+			}
+			if (IsBoundary[1] == false)
+			{
+				bDuplicate[1] = true;
+				bDuplicate[2] = true;
+			}
+			if (IsBoundary[2] == false)
+			{
+				bDuplicate[2] = true;
+				bDuplicate[0] = true;
+			}
+
+			FVertexInstanceID3 NewTrianle(Triangle.A, Triangle.B, Triangle.C);
+			for (int32 Index = 0; Index < 3; ++Index)
+			{
+				if (bDuplicate[Index])
+				{
+					const FVertexID VertexID = Mesh.CreateVertex();
+					VertexPositions[VertexID] = VertexPositions[Triangle[Index]];
+					NewTrianle[Index] = VertexID;
+
+					// Make sure new triangle using faulty vertex index uses the newly created one.
+					VertexIDs[Vertices[Index]] = VertexID;
+				}
+			}
+
+			VertexInstances.A = Mesh.CreateVertexInstance(NewTrianle.A);
+			VertexInstances.B = Mesh.CreateVertexInstance(NewTrianle.B);
+			VertexInstances.C = Mesh.CreateVertexInstance(NewTrianle.C);
+		}
+		else
+		{
+			VertexInstances.A = Mesh.CreateVertexInstance(Triangle.A);
+			VertexInstances.B = Mesh.CreateVertexInstance(Triangle.B);
+			VertexInstances.C = Mesh.CreateVertexInstance(Triangle.C);
+		}
+
+		return true;
+	}
+
+	bool FMeshDescriptionWrapper::SetVertices(TArray<FVector>&& InVertices)
+	{
+		VertexIndexOffset = VertexPositions.GetNumElements();
+		ensure(bAreVerticesSet == false && VertexIndexOffset == 0);
+
+		TArray<FVector> Vertices = MoveTemp(InVertices);;
+		AddNewVertices(MoveTemp(Vertices));
+
+		bAreVerticesSet = true;
+
+		return true;
+	}
+
+	bool FMeshDescriptionWrapper::AddNewVertices(TArray<FVector>&& InVertices)
+	{
+		using namespace UE::CADKernel;
+
+		if (bAreVerticesSet)
+		{
+			return false;
+		}
+
+		VertexIndexOffset = VertexPositions.GetNumElements();
+
+		const int32 VertexCount = InVertices.Num();
+
+		Mesh.ReserveNewVertices(Context.MeshParams.bIsSymmetric ? VertexCount * 2 : VertexCount);
+		VertexIDs.Empty(VertexCount);
+
+		// Make Mesh.VertexPositions and VertexID
+		ConvertVectorArray(Context.ModelParams.ModelCoordSys, InVertices);
+
+		for (const FVector& Vertex : InVertices)
+		{
+			const FVertexID VertexID = Mesh.CreateVertex();
+			VertexPositions[VertexID] = FVector3f(Vertex);
+			VertexIDs.Add(VertexID);
+		}
+
+		return true;
+	}
+
+	bool FMeshDescriptionWrapper::ReserveNewTriangles(int32 TriangleCount)
+	{
+		if (Context.MeshParams.bIsSymmetric)
+		{
+			TriangleCount *= 2;
+		}
+
+		Mesh.ReserveNewPolygons(TriangleCount);
+		Mesh.ReserveNewVertexInstances(TriangleCount * 3);
+		Mesh.ReserveNewUVs(TriangleCount * 3);
+		Mesh.ReserveNewEdges(TriangleCount * 3);
+
+		if (MaterialToPolygonGroupMapping.Num() >= NANITE_MAX_CLUSTER_MATERIALS)
+		{
+			UE_LOG(LogCADKernelEngine, Warning, TEXT("The main UE5 rendering systems do not support more than %d materials per mesh. Only the first %d materials are kept. The others are replaced by the last one"), NANITE_MAX_CLUSTER_MATERIALS, NANITE_MAX_CLUSTER_MATERIALS);
+		}
+
+		// Add to the mesh, a polygon groups per material
+		int32 PolyGroupIndex = 0;
+		FPolygonGroupID PolyGroupID = 0;
+		for (TPair<uint32, FPolygonGroupID>& Material : MaterialToPolygonGroupMapping)
+		{
+			if (PolyGroupIndex < NANITE_MAX_CLUSTER_MATERIALS)
+			{
+				uint32 MaterialHash = Material.Key;
+				FName ImportedSlotName = *LexToString(MaterialHash);
+				PolyGroupID = Mesh.CreatePolygonGroup();
+				PolygonGroupImportedMaterialSlotNames[PolyGroupID] = ImportedSlotName;
+				PolyGroupIndex++;
+			}
+
+			Material.Value = PolyGroupID;
+		}
+
+		return true;
+	}
+
+	bool FMeshDescriptionWrapper::AddTriangle(int32 GroupID, uint32 MaterialID, const FArray3i& VertexIndices, const TArrayView<FVector3f>& InNormals, const TArrayView<FVector2f>& InTexCoords)
+	{
+		ensure(InNormals.Num() == 3 && InTexCoords.Num() == 3);
+
+		const bool bNeedSwapOrientation = Context.MeshParams.bNeedSwapOrientation;
+		const FArray3i& Orientation = bNeedSwapOrientation ? CounterClockwise : Clockwise;
+
+		FPolygonGroupID PolygonGroupID = GetPolygonGroupID(MaterialID);
+		FVertexInstanceID3 VertexInstanceIDs;
+		FArray3i Vertices{
+			VertexIndexOffset + VertexIndices[Orientation[0]],
+			VertexIndexOffset + VertexIndices[Orientation[1]],
+			VertexIndexOffset + VertexIndices[Orientation[2]]
+		};
+
+		if (!GetVertexInstances(Vertices, VertexInstanceIDs))
+		{
+			return false;
+		}
+
+		// Fill up arrays associated with vertex instances
+		VertexInstanceUVs.Set(VertexInstanceIDs[0], 0/*UVChannel*/, FVector2f(InTexCoords[Orientation[0]] * ScaleUV));
+		VertexInstanceUVs.Set(VertexInstanceIDs[1], 0/*UVChannel*/, FVector2f(InTexCoords[Orientation[1]] * ScaleUV));
+		VertexInstanceUVs.Set(VertexInstanceIDs[2], 0/*UVChannel*/, FVector2f(InTexCoords[Orientation[2]] * ScaleUV));
+
+		VertexInstanceNormals[VertexInstanceIDs[0]] = FVector3f(InNormals[Orientation[0]]);
+		VertexInstanceNormals[VertexInstanceIDs[1]] = FVector3f(InNormals[Orientation[1]]);
+		VertexInstanceNormals[VertexInstanceIDs[2]] = FVector3f(InNormals[Orientation[2]]);
+
+		// Add the triangle as a polygon to the mesh description
+		const FPolygonID PolygonID = Mesh.CreatePolygon(PolygonGroupID, VertexInstanceIDs.ABC);
+
+		// Set group attribute
+		PolygonAttributes[PolygonID] = GroupID;
+
+		return true;
+	}
+
+	bool FMeshDescriptionWrapper::StartFaceTriangles(int32 TriangleCount, const TArray<FVector3f>& InNormals, const TArray<FVector2f>& InTexCoords)
+	{
+		ensure(InNormals.Num() == InTexCoords.Num());
+		Normals = InNormals;
+		TexCoords = InTexCoords;
+
+		return true;
+	}
+
+	bool FMeshDescriptionWrapper::AddFaceTriangles(const TArray<FFaceTriangle>& FaceTriangles)
+	{
+		const bool bNeedSwapOrientation = Context.MeshParams.bNeedSwapOrientation;
+		const FArray3i& Orientation = bNeedSwapOrientation ? CounterClockwise : Clockwise;
+		
+		FPolygonGroupID PolygonGroupID;
+		FVertexInstanceID3 VertexInstanceIDs;
+
+		for (const FFaceTriangle& FaceTriangle : FaceTriangles)
+		{
+			FArray3i Vertices{
+				VertexIndexOffset + FaceTriangle.VertexIndices[Orientation[0]],
+				VertexIndexOffset + FaceTriangle.VertexIndices[Orientation[1]],
+				VertexIndexOffset + FaceTriangle.VertexIndices[Orientation[2]]
+			};
+
+			if (!GetVertexInstances(Vertices, VertexInstanceIDs))
+			{
+				return false;
+			}
+
+			// Fill up arrays associated with vertex instances
+			VertexInstanceUVs.Set(VertexInstanceIDs[0], 0/*UVChannel*/, FVector2f(TexCoords[FaceTriangle.TexCoords[Orientation[0]]] * ScaleUV));
+			VertexInstanceUVs.Set(VertexInstanceIDs[1], 0/*UVChannel*/, FVector2f(TexCoords[FaceTriangle.TexCoords[Orientation[1]]] * ScaleUV));
+			VertexInstanceUVs.Set(VertexInstanceIDs[2], 0/*UVChannel*/, FVector2f(TexCoords[FaceTriangle.TexCoords[Orientation[2]]] * ScaleUV));
+
+			VertexInstanceNormals[VertexInstanceIDs[0]] = FVector3f(Normals[FaceTriangle.Normals[Orientation[0]]]);
+			VertexInstanceNormals[VertexInstanceIDs[1]] = FVector3f(Normals[FaceTriangle.Normals[Orientation[1]]]);
+			VertexInstanceNormals[VertexInstanceIDs[2]] = FVector3f(Normals[FaceTriangle.Normals[Orientation[2]]]);
+
+			// Add the triangle as a polygon to the mesh description
+			PolygonGroupID = GetPolygonGroupID(FaceTriangle.MaterialID);
+			const FPolygonID PolygonID = Mesh.CreatePolygon(PolygonGroupID, VertexInstanceIDs.ABC);
+
+			// Set group attribute
+			PolygonAttributes[PolygonID] = FaceTriangle.GroupID;
+		}
+
+		return true;
+	}
+
+	void FMeshDescriptionWrapper::EndFaceTriangles()
+	{
+	}
+
+	void FCADKernelStaticMeshAttributes::Register(bool bKeepExistingAttribute)
+	{
+		FStaticMeshAttributes::Register(bKeepExistingAttribute);
+
+		using namespace ExtendedMeshAttribute;
+
+		if (!MeshDescription.PolygonAttributes().HasAttribute(PolyTriGroups) || !bKeepExistingAttribute)
+		{
+			MeshDescription.PolygonAttributes().RegisterAttribute<int32>(PolyTriGroups, 1, 0, EMeshAttributeFlags::AutoGenerated);
+		}
+
+		ensure(IsValid());
+	}
+
+	TPolygonAttributesRef<int32> FCADKernelStaticMeshAttributes::GetPolygonGroups()
+	{
+		using namespace ExtendedMeshAttribute;
+		return MeshDescription.PolygonAttributes().GetAttributesRef<int32>(PolyTriGroups);
+	}
+
+	TPolygonAttributesConstRef<int32> FCADKernelStaticMeshAttributes::GetPolygonGroups() const
+	{
+		using namespace ExtendedMeshAttribute;
+		return MeshDescription.PolygonAttributes().GetAttributesRef<int32>(PolyTriGroups);
+	}
+}
diff --git a/Engine/Source/Runtime/Datasmith/CADKernel/Engine/Private/MeshDescriptionWrapper.h b/Engine/Source/Runtime/Datasmith/CADKernel/Engine/Private/MeshDescriptionWrapper.h
new file mode 100644
--- /dev/null
+++ b/Engine/Source/Runtime/Datasmith/CADKernel/Engine/Private/MeshDescriptionWrapper.h
@@ -0,0 +1,154 @@
+// Copyright Epic Games, Inc. All Rights Reserved. 
+
+#pragma once
+
+#include "MeshUtilities.h"
+
+#include "MeshDescription.h"
+#include "NaniteDefinitions.h"
+#include "StaticMeshAttributes.h"
+
+namespace UE::CADKernel::MeshUtilities
+{
+	class FCADKernelStaticMeshAttributes : public FStaticMeshAttributes
+	{
+	public:
+
+		explicit FCADKernelStaticMeshAttributes(FMeshDescription& InMeshDescription) : FStaticMeshAttributes(InMeshDescription) {}
+
+		virtual void Register(bool bKeepExistingAttribute = false) override;
+
+		bool IsValid() const
+		{
+			return GetVertexInstanceNormals().IsValid() &&
+				GetVertexInstanceTangents().IsValid() &&
+				GetVertexInstanceBinormalSigns().IsValid() &&
+				GetVertexInstanceColors().IsValid() &&
+				GetVertexInstanceUVs().IsValid() &&
+				GetPolygonGroupMaterialSlotNames().IsValid() &&
+				GetPolygonGroups().IsValid();
+		}
+
+		/**
+		 * Enable per-triangle integer attribute named PolyTriGroups (see ExtendedMeshAttribute::PolyTriGroups)
+		 * This integer defines the identifier of the PolyTriGroup containing the triangle.
+		 * In case of mesh coming from a CAD file, a PolyTriGroup is associated to a CAD topological face
+		 */
+		TPolygonAttributesRef<int32> GetPolygonGroups();
+		TPolygonAttributesConstRef<int32> GetPolygonGroups() const;
+	};
+
+	class FMeshDescriptionWrapper : public FMeshWrapperAbstract
+	{
+	public:
+		FMeshDescriptionWrapper(const FMeshExtractionContext& InContext, FMeshDescription& InMesh);
+
+		~FMeshDescriptionWrapper() { FinalizeMesh(); }
+
+		virtual bool ReserveNewTriangles(int32 AdditionalTriangleCount) override;
+		virtual bool SetVertices(TArray<FVector>&& Vertices) override;
+		virtual bool AddNewVertices(TArray<FVector>&& Vertices) override;
+		virtual bool AddTriangle(int32 GroupID, uint32 MaterialID, const FArray3i& VertexIndices, const TArrayView<FVector3f>& Normals, const TArrayView<FVector2f>& TexCoords) override;
+		virtual bool StartFaceTriangles(int32 TriangleCount, const TArray<FVector3f>& Normals, const TArray<FVector2f>& TexCoords) override;
+		virtual bool AddFaceTriangles(const TArray<FFaceTriangle>& FaceTriangles) override;
+		virtual bool AddFaceTriangle(const FFaceTriangle& FaceTriangle) override
+		{
+			return AddFaceTriangles({ FaceTriangle });
+		}
+		virtual void EndFaceTriangles() override;
+
+		virtual void FinalizeMesh() override;
+
+	private:
+		template<typename VertexType>
+		struct TMeshDescEntity
+		{
+			union
+			{
+				struct
+				{
+					VertexType A;
+					VertexType B;
+					VertexType C;
+				};
+
+				VertexType ABC[3] = { INDEX_NONE, INDEX_NONE, INDEX_NONE };
+			};
+
+			TMeshDescEntity() = default;
+
+			TMeshDescEntity(VertexType ValA, VertexType ValB, VertexType ValC)
+				: A(ValA), B(ValB), C(ValC)
+			{}
+
+			VertexType& operator[](int Index)
+			{
+				return ABC[Index];
+			}
+
+			const VertexType& operator[](int Index) const
+			{
+				return ABC[Index];
+			}
+		};
+
+		using FVertexID3 = TMeshDescEntity<FVertexID>;
+		using FVertexInstanceID3 = TMeshDescEntity<FVertexInstanceID>;
+
+		FPolygonGroupID GetPolygonGroupID(uint32 MaterialID)
+		{
+			if (FPolygonGroupID* PolygonGroupIDPtr = MaterialToPolygonGroupMapping.Find(MaterialID))
+			{
+				return *PolygonGroupIDPtr;
+			}
+
+			if (MaterialToPolygonGroupMapping.Num() < NANITE_MAX_CLUSTER_MATERIALS)
+			{
+				FName ImportedSlotName = *LexToString(MaterialID);
+				LastPolygonGroupID = Mesh.CreatePolygonGroup();
+				PolygonGroupImportedMaterialSlotNames[LastPolygonGroupID] = ImportedSlotName;
+				MaterialToPolygonGroupMapping.Add(MaterialID, LastPolygonGroupID);
+			}
+
+			return LastPolygonGroupID;
+		}
+
+		// bIsBoundary is only modified if an edge exists
+		FEdgeID FindEdge(const FVertexID& Start, const FVertexID& End, bool& bIsBoundary)
+		{
+			FEdgeID EdgeID = Mesh.GetVertexPairEdge(Start, End);
+			if (EdgeID != INDEX_NONE)
+			{
+				TArrayView<const FTriangleID> Triangles = Mesh.GetEdgeConnectedTriangleIDs(EdgeID);
+				bIsBoundary = Triangles.Num() < 2;
+			}
+
+			return EdgeID;
+		}
+
+		bool GetVertexInstances(const FArray3i& Vertices, FVertexInstanceID3& VertexInstances);
+		void AddSymmetry();
+
+	private:
+		int32 VertexIndexOffset = 0;
+		FCADKernelStaticMeshAttributes Attributes;
+		TVertexAttributesRef<FVector3f> VertexPositions;
+		TVertexInstanceAttributesRef<FVector3f> VertexInstanceNormals;
+		TVertexInstanceAttributesRef<FVector3f> VertexInstanceTangents;
+		TVertexInstanceAttributesRef<float> VertexInstanceBinormalSigns;
+		TVertexInstanceAttributesRef<FVector4f> VertexInstanceColors;
+		TVertexInstanceAttributesRef<FVector2f> VertexInstanceUVs;
+		TPolygonAttributesRef<int32> PolygonAttributes;
+		TPolygonGroupAttributesRef<FName> PolygonGroupImportedMaterialSlotNames;
+
+		bool bSelectiveExtraction = false;
+		TMap<uint32, FPolygonGroupID> MaterialToPolygonGroupMapping;
+		FPolygonGroupID LastPolygonGroupID = -1;
+		TArray<FVertexID> VertexIDs;
+		TArray<FVector3f> Normals;
+		TArray<FVector2f> TexCoords;
+
+		FMeshDescription& Mesh;
+		bool bIsFinalized = false;
+	};
+}
diff --git a/Engine/Source/Runtime/Datasmith/CADKernel/Engine/Private/MeshUtilities.cpp b/Engine/Source/Runtime/Datasmith/CADKernel/Engine/Private/MeshUtilities.cpp
new file mode 100644
--- /dev/null
+++ b/Engine/Source/Runtime/Datasmith/CADKernel/Engine/Private/MeshUtilities.cpp
@@ -0,0 +1,27 @@
+// Copyright Epic Games, Inc. All Rights Reserved. 
+
+#include "MeshUtilities.h"
+
+#include "CADKernelEngineLog.h"
+
+#include "MeshDescriptionBuilder.h"
+#include "StaticMeshOperations.h"
+
+namespace UE::CADKernel::MeshUtilities
+{
+	bool FMeshOperations::OrientMesh(FMeshDescription& MeshDescription)
+	{
+		// #cad_import: TODO - Refactor MeshOperator::OrientMesh
+		return true;
+	}
+
+	void FMeshOperations::ResolveTJunctions(FMeshDescription& MeshDescription, double Tolerance)
+	{
+		// #cad_import: TODO - Refactor MeshOperator::ResolveTJunctions
+	}
+
+	void FMeshOperations::RecomputeNullNormal(FMeshDescription& MeshDescriptione)
+	{
+		// #cad_import: TODO Refactor MeshOperator::RecomputeNullNormal
+	}
+}
\ No newline at end of file
diff --git a/Engine/Source/Runtime/Datasmith/CADKernel/Engine/Private/MeshUtilities.h b/Engine/Source/Runtime/Datasmith/CADKernel/Engine/Private/MeshUtilities.h
new file mode 100644
--- /dev/null
+++ b/Engine/Source/Runtime/Datasmith/CADKernel/Engine/Private/MeshUtilities.h
@@ -0,0 +1,198 @@
+// Copyright Epic Games, Inc. All Rights Reserved. 
+
+#pragma once
+
+#include "ParametricSurfaceData.h"
+
+#include "IntVectorTypes.h"
+
+namespace UE::CADKernel
+{
+	class FModelMesh;
+	struct FMeshExtractionContext;
+}
+
+namespace UE::Geometry
+{
+	class FDynamicMesh3;
+}
+
+class FMeshDecription;
+
+using FArray3i = UE::Geometry::FVector3i;
+
+namespace UE::CADKernel::MeshUtilities
+{
+	class FMeshWrapperAbstract
+	{
+	public:
+		struct FFaceTriangle
+		{
+			int32 GroupID;
+			uint32 MaterialID;
+			
+			// Indices referencing the array of positions set calling SetVertices or the latest call to AddNewVertices
+			const FArray3i VertexIndices;
+
+			// Indices referencing the array of normals, Normals, set calling StartFaceTriangles.
+			const FArray3i Normals;
+
+			// Indices referencing the array of texture coordinates, TexCoords, set calling StartFaceTriangles.
+			const FArray3i TexCoords;
+		};
+
+		FMeshWrapperAbstract(const FMeshExtractionContext& InContext) : Context(InContext) {}
+
+		virtual ~FMeshWrapperAbstract() {}
+
+		virtual void FinalizeMesh() = 0;
+
+		virtual bool SetVertices(TArray<FVector>&& Vertices) = 0;
+
+		virtual bool AddNewVertices(TArray<FVector>&& Vertices) = 0;
+
+		virtual bool ReserveNewTriangles(int32 AdditionalTriangleCount) = 0;
+
+		/*
+		 ** Expected Normals.Num() == TexCoords.Num() == 3 * TriangleCount
+		 */
+		virtual bool StartFaceTriangles(int32 TriangleCount, const TArray<FVector3f>& Normals, const TArray<FVector2f>& TexCoords) = 0;
+		
+		virtual bool AddFaceTriangles(const TArray<FFaceTriangle>& FaceTriangles) = 0;
+
+		virtual bool AddFaceTriangle(const FFaceTriangle& FaceTriangle) = 0;
+
+		virtual void EndFaceTriangles() = 0;
+
+		/*
+		 * Normals and TexCoords array are expected to be arrays of 3 elements
+		 * Each value in those arrays is associated to the vertex in VertexIndices at the same index value, 0, 1 and 2.
+		 */
+		virtual bool AddTriangle(int32 GroupID, uint32 MaterialID, const FArray3i& VertexIndices, const TArrayView<FVector3f>& Normals, const TArrayView<FVector2f>& TexCoords) = 0;
+
+		static TSharedPtr<FMeshWrapperAbstract> MakeWrapper(const FMeshExtractionContext& InContext, FMeshDescription& Mesh);
+		static TSharedPtr<FMeshWrapperAbstract> MakeWrapper(const FMeshExtractionContext& Context, UE::Geometry::FDynamicMesh3& Mesh);
+
+	protected:
+		static constexpr FArray3i Clockwise{ 0, 1, 2 };
+		static constexpr FArray3i CounterClockwise{ 0, 2, 1 };
+
+		const FMeshExtractionContext& Context;
+		bool bAreVerticesSet = false;
+
+		// #cad_import: The UV scaling should be done in CADKernel
+		static constexpr double ScaleUV = 0.001; // mm to m
+	};
+
+	template<typename VecType>
+	void ConvertVectorArray(uint8 ModelCoordSys, TArray<VecType>& Array)
+	{
+		switch (ECADKernelModelCoordSystem(ModelCoordSys))
+		{
+		case ECADKernelModelCoordSystem::YUp_LeftHanded:
+			for (VecType& Vector : Array)
+			{
+				Vector.Set(Vector[2], Vector[0], Vector[1]);
+			}
+			break;
+
+		case ECADKernelModelCoordSystem::YUp_RightHanded:
+			for (VecType& Vector : Array)
+			{
+				Vector.Set(-Vector[2], Vector[0], Vector[1]);
+			}
+			break;
+
+		case ECADKernelModelCoordSystem::ZUp_RightHanded:
+			for (VecType& Vector : Array)
+			{
+				Vector.Set(-Vector[0], Vector[1], Vector[2]);
+			}
+			break;
+
+		case ECADKernelModelCoordSystem::ZUp_RightHanded_FBXLegacy:
+			for (VecType& Vector : Array)
+			{
+				Vector.Set(Vector[0], -Vector[1], Vector[2]);
+			}
+			break;
+
+		case ECADKernelModelCoordSystem::ZUp_LeftHanded:
+		default:
+			break;
+		}
+	}
+
+	template<typename VecType>
+	VecType ConvertVector(uint8 ModelCoordSys, const VecType& V)
+	{
+		switch (ECADKernelModelCoordSystem(ModelCoordSys))
+		{
+		case ECADKernelModelCoordSystem::YUp_LeftHanded:
+			return VecType(V[2], V[0], V[1]);
+
+		case ECADKernelModelCoordSystem::YUp_RightHanded:
+			return VecType(-V[2], V[0], V[1]);
+
+		case ECADKernelModelCoordSystem::ZUp_RightHanded:
+			return VecType(-V[0], V[1], V[2]);
+
+		case ECADKernelModelCoordSystem::ZUp_RightHanded_FBXLegacy:
+			return VecType(V[0], -V[1], V[2]);
+
+		case ECADKernelModelCoordSystem::ZUp_LeftHanded:
+		default:
+			return VecType(V[0], V[1], V[2]);
+		}
+	}
+
+	template<typename Type>
+	UE::Math::TMatrix<Type> GetSymmetricMatrix(const UE::Math::TVector<Type>& Origin, const UE::Math::TVector<Type>& Normal)
+	{
+		using namespace UE::Math;
+		//Calculate symmetry matrix
+		//(Px, Py, Pz) = normal
+		// -Px*Px + Pz*Pz + Py*Py  |  - 2 * Px * Py           |  - 2 * Px * Pz
+		// - 2 * Py * Px           |  - Py*Py + Px*Px + Pz*Pz |  - 2 * Py * Pz
+		// - 2 * Pz * Px           |  - 2 * Pz * Py           |  - Pz*Pz + Py*Py + Px*Px
+
+		TVector<Type> LocOrigin = Origin;
+
+		Type NormalXSqr = Normal.X * Normal.X;
+		Type NormalYSqr = Normal.Y * Normal.Y;
+		Type NormalZSqr = Normal.Z * Normal.Z;
+
+		TMatrix<Type> OSymmetricMatrix;
+		OSymmetricMatrix.SetIdentity();
+		TVector<Type> Axis0(-NormalXSqr + NormalZSqr + NormalYSqr, -2 * Normal.X * Normal.Y, -2 * Normal.X * Normal.Z);
+		TVector<Type> Axis1(-2 * Normal.Y * Normal.X, -NormalYSqr + NormalXSqr + NormalZSqr, -2 * Normal.Y * Normal.Z);
+		TVector<Type> Axis2(-2 * Normal.Z * Normal.X, -2 * Normal.Z * Normal.Y, -NormalZSqr + NormalYSqr + NormalXSqr);
+		OSymmetricMatrix.SetAxes(&Axis0, &Axis1, &Axis2);
+
+		TMatrix<Type> SymmetricMatrix;
+		SymmetricMatrix.SetIdentity();
+
+		//Translate to 0, 0, 0
+		LocOrigin *= -1.;
+		SymmetricMatrix.SetOrigin(LocOrigin);
+
+		//// Apply Symmetric
+		SymmetricMatrix *= OSymmetricMatrix;
+
+		//Translate to original position
+		LocOrigin *= -1.;
+		TMatrix<Type> OrigTranslation;
+		OrigTranslation.SetIdentity();
+		OrigTranslation.SetOrigin(LocOrigin);
+		SymmetricMatrix *= OrigTranslation;
+
+		return SymmetricMatrix;
+	}
+
+	class FMeshOperations
+	{
+		static bool OrientMesh(FMeshDescription& MeshDescription);
+		static void ResolveTJunctions(FMeshDescription& MeshDescription, double Tolerance);
+		static void RecomputeNullNormal(FMeshDescription& MeshDescriptione);
+	};
+}
\ No newline at end of file
diff --git a/Engine/Source/Runtime/Datasmith/CADKernel/Engine/Private/ParametricSurfaceData.cpp b/Engine/Source/Runtime/Datasmith/CADKernel/Engine/Private/ParametricSurfaceData.cpp
new file mode 100644
--- /dev/null
+++ b/Engine/Source/Runtime/Datasmith/CADKernel/Engine/Private/ParametricSurfaceData.cpp
@@ -0,0 +1,68 @@
+// Copyright Epic Games, Inc. All Rights Reserved.
+
+#include "ParametricSurfaceData.h"
+
+#include "Core/Session.h"
+#include "Topo/Model.h"
+
+#include "Misc/FileHelper.h"
+#include "Misc/Paths.h"
+#include "UObject/EnterpriseObjectVersion.h"
+
+#define LOCTEXT_NAMESPACE "ParametricSurfaceData"
+
+const double FUnitConverter::CentimeterToMillimeter = 0.1;
+const double FUnitConverter::MillimeterToCentimeter = 10.;
+const double FUnitConverter::CentimeterToMeter = 0.01;
+const double FUnitConverter::MillimeterToMeter = 0.001;
+
+bool UParametricSurfaceData::SetFromFile(const TCHAR* FilePath, bool bForTechSoft)
+{
+	if (FPaths::FileExists(FilePath))
+	{
+		TArray<uint8> ByteArray;
+
+		if (FFileHelper::LoadFileToArray(ByteArray, FilePath))
+		{
+			TArray<uint8>& RawData = bForTechSoft ? TechSoftRawData : CADKernelRawData;
+			RawData = MoveTemp(ByteArray);
+			return true;
+		}
+	}
+
+	return false;
+}
+
+void UParametricSurfaceData::Serialize(FArchive& Ar)
+{
+	Ar.UsingCustomVersion(FEnterpriseObjectVersion::GUID);
+
+	if (Ar.IsLoading() && Ar.CustomVer(FEnterpriseObjectVersion::GUID) < FEnterpriseObjectVersion::AddedParametricSurfaceData)
+	{
+		// #cad_import_todo : Put code to deserialize UDatasmithParametricSurfaceData
+		 return;
+	}
+
+	Super::Serialize(Ar);
+
+	Ar << CADKernelRawData;
+	Ar << TechSoftRawData;
+}
+
+TSharedPtr<UE::CADKernel::FModel> UParametricSurfaceData::LoadModel()
+{
+	using namespace UE::CADKernel;
+
+	if (CADKernelRawData.IsEmpty())
+	{
+		return {};
+	}
+
+	TSharedRef<FSession> Session = MakeShared<FSession>(LastTessellationSettings.GetGeometricTolerance(true));
+	Session->AddDatabase(CADKernelRawData);
+
+	return Session->GetModelAsShared();
+}
+
+#undef LOCTEXT_NAMESPACE // "ParametricSurfaceData"
+
diff --git a/Engine/Source/Runtime/Datasmith/CADKernel/Engine/Private/TechSoft/CADKernelToTechSoft.cpp b/Engine/Source/Runtime/Datasmith/CADKernel/Engine/Private/TechSoft/CADKernelToTechSoft.cpp
new file mode 100644
--- /dev/null
+++ b/Engine/Source/Runtime/Datasmith/CADKernel/Engine/Private/TechSoft/CADKernelToTechSoft.cpp
@@ -0,0 +1,28 @@
+// Copyright Epic Games, Inc. All Rights Reserved.
+
+#include "CADKernelEngine.h"
+
+#ifdef WITH_HOOPS
+#include "TechSoftUniqueObjectImpl.h"
+
+namespace UE::CADKernel
+{
+	A3DRiRepresentationItem* FTechSoftUtilities::CADKernelToTechSoft(TSharedPtr<FModel>& Model)
+	{
+		if (!FTechSoftLibrary::IsInitialized())
+		{
+			return nullptr;
+		}
+
+		return nullptr;
+	}
+}
+#else
+namespace UE::CADKernel
+{
+	A3DRiRepresentationItem* FTechSoftUtilities::CADKernelToTechSoft(TSharedPtr<FModel>& Model)
+	{
+		return nullptr;
+	}
+}
+#endif
diff --git a/Engine/Source/Runtime/Datasmith/CADKernel/Engine/Private/TechSoft/TechSoftLibrary.cpp b/Engine/Source/Runtime/Datasmith/CADKernel/Engine/Private/TechSoft/TechSoftLibrary.cpp
new file mode 100644
--- /dev/null
+++ b/Engine/Source/Runtime/Datasmith/CADKernel/Engine/Private/TechSoft/TechSoftLibrary.cpp
@@ -0,0 +1,63 @@
+// Copyright Epic Games, Inc. All Rights Reserved.
+
+// Required to define INITIALIZE_A3D_API in order to call A3DSDKLoadLibraryA
+#define INITIALIZE_A3D_API
+
+#include "CADKernelEngine.h"
+
+#include "HAL/PlatformProcess.h"
+#include "Misc/Paths.h"
+
+namespace UE::CADKernel
+{
+	bool FTechSoftLibrary::bIsInitialized = false;
+
+#ifdef WITH_HOOPS
+	bool FTechSoftLibrary::Initialize()
+	{
+		if (bIsInitialized)
+		{
+			return true;
+		}
+
+		// #cad_import: Cover case when called in a game
+		FString TechSoftDllPath = FPaths::Combine(FPaths::EngineDir(), TEXT("Binaries/ThirdParty/TechSoft"), FPlatformProcess::GetBinariesSubdirectory());
+		TechSoftDllPath = FPaths::ConvertRelativePathToFull(TechSoftDllPath);
+
+		if (A3DSDKLoadLibraryA(TCHAR_TO_UTF8(*TechSoftDllPath)))
+		{
+			A3DStatus Status = A3DLicPutUnifiedLicense(HOOPS_LICENSE);
+
+			A3DInt32 iMajorVersion = 0, iMinorVersion = 0;
+			Status = A3DDllGetVersion(&iMajorVersion, &iMinorVersion);
+			if (Status == A3D_SUCCESS)
+			{
+				Status = A3DDllInitialize(A3D_DLL_MAJORVERSION, A3D_DLL_MINORVERSION);
+				if (Status == A3D_SUCCESS)
+				{
+					bIsInitialized = true;
+					return true;
+				}
+			}
+		}
+
+		return false;
+	}
+
+	const TCHAR* FTechSoftLibrary::GetVersion()
+	{
+		static const FString Version = TEXT("TechSoft ") + FString::FromInt(A3D_DLL_MAJORVERSION) + TEXT(".") + FString::FromInt(A3D_DLL_MINORVERSION) + TEXT(".") + FString::FromInt(A3D_DLL_UPDATEVERSION);
+		return bIsInitialized ? *Version : TEXT("TechSoft uninitialized");
+	}
+#else
+	bool FTechSoftLibrary::Initialize()
+	{
+		return false;
+	}
+
+	const TCHAR* FTechSoftLibrary::GetVersion()
+	{
+		return TEXT("TechSoft unavailable");
+	}
+#endif
+}
\ No newline at end of file
diff --git a/Engine/Source/Runtime/Datasmith/CADKernel/Engine/Private/TechSoft/TechSoftMesh.cpp b/Engine/Source/Runtime/Datasmith/CADKernel/Engine/Private/TechSoft/TechSoftMesh.cpp
new file mode 100644
--- /dev/null
+++ b/Engine/Source/Runtime/Datasmith/CADKernel/Engine/Private/TechSoft/TechSoftMesh.cpp
@@ -0,0 +1,827 @@
+// Copyright Epic Games, Inc. All Rights Reserved.
+
+#include "CADKernelEngine.h"
+#include "TechSoftUtilities.h"
+
+#ifdef WITH_HOOPS
+#include "MeshUtilities.h"
+#include "TechSoftUniqueObjectImpl.h"
+
+namespace UE::CADKernel::TechSoftUtilities
+{
+	using namespace UE::CADKernel;
+
+	class FRepresentationToMeshHelper
+	{
+	public:
+		FRepresentationToMeshHelper(const A3DRiRepresentationItem* RepresentationItem, MeshUtilities::FMeshWrapperAbstract& InMeshInterface)
+			: MeshInterface(InMeshInterface)
+		{
+			TechSoft::TUniqueObject<A3DRiRepresentationItemData> RepresentationItemData(RepresentationItem);
+			if (!RepresentationItemData.IsValid())
+			{
+				return;
+			}
+
+			A3DEEntityType Type;
+			A3DEntityGetType(RepresentationItemData->m_pTessBase, &Type);
+			if (Type != kA3DTypeTess3D)
+			{
+				return;
+			}
+
+			TessellationBase = RepresentationItemData->m_pTessBase;
+
+			if (TessellationBase)
+			{
+				A3DEntityGetType(RepresentationItem, &Type);
+				if (Type == kA3DTypeRiBrepModel)
+				{
+					TechSoft::TUniqueObject<A3DRiBrepModelData> BRepModelData(RepresentationItem);
+					if (BRepModelData.IsValid())
+					{
+						BRepData = BRepModelData->m_pBrepData;
+					}
+				}
+			}
+		}
+
+		bool IsValid() const { return TessellationBase != nullptr; }
+
+		bool Convert();
+
+	private:
+		struct FTriangleContext
+		{
+			int32 GroupID = 0;
+			uint32 MaterialID = 0;
+			FVector2f UVScale;
+		};
+
+		bool CollectBRepFaces(TArray<A3DTopoFace*>& BRepFaces, int32& TriangleCount);
+		int32 GetTriangleCount(const A3DTessFaceData& FaceTessData);
+		bool ParseTriangles(const TArray<A3DTopoFace*>& BRepFaces);
+
+		void GetUVScale(const A3DTopoFace* TopoFace, FVector2f& UVScale);
+
+		void AddFaceTriangle(const FTriangleContext& Context, const uint32 TriangleCount, uint32& StartIndexInOut);
+		void AddFaceTriangleWithUniqueNormal(const FTriangleContext& Context, const uint32 TriCount, uint32& StartIndexInOut);
+		void AddFaceTriangleWithUniqueNormalAndTexture(const FTriangleContext& Context, const uint32 TriCount, const uint32 TexOffset, uint32& StartIndexInOut);
+		void AddFaceTriangleWithTexture(const FTriangleContext& Context, const uint32 TriCount, const uint32 TexOffset, uint32& StartIndexInOut);
+		
+		void AddFaceTriangleFan(const FTriangleContext& Context, const uint32 TriCount, uint32& StartIndexInOut);
+		void AddFaceTriangleFanWithUniqueNormal(const FTriangleContext& Context, const uint32 TriCount, uint32& StartIndexInOut);
+		void AddFaceTriangleFanWithUniqueNormalAndTexture(const FTriangleContext& Context, const uint32 TriCount, const uint32 TexOffset, uint32& StartIndexInOut);
+		void AddFaceTriangleFanWithTexture(const FTriangleContext& Context, const uint32 TriCount, const uint32 TexOffset, uint32& StartIndexInOut);
+		
+		void AddFaceTriangleStripe(const FTriangleContext& Context, const uint32 TriCount, uint32& StartIndexInOut);
+		void AddFaceTriangleStripeWithTexture(const FTriangleContext& Context, const uint32 TriCount, const uint32 TexOffset, uint32& StartIndexInOut);
+		void AddFaceTriangleStripeWithUniqueNormal(const FTriangleContext& Context, const uint32 TriCount, uint32& StartIndexInOut);
+		void AddFaceTriangleStripeWithUniqueNormalAndTexture(const FTriangleContext& Context, const uint32 TriCount, const uint32 TexOffset, uint32& StartIndexInOut);
+
+		void AddTriangle(const FTriangleContext& Context, const FArray3i& VertexIndices, const FArray3i& NormalIndices, const FArray3i& TexCoordIndicies)
+		{
+			FVector3f Normals[3]{
+				{(float)TessellationNormals[NormalIndices[0]], (float)TessellationNormals[NormalIndices[0] + 1], (float)TessellationNormals[NormalIndices[0] + 2]},
+				{(float)TessellationNormals[NormalIndices[1]], (float)TessellationNormals[NormalIndices[1] + 1], (float)TessellationNormals[NormalIndices[1] + 2]},
+				{(float)TessellationNormals[NormalIndices[2]], (float)TessellationNormals[NormalIndices[2] + 1], (float)TessellationNormals[NormalIndices[2] + 2]},
+			};
+			FVector2f TexCoords[3]{
+				{(float)TessellationTexCoords[TexCoordIndicies[0]] * Context.UVScale[0], (float)TessellationTexCoords[TexCoordIndicies[0] + 1] * Context.UVScale[1]},
+				{(float)TessellationTexCoords[TexCoordIndicies[1]] * Context.UVScale[0], (float)TessellationTexCoords[TexCoordIndicies[1] + 1] * Context.UVScale[1]},
+				{(float)TessellationTexCoords[TexCoordIndicies[2]] * Context.UVScale[0], (float)TessellationTexCoords[TexCoordIndicies[2] + 1] * Context.UVScale[1]},
+			};
+
+			MeshInterface.AddTriangle(Context.GroupID, Context.MaterialID, VertexIndices, {Normals, 3 }, {TexCoords, 3});
+		}
+
+	private:
+		MeshUtilities::FMeshWrapperAbstract& MeshInterface;
+		const A3DTess3D* TessellationBase = nullptr;
+		const A3DTopoBrepData* BRepData = nullptr;
+
+		int32 LastVertexIndex = 0;
+		A3DUns32* TriangulatedIndexes = nullptr;
+		double* TessellationNormals = nullptr;
+		double* TessellationTexCoords = nullptr;				/*!< Array of \ref A3DDouble, as texture coordinates. */
+	};
+
+	bool AddRepresentation(A3DRiRepresentationItem* RepresentationItem, MeshUtilities::FMeshWrapperAbstract& MeshInterface)
+	{
+		FRepresentationToMeshHelper Helper(RepresentationItem, MeshInterface);
+
+		return Helper.IsValid() ? Helper.Convert() : false;
+	}
+
+	bool FRepresentationToMeshHelper::CollectBRepFaces(TArray<A3DTopoFace*>& BRepFaces, int32& TriangleCount)
+	{
+		BRepFaces.Empty();
+		TriangleCount = 0;
+
+		if (BRepData == nullptr)
+		{
+			return false;
+		}
+
+		// #cad_import: Check what this is all about
+		//TechSoft::TUniqueObject<A3DTopoBodyData> TopoBodyData(BRepData);
+		//if (TopoBodyData.IsValid())
+		//{
+		//	if (TopoBodyData->m_pContext)
+		//	{
+		//		TechSoft::TUniqueObject<A3DTopoContextData> TopoContextData(TopoBodyData->m_pContext);
+		//		if (TopoContextData.IsValid())
+		//		{
+		//			if (TopoContextData->m_bHaveScale)
+		//			{
+		//				TextureUnit *= TopoContextData->m_dScale;
+		//			}
+		//		}
+		//	}
+		//}
+
+		TechSoft::TUniqueObject<A3DTopoBrepDataData> TopoBrepData(BRepData);
+		if (!TopoBrepData.IsValid())
+		{
+			return false;
+		}
+
+		for (A3DUns32 Index = 0; Index < TopoBrepData->m_uiConnexSize; ++Index)
+		{
+			TechSoft::TUniqueObject<A3DTopoConnexData> TopoConnexData(TopoBrepData->m_ppConnexes[Index]);
+			if (TopoConnexData.IsValid())
+			{
+				for (A3DUns32 Sndex = 0; Sndex < TopoConnexData->m_uiShellSize; ++Sndex)
+				{
+					TechSoft::TUniqueObject<A3DTopoShellData> ShellData(TopoConnexData->m_ppShells[Sndex]);
+					if (ShellData.IsValid())
+					{
+						for (A3DUns32 Fndex = 0; Fndex < ShellData->m_uiFaceSize; ++Fndex)
+						{
+							BRepFaces.Add(ShellData->m_ppFaces[Fndex]);
+						}
+					}
+				}
+			}
+		}
+
+		if (BRepFaces.IsEmpty())
+		{
+			return false;
+		}
+
+		TechSoft::TUniqueObject<A3DTess3DData> TessellationData(TessellationBase);
+
+		if (!TessellationData.IsValid() || TessellationData->m_uiFaceTessSize == 0)
+		{
+			return false;
+		}
+
+		if (TessellationData->m_uiFaceTessSize != BRepFaces.Num())
+		{
+			BRepFaces.Empty();
+			return false;
+		}
+
+		for (A3DUns32 Index = 0; Index < TessellationData->m_uiFaceTessSize; ++Index)
+		{
+			TriangleCount += GetTriangleCount(TessellationData->m_psFaceTessData[Index]);
+		}
+
+		return TriangleCount > 0;
+	}
+
+	int32 FRepresentationToMeshHelper::GetTriangleCount(const A3DTessFaceData& FaceTessData)
+	{
+		constexpr int32 TessellationFaceDataWithTriangle   = 0x2222;
+		constexpr int32 TessellationFaceDataWithFan        = 0x4444;
+		constexpr int32 TessellationFaceDataWithStripe     = 0x8888;
+		constexpr int32 TessellationFaceDataWithOneNormal  = 0xE0E0;
+
+		uint32 UsedEntitiesFlags = FaceTessData.m_usUsedEntitiesFlags;
+
+		int32 TriangleCount = 0;
+		uint32 FaceSetIndex = 0;
+		
+		if (UsedEntitiesFlags & TessellationFaceDataWithTriangle)
+		{
+			TriangleCount += FaceTessData.m_puiSizesTriangulated[FaceSetIndex];
+			FaceSetIndex++;
+		}
+
+		if (FaceTessData.m_uiSizesTriangulatedSize > FaceSetIndex)
+		{
+			if (UsedEntitiesFlags & TessellationFaceDataWithFan)
+			{
+				uint32 LastFanIndex = 1 + FaceSetIndex + FaceTessData.m_puiSizesTriangulated[FaceSetIndex];
+				FaceSetIndex++;
+				for (; FaceSetIndex < LastFanIndex; FaceSetIndex++)
+				{
+					int32 FanSize = (int32)(FaceTessData.m_puiSizesTriangulated[FaceSetIndex] & kA3DTessFaceDataNormalMask);
+					TriangleCount += (FanSize - 2);
+				}
+			}
+		}
+
+		if (FaceTessData.m_uiSizesTriangulatedSize > FaceSetIndex)
+		{
+			FaceSetIndex++;
+			for (; FaceSetIndex < FaceTessData.m_uiSizesTriangulatedSize; FaceSetIndex++)
+			{
+				int32 StripeSize = (int32)(FaceTessData.m_puiSizesTriangulated[FaceSetIndex] & kA3DTessFaceDataNormalMask);
+				TriangleCount += (StripeSize - 2);
+			}
+		}
+
+		return TriangleCount;
+	}
+
+	void FRepresentationToMeshHelper::AddFaceTriangle(const FTriangleContext& Context, const uint32 TriCount, uint32& StartIndexInOut)
+	{
+		FArray3i VertexIndices;
+		FArray3i NormalIndices;
+		FArray3i TexCoordIndices(0, 0, 0); // #cadkernel_check Why no texture coordinates?
+
+		// Get Triangles
+		for (uint32 FArray3i = 0; FArray3i < TriCount; FArray3i++)
+		{
+			NormalIndices[0] = TriangulatedIndexes[StartIndexInOut++];
+			VertexIndices[0] = TriangulatedIndexes[StartIndexInOut++] / 3;
+			NormalIndices[1] = TriangulatedIndexes[StartIndexInOut++];
+			VertexIndices[1] = TriangulatedIndexes[StartIndexInOut++] / 3;
+			NormalIndices[2] = TriangulatedIndexes[StartIndexInOut++];
+			VertexIndices[2] = TriangulatedIndexes[StartIndexInOut++] / 3;
+
+			AddTriangle(Context, VertexIndices, NormalIndices, TexCoordIndices);
+		}
+	}
+
+	void FRepresentationToMeshHelper::AddFaceTriangleWithUniqueNormal(const FTriangleContext& Context, const uint32 TriCount, uint32& StartIndexInOut)
+	{
+		FArray3i VertexIndices;
+		FArray3i NormalIndices;
+		FArray3i TexCoordIndices(0, 0, 0);
+
+		// Get Triangles
+		for (uint32 FArray3i = 0; FArray3i < TriCount; FArray3i++)
+		{
+			NormalIndices[0] = TriangulatedIndexes[StartIndexInOut++];
+			NormalIndices[1] = NormalIndices[0];
+			NormalIndices[2] = NormalIndices[0];
+
+			VertexIndices[0] = TriangulatedIndexes[StartIndexInOut++] / 3;
+			VertexIndices[1] = TriangulatedIndexes[StartIndexInOut++] / 3;
+			VertexIndices[2] = TriangulatedIndexes[StartIndexInOut++] / 3;
+
+			AddTriangle(Context, VertexIndices, NormalIndices, TexCoordIndices);
+		}
+	}
+
+	void FRepresentationToMeshHelper::AddFaceTriangleWithUniqueNormalAndTexture(const FTriangleContext& Context, const uint32 TriCount, const uint32 TexOffset, uint32& StartIndexInOut)
+	{
+		FArray3i VertexIndices;
+		FArray3i NormalIndices;
+		FArray3i TexCoordIndices;
+
+		// Get Triangles
+		for (uint32 FArray3i = 0; FArray3i < TriCount; FArray3i++)
+		{
+			NormalIndices[0] = TriangulatedIndexes[StartIndexInOut++];
+			NormalIndices[1] = NormalIndices[0];
+			NormalIndices[2] = NormalIndices[0];
+
+			TexCoordIndices[0] = TriangulatedIndexes[StartIndexInOut];
+			StartIndexInOut += TexOffset;
+			VertexIndices[0] = TriangulatedIndexes[StartIndexInOut++] / 3;
+
+			TexCoordIndices[1] = TriangulatedIndexes[StartIndexInOut];
+			StartIndexInOut += TexOffset;
+			VertexIndices[1] = TriangulatedIndexes[StartIndexInOut++] / 3;
+
+			TexCoordIndices[2] = TriangulatedIndexes[StartIndexInOut];
+			StartIndexInOut += TexOffset;
+			VertexIndices[2] = TriangulatedIndexes[StartIndexInOut++] / 3;
+
+			AddTriangle(Context, VertexIndices, NormalIndices, TexCoordIndices);
+		}
+	}
+
+	void FRepresentationToMeshHelper::AddFaceTriangleWithTexture(const FTriangleContext& Context, const uint32 TriCount, const uint32 TexOffset, uint32& StartIndexInOut)
+	{
+		FArray3i VertexIndices;
+		FArray3i NormalIndices;
+		FArray3i TexCoordIndices;
+
+		// Get Triangles
+		for (uint32 FArray3i = 0; FArray3i < TriCount; FArray3i++)
+		{
+			NormalIndices[0] = TriangulatedIndexes[StartIndexInOut++];
+			TexCoordIndices[0] = TriangulatedIndexes[StartIndexInOut];
+			StartIndexInOut += TexOffset;
+			VertexIndices[0] = TriangulatedIndexes[StartIndexInOut++] / 3;
+
+			NormalIndices[1] = TriangulatedIndexes[StartIndexInOut++];
+			TexCoordIndices[1] = TriangulatedIndexes[StartIndexInOut];
+			StartIndexInOut += TexOffset;
+			VertexIndices[1] = TriangulatedIndexes[StartIndexInOut++] / 3;
+
+			NormalIndices[2] = TriangulatedIndexes[StartIndexInOut++];
+			TexCoordIndices[2] = TriangulatedIndexes[StartIndexInOut];
+			StartIndexInOut += TexOffset;
+			VertexIndices[2] = TriangulatedIndexes[StartIndexInOut++] / 3;
+
+			AddTriangle(Context, VertexIndices, NormalIndices, TexCoordIndices);
+		}
+	}
+
+	void FRepresentationToMeshHelper::AddFaceTriangleFan(const FTriangleContext& Context, const uint32 TriCount, uint32& StartIndexInOut)
+	{
+		FArray3i VertexIndices;
+		FArray3i NormalIndices;
+		FArray3i TexCoordIndices(0, 0, 0);
+
+		NormalIndices[0] = TriangulatedIndexes[StartIndexInOut++];
+		VertexIndices[0] = TriangulatedIndexes[StartIndexInOut++] / 3;
+		NormalIndices[1] = TriangulatedIndexes[StartIndexInOut++];
+		VertexIndices[1] = TriangulatedIndexes[StartIndexInOut++] / 3;
+
+		for (unsigned long FArray3i = 2; FArray3i < TriCount; FArray3i++)
+		{
+			NormalIndices[2] = TriangulatedIndexes[StartIndexInOut++];
+			VertexIndices[2] = TriangulatedIndexes[StartIndexInOut++] / 3;
+
+			AddTriangle(Context, VertexIndices, NormalIndices, TexCoordIndices);
+
+			NormalIndices[1] = NormalIndices[2];
+			VertexIndices[1] = VertexIndices[2];
+		}
+	}
+	void FRepresentationToMeshHelper::AddFaceTriangleFanWithUniqueNormal(const FTriangleContext& Context, const uint32 TriCount, uint32& StartIndexInOut)
+	{
+		FArray3i VertexIndices;
+		FArray3i NormalIndices;
+		FArray3i TexCoordIndices(0, 0, 0);
+
+
+		NormalIndices[0] = TriangulatedIndexes[StartIndexInOut++];
+		NormalIndices[1] = NormalIndices[0];
+		NormalIndices[2] = NormalIndices[0];
+
+		VertexIndices[0] = TriangulatedIndexes[StartIndexInOut++] / 3;
+		VertexIndices[1] = TriangulatedIndexes[StartIndexInOut++] / 3;
+
+		// Get Triangles
+		for (uint32 FArray3i = 2; FArray3i < TriCount; FArray3i++)
+		{
+			VertexIndices[2] = TriangulatedIndexes[StartIndexInOut++] / 3;
+
+			AddTriangle(Context, VertexIndices, NormalIndices, TexCoordIndices);
+
+			VertexIndices[1] = VertexIndices[2];
+		}
+	}
+
+	void FRepresentationToMeshHelper::AddFaceTriangleFanWithUniqueNormalAndTexture(const FTriangleContext& Context, const uint32 TriCount, const uint32 TexOffset, uint32& StartIndexInOut)
+	{
+		FArray3i VertexIndices;
+		FArray3i NormalIndices;
+		FArray3i TexCoordIndices;
+
+		NormalIndices[0] = TriangulatedIndexes[StartIndexInOut++];
+		NormalIndices[1] = NormalIndices[0];
+		NormalIndices[2] = NormalIndices[0];
+
+		TexCoordIndices[0] = TriangulatedIndexes[StartIndexInOut];
+		StartIndexInOut += TexOffset;
+		VertexIndices[0] = TriangulatedIndexes[StartIndexInOut++] / 3;
+
+		TexCoordIndices[1] = TriangulatedIndexes[StartIndexInOut];
+		StartIndexInOut += TexOffset;
+		VertexIndices[1] = TriangulatedIndexes[StartIndexInOut++] / 3;
+
+		for (uint32 FArray3i = 2; FArray3i < TriCount; FArray3i++)
+		{
+			TexCoordIndices[2] = TriangulatedIndexes[StartIndexInOut];
+			StartIndexInOut += TexOffset;
+			VertexIndices[2] = TriangulatedIndexes[StartIndexInOut++] / 3;
+
+			AddTriangle(Context, VertexIndices, NormalIndices, TexCoordIndices);
+
+			VertexIndices[1] = VertexIndices[2];
+			TexCoordIndices[1] = TexCoordIndices[2];
+		}
+	}
+
+	void FRepresentationToMeshHelper::AddFaceTriangleFanWithTexture(const FTriangleContext& Context, const uint32 TriCount, const uint32 TexOffset, uint32& StartIndexInOut)
+	{
+		FArray3i VertexIndices;
+		FArray3i NormalIndices;
+		FArray3i TexCoordIndices;
+
+		NormalIndices[0] = TriangulatedIndexes[StartIndexInOut++];
+		TexCoordIndices[0] = TriangulatedIndexes[StartIndexInOut];
+		StartIndexInOut += TexOffset;
+		VertexIndices[0] = TriangulatedIndexes[StartIndexInOut++] / 3;
+
+		NormalIndices[1] = TriangulatedIndexes[StartIndexInOut++];
+		TexCoordIndices[1] = TriangulatedIndexes[StartIndexInOut];
+		StartIndexInOut += TexOffset;
+		VertexIndices[1] = TriangulatedIndexes[StartIndexInOut++] / 3;
+
+		for (uint32 FArray3i = 2; FArray3i < TriCount; FArray3i++)
+		{
+			NormalIndices[2] = TriangulatedIndexes[StartIndexInOut++];
+			TexCoordIndices[2] = TriangulatedIndexes[StartIndexInOut];
+			StartIndexInOut += TexOffset;
+			VertexIndices[2] = TriangulatedIndexes[StartIndexInOut++] / 3;
+
+			AddTriangle(Context, VertexIndices, NormalIndices, TexCoordIndices);
+
+			NormalIndices[1] = NormalIndices[2];
+			TexCoordIndices[1] = TexCoordIndices[2];
+			VertexIndices[1] = VertexIndices[2];
+		}
+	}
+
+	void FRepresentationToMeshHelper::AddFaceTriangleStripe(const FTriangleContext& Context, const uint32 TriCount, uint32& StartIndexInOut)
+	{
+		FArray3i VertexIndices;
+		FArray3i NormalIndices;
+		FArray3i TexCoordIndices;
+
+		NormalIndices[0] = TriangulatedIndexes[StartIndexInOut++];
+		VertexIndices[0] = TriangulatedIndexes[StartIndexInOut++] / 3;
+		NormalIndices[1] = TriangulatedIndexes[StartIndexInOut++];
+		VertexIndices[1] = TriangulatedIndexes[StartIndexInOut++] / 3;
+
+		for (unsigned long FArray3i = 2; FArray3i < TriCount; FArray3i++)
+		{
+			NormalIndices[2] = TriangulatedIndexes[StartIndexInOut++];
+			VertexIndices[2] = TriangulatedIndexes[StartIndexInOut++] / 3;
+
+			AddTriangle(Context, VertexIndices, NormalIndices, TexCoordIndices);
+
+			FArray3i++;
+			if (FArray3i == TriCount)
+			{
+				break;
+			}
+
+			Swap(VertexIndices[1], VertexIndices[2]);
+			Swap(NormalIndices[1], NormalIndices[2]);
+
+			NormalIndices[0] = TriangulatedIndexes[StartIndexInOut++];
+			VertexIndices[0] = TriangulatedIndexes[StartIndexInOut++] / 3;
+
+			AddTriangle(Context, VertexIndices, NormalIndices, TexCoordIndices);
+
+			Swap(VertexIndices[0], VertexIndices[1]);
+			Swap(NormalIndices[0], NormalIndices[1]);
+		}
+	}
+
+	void FRepresentationToMeshHelper::AddFaceTriangleStripeWithTexture(const FTriangleContext& Context, const uint32 TriCount, const uint32 TexOffset, uint32& StartIndexInOut)
+	{
+		FArray3i VertexIndices;
+		FArray3i NormalIndices;
+		FArray3i TexCoordIndices;
+
+		NormalIndices[0] = TriangulatedIndexes[StartIndexInOut++];
+		TexCoordIndices[0] = TriangulatedIndexes[StartIndexInOut];
+		StartIndexInOut += TexOffset;
+		VertexIndices[0] = TriangulatedIndexes[StartIndexInOut++] / 3;
+		NormalIndices[1] = TriangulatedIndexes[StartIndexInOut++];
+		TexCoordIndices[1] = TriangulatedIndexes[StartIndexInOut];
+		StartIndexInOut += TexOffset;
+		VertexIndices[1] = TriangulatedIndexes[StartIndexInOut++] / 3;
+
+		for (unsigned long FArray3i = 2; FArray3i < TriCount; FArray3i++)
+		{
+			NormalIndices[2] = TriangulatedIndexes[StartIndexInOut++];
+			TexCoordIndices[2] = TriangulatedIndexes[StartIndexInOut];
+			StartIndexInOut += TexOffset;
+			VertexIndices[2] = TriangulatedIndexes[StartIndexInOut++] / 3;
+
+			AddTriangle(Context, VertexIndices, NormalIndices, TexCoordIndices);
+
+			FArray3i++;
+			if (FArray3i == TriCount)
+			{
+				break;
+			}
+
+			Swap(VertexIndices[1], VertexIndices[2]);
+			Swap(NormalIndices[1], NormalIndices[2]);
+			Swap(TexCoordIndices[1], TexCoordIndices[2]);
+
+			NormalIndices[0] = TriangulatedIndexes[StartIndexInOut++];
+			VertexIndices[0] = TriangulatedIndexes[StartIndexInOut++] / 3;
+
+			AddTriangle(Context, VertexIndices, NormalIndices, TexCoordIndices);
+
+			Swap(VertexIndices[0], VertexIndices[1]);
+			Swap(NormalIndices[0], NormalIndices[1]);
+			Swap(TexCoordIndices[0], TexCoordIndices[1]);
+		}
+	}
+
+	void FRepresentationToMeshHelper::AddFaceTriangleStripeWithUniqueNormal(const FTriangleContext& Context, const uint32 TriCount, uint32& StartIndexInOut)
+	{
+		FArray3i VertexIndices;
+		FArray3i NormalIndices;
+		FArray3i TexCoordIndices;
+
+		NormalIndices[0] = TriangulatedIndexes[StartIndexInOut++];
+		NormalIndices[1] = NormalIndices[0];
+		NormalIndices[2] = NormalIndices[0];
+
+		VertexIndices[0] = TriangulatedIndexes[StartIndexInOut++] / 3;
+		VertexIndices[1] = TriangulatedIndexes[StartIndexInOut++] / 3;
+
+		for (unsigned long FArray3i = 2; FArray3i < TriCount; FArray3i++)
+		{
+			VertexIndices[2] = TriangulatedIndexes[StartIndexInOut++] / 3;
+
+			AddTriangle(Context, VertexIndices, NormalIndices, TexCoordIndices);
+
+			FArray3i++;
+			if (FArray3i == TriCount)
+			{
+				break;
+			}
+
+			Swap(VertexIndices[1], VertexIndices[2]);
+
+			VertexIndices[0] = TriangulatedIndexes[StartIndexInOut++] / 3;
+
+			AddTriangle(Context, VertexIndices, NormalIndices, TexCoordIndices);
+
+			Swap(VertexIndices[0], VertexIndices[1]);
+		}
+	}
+
+	void FRepresentationToMeshHelper::AddFaceTriangleStripeWithUniqueNormalAndTexture(const FTriangleContext& Context, const uint32 TriCount, const uint32 TexOffset, uint32& StartIndexInOut)
+	{
+		FArray3i VertexIndices;
+		FArray3i NormalIndices;
+		FArray3i TexCoordIndices;
+
+		NormalIndices[0] = TriangulatedIndexes[StartIndexInOut++];
+		NormalIndices[1] = NormalIndices[0];
+		NormalIndices[2] = NormalIndices[0];
+
+		TexCoordIndices[0] = TriangulatedIndexes[StartIndexInOut];
+		StartIndexInOut += TexOffset;
+		VertexIndices[0] = TriangulatedIndexes[StartIndexInOut++] / 3;
+
+		TexCoordIndices[1] = TriangulatedIndexes[StartIndexInOut];
+		StartIndexInOut += TexOffset;
+		VertexIndices[1] = TriangulatedIndexes[StartIndexInOut++] / 3;
+
+		for (unsigned long FArray3i = 2; FArray3i < TriCount; FArray3i++)
+		{
+			TexCoordIndices[2] = TriangulatedIndexes[StartIndexInOut];
+			StartIndexInOut += TexOffset;
+			VertexIndices[2] = TriangulatedIndexes[StartIndexInOut++] / 3;
+
+			AddTriangle(Context, VertexIndices, NormalIndices, TexCoordIndices);
+
+			FArray3i++;
+			if (FArray3i == TriCount)
+			{
+				break;
+			}
+
+			Swap(VertexIndices[1], VertexIndices[2]);
+			Swap(TexCoordIndices[1], TexCoordIndices[2]);
+
+			TexCoordIndices[0] = TriangulatedIndexes[StartIndexInOut];
+			StartIndexInOut += TexOffset;
+			VertexIndices[0] = TriangulatedIndexes[StartIndexInOut++] / 3;
+
+			AddTriangle(Context, VertexIndices, NormalIndices, TexCoordIndices);
+
+			Swap(VertexIndices[0], VertexIndices[1]);
+			Swap(TexCoordIndices[0], TexCoordIndices[1]);
+		}
+	}
+
+	void FRepresentationToMeshHelper::GetUVScale(const A3DTopoFace* TopoFace, FVector2f& UVScale)
+	{
+		// #cad_import: Something to solve based on how this is handle in legacy
+		static FVector2f OneVector(1.f, 1.f);
+		UVScale = OneVector;
+	}
+
+
+	bool FRepresentationToMeshHelper::ParseTriangles(const TArray<A3DTopoFace*>& BRepFaces)
+	{
+		TechSoft::TUniqueObject<A3DTess3DData> TessellationData(TessellationBase);
+
+		TessellationNormals = TessellationData->m_pdNormals;
+		TessellationTexCoords = TessellationData->m_pdTextureCoords;
+		TriangulatedIndexes = TessellationData->m_puiTriangulatedIndexes;
+
+		FTriangleContext Context;
+
+		for (A3DUns32 Index = 0; Index < TessellationData->m_uiFaceTessSize; ++Index)
+		{
+			const A3DTessFaceData& FaceTessData = TessellationData->m_psFaceTessData[Index];
+
+			if (FaceTessData.m_uiSizesTriangulatedSize == 0)
+			{
+				continue;
+			}
+
+			uint32 FaceSetIndex = 0;
+			bool bMustProcess = true;
+
+			Context.GroupID = Index;
+			Context.MaterialID = 0;
+			if (FaceTessData.m_uiStyleIndexesSize != 0)
+			{
+				// Store the StyleIndex on the MaterialName. It will be processed after tessellation
+				Context.MaterialID = FaceTessData.m_puiStyleIndexes[0];
+			}
+
+			GetUVScale(BRepFaces[Index], Context.UVScale);
+
+			uint32 UsedEntitiesFlags = FaceTessData.m_usUsedEntitiesFlags;
+
+			uint32 LastTrianguleIndex = FaceTessData.m_uiStartTriangulated;
+			LastVertexIndex = 0;
+
+			if (bMustProcess && UsedEntitiesFlags & kA3DTessFaceDataTriangle)
+			{
+				AddFaceTriangle(Context, FaceTessData.m_puiSizesTriangulated[FaceSetIndex++], LastTrianguleIndex);
+				bMustProcess = FaceTessData.m_uiSizesTriangulatedSize > FaceSetIndex;
+			}
+
+			if (bMustProcess && UsedEntitiesFlags & kA3DTessFaceDataTriangleOneNormal)
+			{
+				AddFaceTriangleWithUniqueNormal(Context, FaceTessData.m_puiSizesTriangulated[FaceSetIndex++], LastTrianguleIndex);
+				bMustProcess = FaceTessData.m_uiSizesTriangulatedSize > FaceSetIndex;
+			}
+
+			if (bMustProcess && UsedEntitiesFlags & kA3DTessFaceDataTriangleTextured)
+			{
+				AddFaceTriangleWithTexture(Context, FaceTessData.m_puiSizesTriangulated[FaceSetIndex++], FaceTessData.m_uiTextureCoordIndexesSize, LastTrianguleIndex);
+				bMustProcess = FaceTessData.m_uiSizesTriangulatedSize > FaceSetIndex;
+			}
+
+			if (bMustProcess && UsedEntitiesFlags & kA3DTessFaceDataTriangleOneNormalTextured)
+			{
+				AddFaceTriangleWithUniqueNormalAndTexture(Context, FaceTessData.m_puiSizesTriangulated[FaceSetIndex++], FaceTessData.m_uiTextureCoordIndexesSize, LastTrianguleIndex);
+				bMustProcess = FaceTessData.m_uiSizesTriangulatedSize > FaceSetIndex;
+			}
+
+			if (bMustProcess && UsedEntitiesFlags & kA3DTessFaceDataTriangleFan)
+			{
+				uint32 FanCount = FaceTessData.m_puiSizesTriangulated[FaceSetIndex++];
+				for (uint32 FanIndex = 0; FanIndex < FanCount; ++FanIndex)
+				{
+					uint32 VertexCount = FaceTessData.m_puiSizesTriangulated[FaceSetIndex++];
+					AddFaceTriangleFan(Context, VertexCount, LastTrianguleIndex);
+				}
+				bMustProcess = FaceTessData.m_uiSizesTriangulatedSize > FaceSetIndex;
+			}
+
+			if (bMustProcess && UsedEntitiesFlags & kA3DTessFaceDataTriangleFanOneNormal)
+			{
+				uint32 FanCount = FaceTessData.m_puiSizesTriangulated[FaceSetIndex++] & kA3DTessFaceDataNormalMask;
+				for (uint32 FanIndex = 0; FanIndex < FanCount; ++FanIndex)
+				{
+					ensure((FaceTessData.m_puiSizesTriangulated[FaceSetIndex] & kA3DTessFaceDataNormalSingle) != 0);
+
+					uint32 VertexCount = FaceTessData.m_puiSizesTriangulated[FaceSetIndex++] & kA3DTessFaceDataNormalMask;
+					AddFaceTriangleFanWithUniqueNormal(Context, VertexCount, LastTrianguleIndex);
+				}
+				bMustProcess = FaceTessData.m_uiSizesTriangulatedSize > FaceSetIndex;
+			}
+
+			if (bMustProcess && UsedEntitiesFlags & kA3DTessFaceDataTriangleFanTextured)
+			{
+				uint32 FanCount = FaceTessData.m_puiSizesTriangulated[FaceSetIndex++];
+				for (uint32 FanIndex = 0; FanIndex < FanCount; ++FanIndex)
+				{
+					uint32 VertexCount = FaceTessData.m_puiSizesTriangulated[FaceSetIndex++];
+					AddFaceTriangleFanWithTexture(Context, VertexCount, FaceTessData.m_uiTextureCoordIndexesSize, LastTrianguleIndex);
+				}
+				bMustProcess = FaceTessData.m_uiSizesTriangulatedSize > FaceSetIndex;
+			}
+
+			if (bMustProcess && UsedEntitiesFlags & kA3DTessFaceDataTriangleFanOneNormalTextured)
+			{
+				uint32 FanCount = FaceTessData.m_puiSizesTriangulated[FaceSetIndex++] & kA3DTessFaceDataNormalMask;
+				for (uint32 FanIndex = 0; FanIndex < FanCount; ++FanIndex)
+				{
+					ensure((FaceTessData.m_puiSizesTriangulated[FaceSetIndex] & kA3DTessFaceDataNormalSingle) != 0);
+
+					uint32 VertexCount = FaceTessData.m_puiSizesTriangulated[FaceSetIndex++] & kA3DTessFaceDataNormalMask;
+					AddFaceTriangleFanWithUniqueNormalAndTexture(Context, VertexCount, FaceTessData.m_uiTextureCoordIndexesSize, LastTrianguleIndex);
+				}
+				bMustProcess = FaceTessData.m_uiSizesTriangulatedSize > FaceSetIndex;
+			}
+
+			if (bMustProcess && UsedEntitiesFlags & kA3DTessFaceDataTriangleStripe)
+			{
+				A3DUns32 StripeSize = FaceTessData.m_puiSizesTriangulated[FaceSetIndex++];
+				for (A3DUns32 StripeIndex = 0; StripeIndex < StripeSize; ++StripeIndex)
+				{
+					A3DUns32 PointCount = FaceTessData.m_puiSizesTriangulated[FaceSetIndex++];
+					AddFaceTriangleStripe(Context, PointCount, LastTrianguleIndex);
+				}
+				bMustProcess = FaceTessData.m_uiSizesTriangulatedSize > FaceSetIndex;
+			}
+
+			if (bMustProcess && UsedEntitiesFlags & kA3DTessFaceDataTriangleStripeOneNormal)
+			{
+				A3DUns32 StripeSize = FaceTessData.m_puiSizesTriangulated[FaceSetIndex++] & kA3DTessFaceDataNormalMask;
+				for (A3DUns32 StripeIndex = 0; StripeIndex < StripeSize; ++StripeIndex)
+				{
+					bool bIsOneNormal = (FaceTessData.m_puiSizesTriangulated[FaceSetIndex] & kA3DTessFaceDataNormalSingle) != 0;
+					A3DUns32 PointCount = FaceTessData.m_puiSizesTriangulated[FaceSetIndex++] & kA3DTessFaceDataNormalMask;
+
+					// Is there only one normal for the entire stripe?
+					if (bIsOneNormal == false)
+					{
+						AddFaceTriangleStripe(Context, PointCount, LastTrianguleIndex);
+						continue;
+					}
+
+					AddFaceTriangleStripeWithUniqueNormal(Context, PointCount, LastTrianguleIndex);
+				}
+				bMustProcess = FaceTessData.m_uiSizesTriangulatedSize > FaceSetIndex;
+			}
+
+			if (bMustProcess && UsedEntitiesFlags & kA3DTessFaceDataTriangleStripeTextured)
+			{
+				A3DUns32 StripeSize = FaceTessData.m_puiSizesTriangulated[FaceSetIndex++];
+				for (A3DUns32 StripeIndex = 0; StripeIndex < StripeSize; ++StripeIndex)
+				{
+					A3DUns32 PointCount = FaceTessData.m_puiSizesTriangulated[FaceSetIndex++];
+					AddFaceTriangleStripeWithTexture(Context, PointCount, FaceTessData.m_uiTextureCoordIndexesSize, LastTrianguleIndex);
+				}
+				bMustProcess = FaceTessData.m_uiSizesTriangulatedSize > FaceSetIndex;
+			}
+
+			if (bMustProcess && UsedEntitiesFlags & kA3DTessFaceDataTriangleStripeOneNormalTextured)
+			{
+				A3DUns32 StripeSize = FaceTessData.m_puiSizesTriangulated[FaceSetIndex++] & kA3DTessFaceDataNormalMask;
+				for (A3DUns32 StripeIndex = 0; StripeIndex < StripeSize; ++StripeIndex)
+				{
+					A3DUns32 PointCount = FaceTessData.m_puiSizesTriangulated[FaceSetIndex++] & kA3DTessFaceDataNormalMask;
+					AddFaceTriangleStripeWithUniqueNormalAndTexture(Context, PointCount, FaceTessData.m_uiTextureCoordIndexesSize, LastTrianguleIndex);
+				}
+				bMustProcess = FaceTessData.m_uiSizesTriangulatedSize > FaceSetIndex;
+			}
+
+			ensure(!bMustProcess);
+		}
+
+		return true;
+	}
+
+	bool FRepresentationToMeshHelper::Convert()
+	{
+		using namespace UE::CADKernel;
+
+		TechSoft::TUniqueObject<A3DTessBaseData> TessellationBaseData(TessellationBase);
+
+		if (!TessellationBaseData.IsValid() || TessellationBaseData->m_uiCoordSize == 0)
+		{
+			return false;
+		}
+
+		TArray<FVector> Vertices;
+		const int32 VertexCount = TessellationBaseData->m_uiCoordSize / 3;
+		
+		Vertices.Reserve(VertexCount);
+
+		//const double ScaleFactor = BodyUnit * FImportParameters::GUnitScale;
+		const double ScaleFactor = 1.;
+		double* Coordinates = TessellationBaseData->m_pdCoords;
+
+		for (int32 Index = 0; Index < VertexCount; ++Index, Coordinates += 3)
+		{
+			Vertices.Emplace(Coordinates[0], Coordinates[1], Coordinates[2]);
+		}
+
+		if (!MeshInterface.AddNewVertices(MoveTemp(Vertices)))
+		{
+			return false;
+		}
+
+		int32 TriangleCount;
+		TArray<A3DTopoFace*> BRepFaces;
+		if (!CollectBRepFaces(BRepFaces, TriangleCount))
+		{
+			return false;
+		}
+
+		MeshInterface.ReserveNewTriangles(TriangleCount);
+
+		return false;
+	}
+}
+#endif
\ No newline at end of file
diff --git a/Engine/Source/Runtime/Datasmith/CADKernel/Engine/Private/TechSoft/TechSoftToCADKernel.cpp b/Engine/Source/Runtime/Datasmith/CADKernel/Engine/Private/TechSoft/TechSoftToCADKernel.cpp
new file mode 100644
--- /dev/null
+++ b/Engine/Source/Runtime/Datasmith/CADKernel/Engine/Private/TechSoft/TechSoftToCADKernel.cpp
@@ -0,0 +1,29 @@
+// Copyright Epic Games, Inc. All Rights Reserved.
+
+#include "CADKernelEngine.h"
+
+#ifdef WITH_HOOPS
+#include "TechSoftUniqueObjectImpl.h"
+
+namespace UE::CADKernel
+{
+
+	TSharedPtr<UE::CADKernel::FModel> FTechSoftUtilities::TechSoftToCADKernel(A3DRiRepresentationItem* Representation)
+	{
+		if (!FTechSoftLibrary::IsInitialized())
+		{
+			return {};
+		}
+
+		return {};
+	}
+}
+#else
+namespace UE::CADKernel
+{
+	TSharedPtr<UE::CADKernel::FModel> FTechSoftUtilities::TechSoftToCADKernel(A3DRiRepresentationItem* Representation)
+	{
+		return {};
+	}
+}
+#endif
diff --git a/Engine/Source/Runtime/Datasmith/CADKernel/Engine/Private/TechSoft/TechSoftUniqueObjectImpl.cpp b/Engine/Source/Runtime/Datasmith/CADKernel/Engine/Private/TechSoft/TechSoftUniqueObjectImpl.cpp
new file mode 100644
--- /dev/null
+++ b/Engine/Source/Runtime/Datasmith/CADKernel/Engine/Private/TechSoft/TechSoftUniqueObjectImpl.cpp
@@ -0,0 +1,1053 @@
+// Copyright Epic Games, Inc. All Rights Reserved.
+
+#include "TechSoftUniqueObjectImpl.h"
+
+#ifdef WITH_HOOPS
+namespace UE::CADKernel::TechSoft
+{
+
+	// TUniqueObject InitializeData -----------------------------------
+	template<>
+	void TUniqueObject<A3DAsmModelFileData>::InitializeData()
+	{
+		A3D_INITIALIZE_DATA(A3DAsmModelFileData, Data);
+	}
+	template<>
+	void TUniqueObject<A3DAsmPartDefinitionData>::InitializeData()
+	{
+		A3D_INITIALIZE_DATA(A3DAsmPartDefinitionData, Data);
+	}
+	template<>
+	void TUniqueObject<A3DAsmProductOccurrenceData>::InitializeData()
+	{
+		A3D_INITIALIZE_DATA(A3DAsmProductOccurrenceData, Data);
+	}
+	template<>
+	void TUniqueObject<A3DAsmProductOccurrenceDataCV5>::InitializeData()
+	{
+		A3D_INITIALIZE_DATA(A3DAsmProductOccurrenceDataCV5, Data);
+	}
+	template<>
+	void TUniqueObject<A3DAsmProductOccurrenceDataSLW>::InitializeData()
+	{
+		A3D_INITIALIZE_DATA(A3DAsmProductOccurrenceDataSLW, Data);
+	}
+	template<>
+	void TUniqueObject<A3DAsmProductOccurrenceDataUg>::InitializeData()
+	{
+		A3D_INITIALIZE_DATA(A3DAsmProductOccurrenceDataUg, Data);
+	}
+	template<>
+	void TUniqueObject<A3DBoundingBoxData>::InitializeData()
+	{
+		A3D_INITIALIZE_DATA(A3DBoundingBoxData, Data);
+	}
+	template<>
+	void TUniqueObject<A3DCopyAndAdaptBrepModelData>::InitializeData()
+	{
+		A3D_INITIALIZE_DATA(A3DCopyAndAdaptBrepModelData, Data);
+	}
+	template<>
+	void TUniqueObject<A3DCrvCircleData>::InitializeData()
+	{
+		A3D_INITIALIZE_DATA(A3DCrvCircleData, Data);
+	}
+	template<>
+	void TUniqueObject<A3DCrvCompositeData>::InitializeData()
+	{
+		A3D_INITIALIZE_DATA(A3DCrvCompositeData, Data);
+	}
+	template<>
+	void TUniqueObject<A3DCrvEllipseData>::InitializeData()
+	{
+		A3D_INITIALIZE_DATA(A3DCrvEllipseData, Data);
+	}
+	template<>
+	void TUniqueObject<A3DCrvHelixData>::InitializeData()
+	{
+		A3D_INITIALIZE_DATA(A3DCrvHelixData, Data);
+	}
+	template<>
+	void TUniqueObject<A3DCrvHyperbolaData>::InitializeData()
+	{
+		A3D_INITIALIZE_DATA(A3DCrvHyperbolaData, Data);
+	}
+	template<>
+	void TUniqueObject<A3DCrvLineData>::InitializeData()
+	{
+		A3D_INITIALIZE_DATA(A3DCrvLineData, Data);
+	}
+	template<>
+	void TUniqueObject<A3DCrvNurbsData>::InitializeData()
+	{
+		A3D_INITIALIZE_DATA(A3DCrvNurbsData, Data);
+	}
+	template<>
+	void TUniqueObject<A3DCrvParabolaData>::InitializeData()
+	{
+		A3D_INITIALIZE_DATA(A3DCrvParabolaData, Data);
+	}
+	template<>
+	void TUniqueObject<A3DCrvPolyLineData>::InitializeData()
+	{
+		A3D_INITIALIZE_DATA(A3DCrvPolyLineData, Data);
+	}
+	template<>
+	void TUniqueObject<A3DDomainData>::InitializeData()
+	{
+		A3D_INITIALIZE_DATA(A3DDomainData, Data);
+	}
+	template<>
+	void TUniqueObject<A3DCrvTransformData>::InitializeData()
+	{
+		A3D_INITIALIZE_DATA(A3DCrvTransformData, Data);
+	}
+	template<>
+	void TUniqueObject<A3DGlobalData>::InitializeData()
+	{
+		A3D_INITIALIZE_DATA(A3DGlobalData, Data);
+	}
+	template<>
+	void TUniqueObject<A3DGraphicsData>::InitializeData()
+	{
+		A3D_INITIALIZE_DATA(A3DGraphicsData, Data);
+	}
+	template<>
+	void TUniqueObject<A3DIntervalData>::InitializeData()
+	{
+		A3D_INITIALIZE_DATA(A3DIntervalData, Data);
+	}
+	template<>
+	void TUniqueObject<A3DMiscAttributeData>::InitializeData()
+	{
+		A3D_INITIALIZE_DATA(A3DMiscAttributeData, Data);
+	}
+	template<>
+	void TUniqueObject<A3DMiscCartesianTransformationData>::InitializeData()
+	{
+		A3D_INITIALIZE_DATA(A3DMiscCartesianTransformationData, Data);
+	}
+	template<>
+	void TUniqueObject<A3DMiscEntityReferenceData>::InitializeData()
+	{
+		A3D_INITIALIZE_DATA(A3DMiscEntityReferenceData, Data);
+	}
+	template<>
+	void TUniqueObject<A3DMiscGeneralTransformationData>::InitializeData()
+	{
+		A3D_INITIALIZE_DATA(A3DMiscGeneralTransformationData, Data);
+	}
+	template<>
+	void TUniqueObject<A3DMiscMaterialPropertiesData>::InitializeData()
+	{
+		A3D_INITIALIZE_DATA(A3DMiscMaterialPropertiesData, Data);
+	}
+	template<>
+	void TUniqueObject<A3DMiscReferenceOnCsysItemData>::InitializeData()
+	{
+		A3D_INITIALIZE_DATA(A3DMiscReferenceOnCsysItemData, Data);
+	}
+	template<>
+	void TUniqueObject<A3DMiscReferenceOnTessData>::InitializeData()
+	{
+		A3D_INITIALIZE_DATA(A3DMiscReferenceOnTessData, Data);
+	}
+	template<>
+	void TUniqueObject<A3DMiscReferenceOnTopologyData>::InitializeData()
+	{
+		A3D_INITIALIZE_DATA(A3DMiscReferenceOnTopologyData, Data);
+	}
+	template<>
+	void TUniqueObject<A3DMiscSingleAttributeData>::InitializeData()
+	{
+		A3D_INITIALIZE_DATA(A3DMiscSingleAttributeData, Data);
+	}
+	template<>
+	void TUniqueObject<A3DRWParamsExportPrcData>::InitializeData()
+	{
+		A3D_INITIALIZE_DATA(A3DRWParamsExportPrcData, Data);
+	}
+	template<>
+	void TUniqueObject<A3DRiBrepModelData>::InitializeData()
+	{
+		A3D_INITIALIZE_DATA(A3DRiBrepModelData, Data);
+	}
+	template<>
+	void TUniqueObject<A3DRiCoordinateSystemData>::InitializeData()
+	{
+		A3D_INITIALIZE_DATA(A3DRiCoordinateSystemData, Data);
+	}
+	template<>
+	void TUniqueObject<A3DRiDirectionData>::InitializeData()
+	{
+		A3D_INITIALIZE_DATA(A3DRiDirectionData, Data);
+	}
+	template<>
+	void TUniqueObject<A3DRiPolyBrepModelData>::InitializeData()
+	{
+		A3D_INITIALIZE_DATA(A3DRiPolyBrepModelData, Data);
+	}
+	template<>
+	void TUniqueObject<A3DRiRepresentationItemData>::InitializeData()
+	{
+		A3D_INITIALIZE_DATA(A3DRiRepresentationItemData, Data);
+	}
+	template<>
+	void TUniqueObject<A3DRiSetData>::InitializeData()
+	{
+		A3D_INITIALIZE_DATA(A3DRiSetData, Data);
+	}
+	template<>
+	void TUniqueObject<A3DRootBaseData>::InitializeData()
+	{
+		A3D_INITIALIZE_DATA(A3DRootBaseData, Data);
+	}
+	template<>
+	void TUniqueObject<A3DRootBaseWithGraphicsData>::InitializeData()
+	{
+		A3D_INITIALIZE_DATA(A3DRootBaseWithGraphicsData, Data);
+	}
+	template<>
+	void TUniqueObject<A3DRWParamsTessellationData>::InitializeData()
+	{
+		A3D_INITIALIZE_DATA(A3DRWParamsTessellationData, Data);
+	}
+	template<>
+	void TUniqueObject<A3DSewOptionsData>::InitializeData()
+	{
+		A3D_INITIALIZE_DATA(A3DSewOptionsData, Data);
+	}
+	template<>
+	void TUniqueObject<A3DSurfBlend01Data>::InitializeData()
+	{
+		A3D_INITIALIZE_DATA(A3DSurfBlend01Data, Data);
+	}
+	template<>
+	void TUniqueObject<A3DSurfBlend02Data>::InitializeData()
+	{
+		A3D_INITIALIZE_DATA(A3DSurfBlend02Data, Data);
+	}
+	template<>
+	void TUniqueObject<A3DSurfBlend03Data>::InitializeData()
+	{
+		A3D_INITIALIZE_DATA(A3DSurfBlend03Data, Data);
+	}
+	template<>
+	void TUniqueObject<A3DSurfConeData>::InitializeData()
+	{
+		A3D_INITIALIZE_DATA(A3DSurfConeData, Data);
+	}
+	template<>
+	void TUniqueObject<A3DSurfCylinderData>::InitializeData()
+	{
+		A3D_INITIALIZE_DATA(A3DSurfCylinderData, Data);
+	}
+	template<>
+	void TUniqueObject<A3DSurfCylindricalData>::InitializeData()
+	{
+		A3D_INITIALIZE_DATA(A3DSurfCylindricalData, Data);
+	}
+	template<>
+	void TUniqueObject<A3DSurfExtrusionData>::InitializeData()
+	{
+		A3D_INITIALIZE_DATA(A3DSurfExtrusionData, Data);
+	}
+	template<>
+	void TUniqueObject<A3DSurfFromCurvesData>::InitializeData()
+	{
+		A3D_INITIALIZE_DATA(A3DSurfFromCurvesData, Data);
+	}
+	template<>
+	void TUniqueObject<A3DSurfNurbsData>::InitializeData()
+	{
+		A3D_INITIALIZE_DATA(A3DSurfNurbsData, Data);
+	}
+	template<>
+	void TUniqueObject<A3DSurfPipeData>::InitializeData()
+	{
+		A3D_INITIALIZE_DATA(A3DSurfPipeData, Data);
+	}
+	template<>
+	void TUniqueObject<A3DSurfPlaneData>::InitializeData()
+	{
+		A3D_INITIALIZE_DATA(A3DSurfPlaneData, Data);
+	}
+	template<>
+	void TUniqueObject<A3DSurfRevolutionData>::InitializeData()
+	{
+		A3D_INITIALIZE_DATA(A3DSurfRevolutionData, Data);
+	}
+	template<>
+	void TUniqueObject<A3DSurfRuledData>::InitializeData()
+	{
+		A3D_INITIALIZE_DATA(A3DSurfRuledData, Data);
+	}
+	template<>
+	void TUniqueObject<A3DSurfSphereData>::InitializeData()
+	{
+		A3D_INITIALIZE_DATA(A3DSurfSphereData, Data);
+	}
+	template<>
+	void TUniqueObject<A3DSurfTorusData>::InitializeData()
+	{
+		A3D_INITIALIZE_DATA(A3DSurfTorusData, Data);
+	}
+	template<>
+	void TUniqueObject<A3DTess3DData>::InitializeData()
+	{
+		A3D_INITIALIZE_DATA(A3DTess3DData, Data);
+	}
+	template<>
+	void TUniqueObject<A3DTessBaseData>::InitializeData()
+	{
+		A3D_INITIALIZE_DATA(A3DTessBaseData, Data);
+	}
+	template<>
+	void TUniqueObject<A3DTopoBodyData>::InitializeData()
+	{
+		A3D_INITIALIZE_DATA(A3DTopoBodyData, Data);
+	}
+	template<>
+	void TUniqueObject<A3DTopoBrepDataData>::InitializeData()
+	{
+		A3D_INITIALIZE_DATA(A3DTopoBrepDataData, Data);
+	}
+	template<>
+	void TUniqueObject<A3DTopoCoEdgeData>::InitializeData()
+	{
+		A3D_INITIALIZE_DATA(A3DTopoCoEdgeData, Data);
+	}
+	template<>
+	void TUniqueObject<A3DTopoConnexData>::InitializeData()
+	{
+		A3D_INITIALIZE_DATA(A3DTopoConnexData, Data);
+	}
+	template<>
+	void TUniqueObject<A3DTopoContextData>::InitializeData()
+	{
+		A3D_INITIALIZE_DATA(A3DTopoContextData, Data);
+	}
+	template<>
+	void TUniqueObject<A3DTopoEdgeData>::InitializeData()
+	{
+		A3D_INITIALIZE_DATA(A3DTopoEdgeData, Data);
+	}
+	template<>
+	void TUniqueObject<A3DTopoFaceData>::InitializeData()
+	{
+		A3D_INITIALIZE_DATA(A3DTopoFaceData, Data);
+	}
+	template<>
+	void TUniqueObject<A3DTopoLoopData>::InitializeData()
+	{
+		A3D_INITIALIZE_DATA(A3DTopoLoopData, Data);
+	}
+	template<>
+	void TUniqueObject<A3DTopoShellData>::InitializeData()
+	{
+		A3D_INITIALIZE_DATA(A3DTopoShellData, Data);
+	}
+	template<>
+	void TUniqueObject<A3DTopoUniqueVertexData>::InitializeData()
+	{
+		A3D_INITIALIZE_DATA(A3DTopoUniqueVertexData, Data);
+	}
+	template<>
+	void TUniqueObject<A3DTopoWireEdgeData>::InitializeData()
+	{
+		A3D_INITIALIZE_DATA(A3DTopoWireEdgeData, Data);
+	}
+	template<>
+	void TUniqueObject<A3DUTF8Char*>::InitializeData()
+	{
+		Data = nullptr;
+	}
+	template<>
+	void TUniqueObject<A3DVector2dData>::InitializeData()
+	{
+		A3D_INITIALIZE_DATA(A3DVector2dData, Data);
+	}
+	template<>
+	void TUniqueObject<A3DVector3dData>::InitializeData()
+	{
+		A3D_INITIALIZE_DATA(A3DVector3dData, Data);
+	}
+
+
+	// TUniqueObjectFromIndex InitializeData -----------------------------------
+	template<>
+	void TUniqueObjectFromIndex<A3DGraphMaterialData>::InitializeData()
+	{
+		A3D_INITIALIZE_DATA(A3DGraphMaterialData, Data);
+	}
+	template<>
+	void TUniqueObjectFromIndex<A3DGraphPictureData>::InitializeData()
+	{
+		A3D_INITIALIZE_DATA(A3DGraphPictureData, Data);
+	}
+	template<>
+	void TUniqueObjectFromIndex<A3DGraphRgbColorData>::InitializeData()
+	{
+		A3D_INITIALIZE_DATA(A3DGraphRgbColorData, Data);
+	}
+	template<>
+	void TUniqueObjectFromIndex<A3DGraphStyleData>::InitializeData()
+	{
+		A3D_INITIALIZE_DATA(A3DGraphStyleData, Data);
+	}
+	template<>
+	void TUniqueObjectFromIndex<A3DGraphTextureApplicationData>::InitializeData()
+	{
+		A3D_INITIALIZE_DATA(A3DGraphTextureApplicationData, Data);
+	}
+	template<>
+	void TUniqueObjectFromIndex<A3DGraphTextureDefinitionData>::InitializeData()
+	{
+		A3D_INITIALIZE_DATA(A3DGraphTextureDefinitionData, Data);
+	}
+
+
+
+	// TUniqueObject GetData -----------------------------------
+
+	template<>
+	A3DStatus TUniqueObject<A3DAsmModelFileData>::GetData(const A3DEntity* InEntityPtr)
+	{
+		return A3DAsmModelFileGet(InEntityPtr, &Data);
+	}
+
+	template<>
+	A3DStatus TUniqueObject<A3DAsmPartDefinitionData>::GetData(const A3DEntity* InEntityPtr)
+	{
+		return A3DAsmPartDefinitionGet(InEntityPtr, &Data);
+	}
+
+	template<>
+	A3DStatus TUniqueObject<A3DAsmProductOccurrenceData>::GetData(const A3DEntity* InEntityPtr)
+	{
+		return A3DAsmProductOccurrenceGet(InEntityPtr, &Data);
+	}
+
+	template<>
+	A3DStatus TUniqueObject<A3DAsmProductOccurrenceDataCV5>::GetData(const A3DEntity* InEntityPtr)
+	{
+		return A3DAsmProductOccurrenceGetCV5(InEntityPtr, &Data);
+	}
+
+	template<>
+	A3DStatus TUniqueObject<A3DAsmProductOccurrenceDataSLW>::GetData(const A3DEntity* InEntityPtr)
+	{
+		return A3DAsmProductOccurrenceGetSLW(InEntityPtr, &Data);
+	}
+
+	template<>
+	A3DStatus TUniqueObject<A3DAsmProductOccurrenceDataUg>::GetData(const A3DEntity* InEntityPtr)
+	{
+		return A3DAsmProductOccurrenceGetUg(InEntityPtr, &Data);
+	}
+
+	template<>
+	A3DStatus TUniqueObject<A3DBoundingBoxData>::GetData(const A3DEntity* InEntityPtr)
+	{
+		return A3DMiscGetBoundingBox(InEntityPtr, &Data);
+	}
+
+	template<>
+	A3DStatus TUniqueObject<A3DCopyAndAdaptBrepModelData>::GetData(const A3DEntity* InEntityPtr)
+	{
+		return A3DStatus::A3D_ERROR;
+	}
+
+	template<>
+	A3DStatus TUniqueObject<A3DCrvCircleData>::GetData(const A3DEntity* InEntityPtr)
+	{
+		return A3DCrvCircleGet(InEntityPtr, &Data);
+	}
+
+	template<>
+	A3DStatus TUniqueObject<A3DCrvCompositeData>::GetData(const A3DEntity* InEntityPtr)
+	{
+		return A3DCrvCompositeGet(InEntityPtr, &Data);
+	}
+
+	template<>
+	A3DStatus TUniqueObject<A3DCrvEllipseData>::GetData(const A3DEntity* InEntityPtr)
+	{
+		return A3DCrvEllipseGet(InEntityPtr, &Data);
+	}
+
+	template<>
+	A3DStatus TUniqueObject<A3DCrvHelixData>::GetData(const A3DEntity* InEntityPtr)
+	{
+		return A3DCrvHelixGet(InEntityPtr, &Data);
+	}
+
+	template<>
+	A3DStatus TUniqueObject<A3DCrvHyperbolaData>::GetData(const A3DEntity* InEntityPtr)
+	{
+		return A3DCrvHyperbolaGet(InEntityPtr, &Data);
+	}
+
+	template<>
+	A3DStatus TUniqueObject<A3DCrvLineData>::GetData(const A3DEntity* InEntityPtr)
+	{
+		return A3DCrvLineGet(InEntityPtr, &Data);
+	}
+
+	template<>
+	A3DStatus TUniqueObject<A3DCrvNurbsData>::GetData(const A3DEntity* InEntityPtr)
+	{
+		return A3DCrvNurbsGet(InEntityPtr, &Data);
+	}
+
+	template<>
+	A3DStatus TUniqueObject<A3DCrvParabolaData>::GetData(const A3DEntity* InEntityPtr)
+	{
+		return A3DCrvParabolaGet(InEntityPtr, &Data);
+	}
+
+	template<>
+	A3DStatus TUniqueObject<A3DCrvPolyLineData>::GetData(const A3DEntity* InEntityPtr)
+	{
+		return A3DCrvPolyLineGet(InEntityPtr, &Data);
+	}
+
+	template<>
+	A3DStatus TUniqueObject<A3DDomainData>::GetData(const A3DEntity* InEntityPtr)
+	{
+		return A3DStatus::A3D_ERROR;
+	}
+
+	template<>
+	A3DStatus TUniqueObject<A3DCrvTransformData>::GetData(const A3DEntity* InEntityPtr)
+	{
+		return A3DCrvTransformGet(InEntityPtr, &Data);
+	}
+
+	template<>
+	A3DStatus TUniqueObject<A3DGlobalData>::GetData(const A3DEntity* InEntityPtr)
+	{
+		return A3DGlobalGet(InEntityPtr, &Data);
+	}
+
+	template<>
+	A3DStatus TUniqueObject<A3DGraphicsData>::GetData(const A3DEntity* InEntityPtr)
+	{
+		return A3DGraphicsGet(InEntityPtr, &Data);
+	}
+
+	template<>
+	A3DStatus TUniqueObject<A3DIntervalData>::GetData(const A3DEntity* InEntityPtr)
+	{
+		return A3DCrvGetInterval(InEntityPtr, &Data);
+	}
+
+	template<>
+	A3DStatus TUniqueObject<A3DMiscAttributeData>::GetData(const A3DEntity* InEntityPtr)
+	{
+		return A3DMiscAttributeGet(InEntityPtr, &Data);
+	}
+
+	template<>
+	A3DStatus TUniqueObject<A3DMiscCartesianTransformationData>::GetData(const A3DEntity* InEntityPtr)
+	{
+		return A3DMiscCartesianTransformationGet(InEntityPtr, &Data);
+	}
+
+	template<>
+	A3DStatus TUniqueObject<A3DMiscEntityReferenceData>::GetData(const A3DEntity* InEntityPtr)
+	{
+		return A3DMiscEntityReferenceGet(InEntityPtr, &Data);
+	}
+
+	template<>
+	A3DStatus TUniqueObject<A3DMiscGeneralTransformationData>::GetData(const A3DEntity* InEntityPtr)
+	{
+		return A3DMiscGeneralTransformationGet(InEntityPtr, &Data);
+	}
+
+	template<>
+	A3DStatus TUniqueObject<A3DMiscMaterialPropertiesData>::GetData(const A3DEntity* InEntityPtr)
+	{
+		return A3DMiscGetMaterialProperties(InEntityPtr, &Data);
+	}
+
+	template<>
+	A3DStatus TUniqueObject<A3DMiscReferenceOnCsysItemData>::GetData(const A3DEntity* InEntityPtr)
+	{
+		return A3DMiscReferenceOnCsysItemGet(InEntityPtr, &Data);
+	}
+
+	template<>
+	A3DStatus TUniqueObject<A3DMiscReferenceOnTessData>::GetData(const A3DEntity* InEntityPtr)
+	{
+		return A3DMiscReferenceOnTessGet(InEntityPtr, &Data);
+	}
+
+	template<>
+	A3DStatus TUniqueObject<A3DMiscReferenceOnTopologyData>::GetData(const A3DEntity* InEntityPtr)
+	{
+		return A3DMiscReferenceOnTopologyGet(InEntityPtr, &Data);
+	}
+
+	template<>
+	A3DStatus TUniqueObject<A3DMiscSingleAttributeData>::GetData(const A3DEntity* InEntityPtr)
+	{
+		return A3DStatus::A3D_ERROR;
+	}
+
+	template<>
+	A3DStatus TUniqueObject<A3DRWParamsExportPrcData>::GetData(const A3DEntity* InEntityPtr)
+	{
+		return A3DStatus::A3D_ERROR;
+	}
+
+	template<>
+	A3DStatus TUniqueObject<A3DRiBrepModelData>::GetData(const A3DEntity* InEntityPtr)
+	{
+		return A3DRiBrepModelGet(InEntityPtr, &Data);
+	}
+
+	template<>
+	A3DStatus TUniqueObject<A3DRiCoordinateSystemData>::GetData(const A3DEntity* InEntityPtr)
+	{
+		return A3DRiCoordinateSystemGet(InEntityPtr, &Data);
+	}
+
+	template<>
+	A3DStatus TUniqueObject<A3DRiDirectionData>::GetData(const A3DEntity* InEntityPtr)
+	{
+		return A3DRiDirectionGet(InEntityPtr, &Data);
+	}
+
+	template<>
+	A3DStatus TUniqueObject<A3DRiPolyBrepModelData>::GetData(const A3DEntity* InEntityPtr)
+	{
+		return A3DRiPolyBrepModelGet(InEntityPtr, &Data);
+	}
+
+	template<>
+	A3DStatus TUniqueObject<A3DRiRepresentationItemData>::GetData(const A3DEntity* InEntityPtr)
+	{
+		return A3DRiRepresentationItemGet(InEntityPtr, &Data);
+	}
+
+	template<>
+	A3DStatus TUniqueObject<A3DRiSetData>::GetData(const A3DEntity* InEntityPtr)
+	{
+		return A3DRiSetGet(InEntityPtr, &Data);
+	}
+
+	template<>
+	A3DStatus TUniqueObject<A3DRootBaseData>::GetData(const A3DEntity* InEntityPtr)
+	{
+		return A3DRootBaseGet(InEntityPtr, &Data);
+	}
+
+	template<>
+	A3DStatus TUniqueObject<A3DRootBaseWithGraphicsData>::GetData(const A3DEntity* InEntityPtr)
+	{
+		return A3DRootBaseWithGraphicsGet(InEntityPtr, &Data);
+	}
+
+	template<>
+	A3DStatus TUniqueObject<A3DRWParamsTessellationData>::GetData(const A3DEntity* InEntityPtr)
+	{
+		return A3DStatus::A3D_ERROR;
+	}
+
+	template<>
+	A3DStatus TUniqueObject<A3DSewOptionsData>::GetData(const A3DEntity* InEntityPtr)
+	{
+		return A3DStatus::A3D_ERROR;
+	}
+
+	template<>
+	A3DStatus TUniqueObject<A3DSurfBlend01Data>::GetData(const A3DEntity* InEntityPtr)
+	{
+		return A3DSurfBlend01Get(InEntityPtr, &Data);
+	}
+
+	template<>
+	A3DStatus TUniqueObject<A3DSurfBlend02Data>::GetData(const A3DEntity* InEntityPtr)
+	{
+		return A3DSurfBlend02Get(InEntityPtr, &Data);
+	}
+
+	template<>
+	A3DStatus TUniqueObject<A3DSurfBlend03Data>::GetData(const A3DEntity* InEntityPtr)
+	{
+		return A3DSurfBlend03Get(InEntityPtr, &Data);
+	}
+
+	template<>
+	A3DStatus TUniqueObject<A3DSurfConeData>::GetData(const A3DEntity* InEntityPtr)
+	{
+		return A3DSurfConeGet(InEntityPtr, &Data);
+	}
+
+	template<>
+	A3DStatus TUniqueObject<A3DSurfCylinderData>::GetData(const A3DEntity* InEntityPtr)
+	{
+		return A3DSurfCylinderGet(InEntityPtr, &Data);
+	}
+
+	template<>
+	A3DStatus TUniqueObject<A3DSurfCylindricalData>::GetData(const A3DEntity* InEntityPtr)
+	{
+		return A3DSurfCylindricalGet(InEntityPtr, &Data);
+	}
+
+	template<>
+	A3DStatus TUniqueObject<A3DSurfExtrusionData>::GetData(const A3DEntity* InEntityPtr)
+	{
+		return A3DSurfExtrusionGet(InEntityPtr, &Data);
+	}
+
+	template<>
+	A3DStatus TUniqueObject<A3DSurfFromCurvesData>::GetData(const A3DEntity* InEntityPtr)
+	{
+		return A3DSurfFromCurvesGet(InEntityPtr, &Data);
+	}
+
+	template<>
+	A3DStatus TUniqueObject<A3DSurfNurbsData>::GetData(const A3DEntity* InEntityPtr)
+	{
+		return A3DSurfNurbsGet(InEntityPtr, &Data);
+	}
+
+	template<>
+	A3DStatus TUniqueObject<A3DSurfPipeData>::GetData(const A3DEntity* InEntityPtr)
+	{
+		return A3DSurfPipeGet(InEntityPtr, &Data);
+	}
+
+	template<>
+	A3DStatus TUniqueObject<A3DSurfPlaneData>::GetData(const A3DEntity* InEntityPtr)
+	{
+		return A3DSurfPlaneGet(InEntityPtr, &Data);
+	}
+
+	template<>
+	A3DStatus TUniqueObject<A3DSurfRevolutionData>::GetData(const A3DEntity* InEntityPtr)
+	{
+		return A3DSurfRevolutionGet(InEntityPtr, &Data);
+	}
+
+	template<>
+	A3DStatus TUniqueObject<A3DSurfRuledData>::GetData(const A3DEntity* InEntityPtr)
+	{
+		return A3DSurfRuledGet(InEntityPtr, &Data);
+	}
+
+	template<>
+	A3DStatus TUniqueObject<A3DSurfSphereData>::GetData(const A3DEntity* InEntityPtr)
+	{
+		return A3DSurfSphereGet(InEntityPtr, &Data);
+	}
+
+	template<>
+	A3DStatus TUniqueObject<A3DSurfTorusData>::GetData(const A3DEntity* InEntityPtr)
+	{
+		return A3DSurfTorusGet(InEntityPtr, &Data);
+	}
+
+	template<>
+	A3DStatus TUniqueObject<A3DTess3DData>::GetData(const A3DEntity* InEntityPtr)
+	{
+		return A3DTess3DGet(InEntityPtr, &Data);
+	}
+
+	template<>
+	A3DStatus TUniqueObject<A3DTessBaseData>::GetData(const A3DEntity* InEntityPtr)
+	{
+		return A3DTessBaseGet(InEntityPtr, &Data);
+	}
+
+	template<>
+	A3DStatus TUniqueObject<A3DTopoBodyData>::GetData(const A3DEntity* InEntityPtr)
+	{
+		return A3DTopoBodyGet(InEntityPtr, &Data);
+	}
+
+	template<>
+	A3DStatus TUniqueObject<A3DTopoBrepDataData>::GetData(const A3DEntity* InEntityPtr)
+	{
+		return A3DTopoBrepDataGet(InEntityPtr, &Data);
+	}
+
+	template<>
+	A3DStatus TUniqueObject<A3DTopoCoEdgeData>::GetData(const A3DEntity* InEntityPtr)
+	{
+		return A3DTopoCoEdgeGet(InEntityPtr, &Data);
+	}
+
+	template<>
+	A3DStatus TUniqueObject<A3DTopoConnexData>::GetData(const A3DEntity* InEntityPtr)
+	{
+		return A3DTopoConnexGet(InEntityPtr, &Data);
+	}
+
+	template<>
+	A3DStatus TUniqueObject<A3DTopoContextData>::GetData(const A3DEntity* InEntityPtr)
+	{
+		return A3DTopoContextGet(InEntityPtr, &Data);
+	}
+
+	template<>
+	A3DStatus TUniqueObject<A3DTopoEdgeData>::GetData(const A3DEntity* InEntityPtr)
+	{
+		return A3DTopoEdgeGet(InEntityPtr, &Data);
+	}
+
+	template<>
+	A3DStatus TUniqueObject<A3DTopoFaceData>::GetData(const A3DEntity* InEntityPtr)
+	{
+		return A3DTopoFaceGet(InEntityPtr, &Data);
+	}
+
+	template<>
+	A3DStatus TUniqueObject<A3DTopoLoopData>::GetData(const A3DEntity* InEntityPtr)
+	{
+		return A3DTopoLoopGet(InEntityPtr, &Data);
+	}
+
+	template<>
+	A3DStatus TUniqueObject<A3DTopoShellData>::GetData(const A3DEntity* InEntityPtr)
+	{
+		return A3DTopoShellGet(InEntityPtr, &Data);
+	}
+
+	template<>
+	A3DStatus TUniqueObject<A3DTopoUniqueVertexData>::GetData(const A3DEntity* InEntityPtr)
+	{
+		return A3DTopoUniqueVertexGet(InEntityPtr, &Data);
+	}
+
+	template<>
+	A3DStatus TUniqueObject<A3DTopoWireEdgeData>::GetData(const A3DEntity* InEntityPtr)
+	{
+		return A3DTopoWireEdgeGet(InEntityPtr, &Data);
+	}
+
+	template<>
+	A3DStatus TUniqueObject<A3DUTF8Char*>::GetData(const A3DEntity* InEntityPtr)
+	{
+		return A3DStatus::A3D_ERROR;
+	}
+
+	template<>
+	A3DStatus TUniqueObject<A3DVector2dData>::GetData(const A3DEntity* InEntityPtr)
+	{
+		return A3DStatus::A3D_ERROR;
+	}
+
+	template<>
+	A3DStatus TUniqueObject<A3DVector3dData>::GetData(const A3DEntity* InEntityPtr)
+	{
+		return A3DStatus::A3D_ERROR;
+	}
+
+
+	// TUniqueObjectFromIndex GetData -----------------------------------
+	template<>
+	A3DStatus TUniqueObjectFromIndex<A3DGraphMaterialData>::GetData(const uint32 InEntityIndex)
+	{
+		return A3DGlobalGetGraphMaterialData(InEntityIndex, &Data);
+	}
+	template<>
+	A3DStatus TUniqueObjectFromIndex<A3DGraphPictureData>::GetData(const uint32 InEntityIndex)
+	{
+		return A3DGlobalGetGraphPictureData(InEntityIndex, &Data);
+	}
+	template<>
+	A3DStatus TUniqueObjectFromIndex<A3DGraphRgbColorData>::GetData(const uint32 InEntityIndex)
+	{
+		return A3DGlobalGetGraphRgbColorData(InEntityIndex, &Data);
+	}
+	template<>
+	A3DStatus TUniqueObjectFromIndex<A3DGraphStyleData>::GetData(const uint32 InEntityIndex)
+	{
+		return A3DGlobalGetGraphStyleData(InEntityIndex, &Data);
+	}
+	template<>
+	A3DStatus TUniqueObjectFromIndex<A3DGraphTextureApplicationData>::GetData(const uint32 InEntityIndex)
+	{
+		return A3DGlobalGetGraphTextureApplicationData(InEntityIndex, &Data);
+	}
+	template<>
+	A3DStatus TUniqueObjectFromIndex<A3DGraphTextureDefinitionData>::GetData(const uint32 InEntityIndex)
+	{
+		return A3DGlobalGetGraphTextureDefinitionData(InEntityIndex, &Data);
+	}
+
+	// DefaultValue -----------------------------------
+	template<>
+	const A3DEntity* TUniqueObject<A3DAsmModelFileData>::GetDefaultIndexerValue() const { return nullptr; }
+	template<>
+	const A3DEntity* TUniqueObject<A3DAsmPartDefinitionData>::GetDefaultIndexerValue() const { return nullptr; }
+	template<>
+	const A3DEntity* TUniqueObject<A3DAsmProductOccurrenceData>::GetDefaultIndexerValue() const { return nullptr; }
+	template<>
+	const A3DEntity* TUniqueObject<A3DAsmProductOccurrenceDataCV5>::GetDefaultIndexerValue() const { return nullptr; }
+	template<>
+	const A3DEntity* TUniqueObject<A3DAsmProductOccurrenceDataSLW>::GetDefaultIndexerValue() const { return nullptr; }
+	template<>
+	const A3DEntity* TUniqueObject<A3DAsmProductOccurrenceDataUg>::GetDefaultIndexerValue() const { return nullptr; }
+	template<>
+	const A3DEntity* TUniqueObject<A3DBoundingBoxData>::GetDefaultIndexerValue() const { return nullptr; }
+	template<>
+	const A3DEntity* TUniqueObject<A3DCopyAndAdaptBrepModelData>::GetDefaultIndexerValue() const { return nullptr; }
+	template<>
+	const A3DEntity* TUniqueObject<A3DCrvCircleData>::GetDefaultIndexerValue() const { return nullptr; }
+	template<>
+	const A3DEntity* TUniqueObject<A3DCrvCompositeData>::GetDefaultIndexerValue() const { return nullptr; }
+	template<>
+	const A3DEntity* TUniqueObject<A3DCrvEllipseData>::GetDefaultIndexerValue() const { return nullptr; }
+	template<>
+	const A3DEntity* TUniqueObject<A3DCrvHelixData>::GetDefaultIndexerValue() const { return nullptr; }
+	template<>
+	const A3DEntity* TUniqueObject<A3DCrvHyperbolaData>::GetDefaultIndexerValue() const { return nullptr; }
+	template<>
+	const A3DEntity* TUniqueObject<A3DCrvLineData>::GetDefaultIndexerValue() const { return nullptr; }
+	template<>
+	const A3DEntity* TUniqueObject<A3DCrvNurbsData>::GetDefaultIndexerValue() const { return nullptr; }
+	template<>
+	const A3DEntity* TUniqueObject<A3DCrvParabolaData>::GetDefaultIndexerValue() const { return nullptr; }
+	template<>
+	const A3DEntity* TUniqueObject<A3DCrvPolyLineData>::GetDefaultIndexerValue() const { return nullptr; }
+	template<>
+	const A3DEntity* TUniqueObject<A3DDomainData>::GetDefaultIndexerValue() const { return nullptr; }
+	template<>
+	const A3DEntity* TUniqueObject<A3DCrvTransformData>::GetDefaultIndexerValue() const { return nullptr; }
+	template<>
+	const A3DEntity* TUniqueObject<A3DGlobalData>::GetDefaultIndexerValue() const { return nullptr; }
+	template<>
+	const A3DEntity* TUniqueObject<A3DGraphicsData>::GetDefaultIndexerValue() const { return nullptr; }
+	template<>
+	const A3DEntity* TUniqueObject<A3DIntervalData>::GetDefaultIndexerValue() const { return nullptr; }
+	template<>
+	const A3DEntity* TUniqueObject<A3DMiscAttributeData>::GetDefaultIndexerValue() const { return nullptr; }
+	template<>
+	const A3DEntity* TUniqueObject<A3DMiscCartesianTransformationData>::GetDefaultIndexerValue() const { return nullptr; }
+	template<>
+	const A3DEntity* TUniqueObject<A3DMiscEntityReferenceData>::GetDefaultIndexerValue() const { return nullptr; }
+	template<>
+	const A3DEntity* TUniqueObject<A3DMiscGeneralTransformationData>::GetDefaultIndexerValue() const { return nullptr; }
+	template<>
+	const A3DEntity* TUniqueObject<A3DMiscMaterialPropertiesData>::GetDefaultIndexerValue() const { return nullptr; }
+	template<>
+	const A3DEntity* TUniqueObject<A3DMiscReferenceOnCsysItemData>::GetDefaultIndexerValue() const { return nullptr; }
+	template<>
+	const A3DEntity* TUniqueObject<A3DMiscReferenceOnTessData>::GetDefaultIndexerValue() const { return nullptr; }
+	template<>
+	const A3DEntity* TUniqueObject<A3DMiscReferenceOnTopologyData>::GetDefaultIndexerValue() const { return nullptr; }
+	template<>
+	const A3DEntity* TUniqueObject<A3DMiscSingleAttributeData>::GetDefaultIndexerValue() const { return nullptr; }
+	template<>
+	const A3DEntity* TUniqueObject<A3DRWParamsExportPrcData>::GetDefaultIndexerValue() const { return nullptr; }
+	template<>
+	const A3DEntity* TUniqueObject<A3DRiBrepModelData>::GetDefaultIndexerValue() const { return nullptr; }
+	template<>
+	const A3DEntity* TUniqueObject<A3DRiCoordinateSystemData>::GetDefaultIndexerValue() const { return nullptr; }
+	template<>
+	const A3DEntity* TUniqueObject<A3DRiDirectionData>::GetDefaultIndexerValue() const { return nullptr; }
+	template<>
+	const A3DEntity* TUniqueObject<A3DRiPolyBrepModelData>::GetDefaultIndexerValue() const { return nullptr; }
+	template<>
+	const A3DEntity* TUniqueObject<A3DRiRepresentationItemData>::GetDefaultIndexerValue() const { return nullptr; }
+	template<>
+	const A3DEntity* TUniqueObject<A3DRiSetData>::GetDefaultIndexerValue() const { return nullptr; }
+	template<>
+	const A3DEntity* TUniqueObject<A3DRootBaseData>::GetDefaultIndexerValue() const { return nullptr; }
+	template<>
+	const A3DEntity* TUniqueObject<A3DRootBaseWithGraphicsData>::GetDefaultIndexerValue() const { return nullptr; }
+	template<>
+	const A3DEntity* TUniqueObject<A3DRWParamsTessellationData>::GetDefaultIndexerValue() const { return nullptr; }
+	template<>
+	const A3DEntity* TUniqueObject<A3DSewOptionsData>::GetDefaultIndexerValue() const { return nullptr; }
+	template<>
+	const A3DEntity* TUniqueObject<A3DSurfBlend01Data>::GetDefaultIndexerValue() const { return nullptr; }
+	template<>
+	const A3DEntity* TUniqueObject<A3DSurfBlend02Data>::GetDefaultIndexerValue() const { return nullptr; }
+	template<>
+	const A3DEntity* TUniqueObject<A3DSurfBlend03Data>::GetDefaultIndexerValue() const { return nullptr; }
+	template<>
+	const A3DEntity* TUniqueObject<A3DSurfConeData>::GetDefaultIndexerValue() const { return nullptr; }
+	template<>
+	const A3DEntity* TUniqueObject<A3DSurfCylinderData>::GetDefaultIndexerValue() const { return nullptr; }
+	template<>
+	const A3DEntity* TUniqueObject<A3DSurfCylindricalData>::GetDefaultIndexerValue() const { return nullptr; }
+	template<>
+	const A3DEntity* TUniqueObject<A3DSurfExtrusionData>::GetDefaultIndexerValue() const { return nullptr; }
+	template<>
+	const A3DEntity* TUniqueObject<A3DSurfFromCurvesData>::GetDefaultIndexerValue() const { return nullptr; }
+	template<>
+	const A3DEntity* TUniqueObject<A3DSurfNurbsData>::GetDefaultIndexerValue() const { return nullptr; }
+	template<>
+	const A3DEntity* TUniqueObject<A3DSurfPipeData>::GetDefaultIndexerValue() const { return nullptr; }
+	template<>
+	const A3DEntity* TUniqueObject<A3DSurfPlaneData>::GetDefaultIndexerValue() const { return nullptr; }
+	template<>
+	const A3DEntity* TUniqueObject<A3DSurfRevolutionData>::GetDefaultIndexerValue() const { return nullptr; }
+	template<>
+	const A3DEntity* TUniqueObject<A3DSurfRuledData>::GetDefaultIndexerValue() const { return nullptr; }
+	template<>
+	const A3DEntity* TUniqueObject<A3DSurfSphereData>::GetDefaultIndexerValue() const { return nullptr; }
+	template<>
+	const A3DEntity* TUniqueObject<A3DSurfTorusData>::GetDefaultIndexerValue() const { return nullptr; }
+	template<>
+	const A3DEntity* TUniqueObject<A3DTess3DData>::GetDefaultIndexerValue() const { return nullptr; }
+	template<>
+	const A3DEntity* TUniqueObject<A3DTessBaseData>::GetDefaultIndexerValue() const { return nullptr; }
+	template<>
+	const A3DEntity* TUniqueObject<A3DTopoBodyData>::GetDefaultIndexerValue() const { return nullptr; }
+	template<>
+	const A3DEntity* TUniqueObject<A3DTopoBrepDataData>::GetDefaultIndexerValue() const { return nullptr; }
+	template<>
+	const A3DEntity* TUniqueObject<A3DTopoCoEdgeData>::GetDefaultIndexerValue() const { return nullptr; }
+	template<>
+	const A3DEntity* TUniqueObject<A3DTopoConnexData>::GetDefaultIndexerValue() const { return nullptr; }
+	template<>
+	const A3DEntity* TUniqueObject<A3DTopoContextData>::GetDefaultIndexerValue() const { return nullptr; }
+	template<>
+	const A3DEntity* TUniqueObject<A3DTopoEdgeData>::GetDefaultIndexerValue() const { return nullptr; }
+	template<>
+	const A3DEntity* TUniqueObject<A3DTopoFaceData>::GetDefaultIndexerValue() const { return nullptr; }
+	template<>
+	const A3DEntity* TUniqueObject<A3DTopoLoopData>::GetDefaultIndexerValue() const { return nullptr; }
+	template<>
+	const A3DEntity* TUniqueObject<A3DTopoShellData>::GetDefaultIndexerValue() const { return nullptr; }
+	template<>
+	const A3DEntity* TUniqueObject<A3DTopoUniqueVertexData>::GetDefaultIndexerValue() const { return nullptr; }
+	template<>
+	const A3DEntity* TUniqueObject<A3DTopoWireEdgeData>::GetDefaultIndexerValue() const { return nullptr; }
+	template<>
+	const A3DEntity* TUniqueObject<A3DVector2dData>::GetDefaultIndexerValue() const { return nullptr; }
+	template<>
+	const A3DEntity* TUniqueObject<A3DVector3dData>::GetDefaultIndexerValue() const { return nullptr; }
+
+	template<>
+	uint32 TUniqueObjectFromIndex<A3DGraphMaterialData>::GetDefaultIndexerValue() const { return A3D_DEFAULT_MATERIAL_INDEX; }
+	template<>
+	uint32 TUniqueObjectFromIndex<A3DGraphPictureData>::GetDefaultIndexerValue() const { return A3D_DEFAULT_PICTURE_INDEX; }
+	template<>
+	uint32 TUniqueObjectFromIndex<A3DGraphRgbColorData>::GetDefaultIndexerValue() const { return A3D_DEFAULT_COLOR_INDEX; }
+	template<>
+	uint32 TUniqueObjectFromIndex<A3DGraphStyleData>::GetDefaultIndexerValue() const { return A3D_DEFAULT_STYLE_INDEX; }
+	template<>
+	uint32 TUniqueObjectFromIndex<A3DGraphTextureApplicationData>::GetDefaultIndexerValue() const { return A3D_DEFAULT_TEXTURE_APPLICATION_INDEX; }
+	template<>
+	uint32 TUniqueObjectFromIndex<A3DGraphTextureDefinitionData>::GetDefaultIndexerValue() const { return A3D_DEFAULT_TEXTURE_DEFINITION_INDEX; }
+	template<>
+	const A3DEntity* TUniqueObject<A3DUTF8Char*>::GetDefaultIndexerValue() const { return nullptr; }
+
+	const uint32 FTechSoftDefaultValue::Material = A3D_DEFAULT_MATERIAL_INDEX;
+	const uint32 FTechSoftDefaultValue::Picture = A3D_DEFAULT_PICTURE_INDEX;
+	const uint32 FTechSoftDefaultValue::RgbColor = A3D_DEFAULT_COLOR_INDEX;
+	const uint32 FTechSoftDefaultValue::Style = A3D_DEFAULT_STYLE_INDEX;
+	const uint32 FTechSoftDefaultValue::TextureApplication = A3D_DEFAULT_TEXTURE_APPLICATION_INDEX;
+	const uint32 FTechSoftDefaultValue::TextureDefinition = A3D_DEFAULT_TEXTURE_DEFINITION_INDEX;
+
+}
+#endif
diff --git a/Engine/Source/Runtime/Datasmith/CADKernel/Engine/Private/TechSoft/TechSoftUniqueObjectImpl.h b/Engine/Source/Runtime/Datasmith/CADKernel/Engine/Private/TechSoft/TechSoftUniqueObjectImpl.h
new file mode 100644
--- /dev/null
+++ b/Engine/Source/Runtime/Datasmith/CADKernel/Engine/Private/TechSoft/TechSoftUniqueObjectImpl.h
@@ -0,0 +1,528 @@
+// Copyright Epic Games, Inc. All Rights Reserved.
+#pragma once
+
+#include "TechSoftIncludes.h"
+#include "TechSoftUniqueObject.h"
+
+#ifdef WITH_HOOPS
+
+namespace UE::CADKernel::TechSoft
+{
+
+	// Single-ownership smart TeshSoft object
+	// Use this when you need to manage TechSoft object's lifetime.
+	//
+	// TechSoft give access to void pointers
+	// According to the context, the class name of the void pointer is known but the class is unknown
+	// i.e. A3DSDKTypes.h defines all type like :
+	// 	   typedef void A3DEntity;		
+	// 	   typedef void A3DAsmModelFile; ...
+	// 
+	// From a pointer, TechSoft give access to a copy of the associated structure :
+	//
+	// const A3DXXXXX* pPointer;
+	// A3DXXXXXData sData; // the structure
+	// A3D_INITIALIZE_DATA(A3DXXXXXData, sData); // initialization of the structure
+	// A3DXXXXXXGet(pPointer, &sData); // Copy of the data of the pointer in the structure
+	// ...
+	// A3DXXXXXXGet(NULL, &sData); // Free the structure
+	//
+	// A3D_INITIALIZE_DATA, and all A3DXXXXXXGet methods are TechSoft macro
+	//
+
+
+	// TUniqueObject -----------------------------------
+
+	template<>
+	CADKERNELENGINE_API void TUniqueObject<A3DAsmModelFileData>::InitializeData();
+	template<>
+	CADKERNELENGINE_API void TUniqueObject<A3DAsmPartDefinitionData>::InitializeData();
+	template<>
+	CADKERNELENGINE_API void TUniqueObject<A3DAsmProductOccurrenceData>::InitializeData();
+	template<>
+	CADKERNELENGINE_API void TUniqueObject<A3DAsmProductOccurrenceDataCV5>::InitializeData();
+	template<>
+	CADKERNELENGINE_API void TUniqueObject<A3DAsmProductOccurrenceDataSLW>::InitializeData();
+	template<>
+	CADKERNELENGINE_API void TUniqueObject<A3DAsmProductOccurrenceDataUg>::InitializeData();
+	template<>
+	CADKERNELENGINE_API void TUniqueObject<A3DBoundingBoxData>::InitializeData();
+	template<>
+	CADKERNELENGINE_API void TUniqueObject<A3DCopyAndAdaptBrepModelData>::InitializeData();
+	template<>
+	CADKERNELENGINE_API void TUniqueObject<A3DCrvCircleData>::InitializeData();
+	template<>
+	CADKERNELENGINE_API void TUniqueObject<A3DCrvCompositeData>::InitializeData();
+	template<>
+	CADKERNELENGINE_API void TUniqueObject<A3DCrvEllipseData>::InitializeData();
+	template<>
+	CADKERNELENGINE_API void TUniqueObject<A3DCrvHelixData>::InitializeData();
+	template<>
+	CADKERNELENGINE_API void TUniqueObject<A3DCrvHyperbolaData>::InitializeData();
+	template<>
+	CADKERNELENGINE_API void TUniqueObject<A3DCrvLineData>::InitializeData();
+	template<>
+	CADKERNELENGINE_API void TUniqueObject<A3DCrvNurbsData>::InitializeData();
+	template<>
+	CADKERNELENGINE_API void TUniqueObject<A3DCrvParabolaData>::InitializeData();
+	template<>
+	CADKERNELENGINE_API void TUniqueObject<A3DCrvPolyLineData>::InitializeData();
+	template<>
+	CADKERNELENGINE_API void TUniqueObject<A3DCrvTransformData>::InitializeData();
+	template<>
+	CADKERNELENGINE_API void TUniqueObject<A3DDomainData>::InitializeData();
+	template<>
+	CADKERNELENGINE_API void TUniqueObject<A3DGlobalData>::InitializeData();
+	template<>
+	CADKERNELENGINE_API void TUniqueObject<A3DGraphicsData>::InitializeData();
+	template<>
+	CADKERNELENGINE_API void TUniqueObject<A3DIntervalData>::InitializeData();
+	template<>
+	CADKERNELENGINE_API void TUniqueObject<A3DMiscAttributeData>::InitializeData();
+	template<>
+	CADKERNELENGINE_API void TUniqueObject<A3DMiscCartesianTransformationData>::InitializeData();
+	template<>
+	CADKERNELENGINE_API void TUniqueObject<A3DMiscEntityReferenceData>::InitializeData();
+	template<>
+	CADKERNELENGINE_API void TUniqueObject<A3DMiscGeneralTransformationData>::InitializeData();
+	template<>
+	CADKERNELENGINE_API void TUniqueObject<A3DMiscMaterialPropertiesData>::InitializeData();
+	template<>
+	CADKERNELENGINE_API void TUniqueObject<A3DMiscReferenceOnCsysItemData>::InitializeData();
+	template<>
+	CADKERNELENGINE_API void TUniqueObject<A3DMiscReferenceOnTessData>::InitializeData();
+	template<>
+	CADKERNELENGINE_API void TUniqueObject<A3DMiscReferenceOnTopologyData>::InitializeData();
+	template<>
+	CADKERNELENGINE_API void TUniqueObject<A3DMiscSingleAttributeData>::InitializeData();
+	template<>
+	CADKERNELENGINE_API void TUniqueObject<A3DRWParamsExportPrcData>::InitializeData();
+	template<>
+	CADKERNELENGINE_API void TUniqueObject<A3DRiBrepModelData>::InitializeData();
+	template<>
+	CADKERNELENGINE_API void TUniqueObject<A3DRiCoordinateSystemData>::InitializeData();
+	template<>
+	CADKERNELENGINE_API void TUniqueObject<A3DRiDirectionData>::InitializeData();
+	template<>
+	CADKERNELENGINE_API void TUniqueObject<A3DRiPolyBrepModelData>::InitializeData();
+	template<>
+	CADKERNELENGINE_API void TUniqueObject<A3DRiRepresentationItemData>::InitializeData();
+	template<>
+	CADKERNELENGINE_API void TUniqueObject<A3DRiSetData>::InitializeData();
+	template<>
+	CADKERNELENGINE_API void TUniqueObject<A3DRootBaseData>::InitializeData();
+	template<>
+	CADKERNELENGINE_API void TUniqueObject<A3DRootBaseWithGraphicsData>::InitializeData();
+	template<>
+	CADKERNELENGINE_API void TUniqueObject<A3DRWParamsTessellationData>::InitializeData();
+	template<>
+	CADKERNELENGINE_API void TUniqueObject<A3DSewOptionsData>::InitializeData();
+	template<>
+	CADKERNELENGINE_API void TUniqueObject<A3DSurfBlend01Data>::InitializeData();
+	template<>
+	CADKERNELENGINE_API void TUniqueObject<A3DSurfBlend02Data>::InitializeData();
+	template<>
+	CADKERNELENGINE_API void TUniqueObject<A3DSurfBlend03Data>::InitializeData();
+	template<>
+	CADKERNELENGINE_API void TUniqueObject<A3DSurfConeData>::InitializeData();
+	template<>
+	CADKERNELENGINE_API void TUniqueObject<A3DSurfCylinderData>::InitializeData();
+	template<>
+	CADKERNELENGINE_API void TUniqueObject<A3DSurfCylindricalData>::InitializeData();
+	template<>
+	CADKERNELENGINE_API void TUniqueObject<A3DSurfExtrusionData>::InitializeData();
+	template<>
+	CADKERNELENGINE_API void TUniqueObject<A3DSurfFromCurvesData>::InitializeData();
+	template<>
+	CADKERNELENGINE_API void TUniqueObject<A3DSurfNurbsData>::InitializeData();
+	template<>
+	CADKERNELENGINE_API void TUniqueObject<A3DSurfPipeData>::InitializeData();
+	template<>
+	CADKERNELENGINE_API void TUniqueObject<A3DSurfPlaneData>::InitializeData();
+	template<>
+	CADKERNELENGINE_API void TUniqueObject<A3DSurfRevolutionData>::InitializeData();
+	template<>
+	CADKERNELENGINE_API void TUniqueObject<A3DSurfRuledData>::InitializeData();
+	template<>
+	CADKERNELENGINE_API void TUniqueObject<A3DSurfSphereData>::InitializeData();
+	template<>
+	CADKERNELENGINE_API void TUniqueObject<A3DSurfTorusData>::InitializeData();
+	template<>
+	CADKERNELENGINE_API void TUniqueObject<A3DTess3DData>::InitializeData();
+	template<>
+	CADKERNELENGINE_API void TUniqueObject<A3DTessBaseData>::InitializeData();
+	template<>
+	CADKERNELENGINE_API void TUniqueObject<A3DTopoBodyData>::InitializeData();
+	template<>
+	CADKERNELENGINE_API void TUniqueObject<A3DTopoBrepDataData>::InitializeData();
+	template<>
+	CADKERNELENGINE_API void TUniqueObject<A3DTopoCoEdgeData>::InitializeData();
+	template<>
+	CADKERNELENGINE_API void TUniqueObject<A3DTopoConnexData>::InitializeData();
+	template<>
+	CADKERNELENGINE_API void TUniqueObject<A3DTopoContextData>::InitializeData();
+	template<>
+	CADKERNELENGINE_API void TUniqueObject<A3DTopoEdgeData>::InitializeData();
+	template<>
+	CADKERNELENGINE_API void TUniqueObject<A3DTopoFaceData>::InitializeData();
+	template<>
+	CADKERNELENGINE_API void TUniqueObject<A3DTopoLoopData>::InitializeData();
+	template<>
+	CADKERNELENGINE_API void TUniqueObject<A3DTopoShellData>::InitializeData();
+	template<>
+	CADKERNELENGINE_API void TUniqueObject<A3DTopoUniqueVertexData>::InitializeData();
+	template<>
+	CADKERNELENGINE_API void TUniqueObject<A3DTopoMultipleVertexData>::InitializeData();
+	template<>
+	CADKERNELENGINE_API void TUniqueObject<A3DTopoWireEdgeData>::InitializeData();
+	template<>
+	CADKERNELENGINE_API void TUniqueObject<A3DVector2dData>::InitializeData();
+	template<>
+	CADKERNELENGINE_API void TUniqueObject<A3DVector3dData>::InitializeData();
+
+	template<>
+	CADKERNELENGINE_API A3DStatus TUniqueObject<A3DAsmModelFileData>::GetData(const A3DEntity* InEntityPtr);
+	template<>
+	CADKERNELENGINE_API A3DStatus TUniqueObject<A3DAsmPartDefinitionData>::GetData(const A3DEntity* InEntityPtr);
+	template<>
+	CADKERNELENGINE_API A3DStatus TUniqueObject<A3DAsmProductOccurrenceData>::GetData(const A3DEntity* InEntityPtr);
+	template<>
+	CADKERNELENGINE_API A3DStatus TUniqueObject<A3DAsmProductOccurrenceDataCV5>::GetData(const A3DEntity* InEntityPtr);
+	template<>
+	CADKERNELENGINE_API A3DStatus TUniqueObject<A3DAsmProductOccurrenceDataSLW>::GetData(const A3DEntity* InEntityPtr);
+	template<>
+	CADKERNELENGINE_API A3DStatus TUniqueObject<A3DAsmProductOccurrenceDataUg>::GetData(const A3DEntity* InEntityPtr);
+	template<>
+	CADKERNELENGINE_API A3DStatus TUniqueObject<A3DBoundingBoxData>::GetData(const A3DEntity* InEntityPtr);
+	template<>
+	CADKERNELENGINE_API A3DStatus TUniqueObject<A3DCopyAndAdaptBrepModelData>::GetData(const A3DEntity* InEntityPtr);
+	template<>
+	CADKERNELENGINE_API A3DStatus TUniqueObject<A3DCrvCircleData>::GetData(const A3DEntity* InEntityPtr);
+	template<>
+	CADKERNELENGINE_API A3DStatus TUniqueObject<A3DCrvCompositeData>::GetData(const A3DEntity* InEntityPtr);
+	template<>
+	CADKERNELENGINE_API A3DStatus TUniqueObject<A3DCrvEllipseData>::GetData(const A3DEntity* InEntityPtr);
+	template<>
+	CADKERNELENGINE_API A3DStatus TUniqueObject<A3DCrvHelixData>::GetData(const A3DEntity* InEntityPtr);
+	template<>
+	CADKERNELENGINE_API A3DStatus TUniqueObject<A3DCrvHyperbolaData>::GetData(const A3DEntity* InEntityPtr);
+	template<>
+	CADKERNELENGINE_API A3DStatus TUniqueObject<A3DCrvLineData>::GetData(const A3DEntity* InEntityPtr);
+	template<>
+	CADKERNELENGINE_API A3DStatus TUniqueObject<A3DCrvNurbsData>::GetData(const A3DEntity* InEntityPtr);
+	template<>
+	CADKERNELENGINE_API A3DStatus TUniqueObject<A3DCrvParabolaData>::GetData(const A3DEntity* InEntityPtr);
+	template<>
+	CADKERNELENGINE_API A3DStatus TUniqueObject<A3DCrvPolyLineData>::GetData(const A3DEntity* InEntityPtr);
+	template<>
+	CADKERNELENGINE_API A3DStatus TUniqueObject<A3DCrvTransformData>::GetData(const A3DEntity* InEntityPtr);
+	template<>
+	CADKERNELENGINE_API A3DStatus TUniqueObject<A3DDomainData>::GetData(const A3DEntity* InEntityPtr);
+	template<>
+	CADKERNELENGINE_API A3DStatus TUniqueObject<A3DGlobalData>::GetData(const A3DEntity* InEntityPtr);
+	template<>
+	CADKERNELENGINE_API A3DStatus TUniqueObject<A3DGraphicsData>::GetData(const A3DEntity* InEntityPtr);
+	template<>
+	CADKERNELENGINE_API A3DStatus TUniqueObject<A3DIntervalData>::GetData(const A3DEntity* InEntityPtr);
+	template<>
+	CADKERNELENGINE_API A3DStatus TUniqueObject<A3DMiscAttributeData>::GetData(const A3DEntity* InEntityPtr);
+	template<>
+	CADKERNELENGINE_API A3DStatus TUniqueObject<A3DMiscCartesianTransformationData>::GetData(const A3DEntity* InEntityPtr);
+	template<>
+	CADKERNELENGINE_API A3DStatus TUniqueObject<A3DMiscEntityReferenceData>::GetData(const A3DEntity* InEntityPtr);
+	template<>
+	CADKERNELENGINE_API A3DStatus TUniqueObject<A3DMiscGeneralTransformationData>::GetData(const A3DEntity* InEntityPtr);
+	template<>
+	CADKERNELENGINE_API A3DStatus TUniqueObject<A3DMiscMaterialPropertiesData>::GetData(const A3DEntity* InEntityPtr);
+	template<>
+	CADKERNELENGINE_API A3DStatus TUniqueObject<A3DMiscReferenceOnCsysItemData>::GetData(const A3DEntity* InEntityPtr);
+	template<>
+	CADKERNELENGINE_API A3DStatus TUniqueObject<A3DMiscReferenceOnTessData>::GetData(const A3DEntity* InEntityPtr);
+	template<>
+	CADKERNELENGINE_API A3DStatus TUniqueObject<A3DMiscReferenceOnTopologyData>::GetData(const A3DEntity* InEntityPtr);
+	template<>
+	CADKERNELENGINE_API A3DStatus TUniqueObject<A3DMiscSingleAttributeData>::GetData(const A3DEntity* InEntityPtr);
+	template<>
+	CADKERNELENGINE_API A3DStatus TUniqueObject<A3DRWParamsExportPrcData>::GetData(const A3DEntity* InEntityPtr);
+	template<>
+	CADKERNELENGINE_API A3DStatus TUniqueObject<A3DRiBrepModelData>::GetData(const A3DEntity* InEntityPtr);
+	template<>
+	CADKERNELENGINE_API A3DStatus TUniqueObject<A3DRiCoordinateSystemData>::GetData(const A3DEntity* InEntityPtr);
+	template<>
+	CADKERNELENGINE_API A3DStatus TUniqueObject<A3DRiDirectionData>::GetData(const A3DEntity* InEntityPtr);
+	template<>
+	CADKERNELENGINE_API A3DStatus TUniqueObject<A3DRiPolyBrepModelData>::GetData(const A3DEntity* InEntityPtr);
+	template<>
+	CADKERNELENGINE_API A3DStatus TUniqueObject<A3DRiRepresentationItemData>::GetData(const A3DEntity* InEntityPtr);
+	template<>
+	CADKERNELENGINE_API A3DStatus TUniqueObject<A3DRiSetData>::GetData(const A3DEntity* InEntityPtr);
+	template<>
+	CADKERNELENGINE_API A3DStatus TUniqueObject<A3DRootBaseData>::GetData(const A3DEntity* InEntityPtr);
+	template<>
+	CADKERNELENGINE_API A3DStatus TUniqueObject<A3DRootBaseWithGraphicsData>::GetData(const A3DEntity* InEntityPtr);
+	template<>
+	CADKERNELENGINE_API A3DStatus TUniqueObject<A3DRWParamsTessellationData>::GetData(const A3DEntity* InEntityPtr);
+	template<>
+	CADKERNELENGINE_API A3DStatus TUniqueObject<A3DSewOptionsData>::GetData(const A3DEntity* InEntityPtr);
+	template<>
+	CADKERNELENGINE_API A3DStatus TUniqueObject<A3DSurfBlend01Data>::GetData(const A3DEntity* InEntityPtr);
+	template<>
+	CADKERNELENGINE_API A3DStatus TUniqueObject<A3DSurfBlend02Data>::GetData(const A3DEntity* InEntityPtr);
+	template<>
+	CADKERNELENGINE_API A3DStatus TUniqueObject<A3DSurfBlend03Data>::GetData(const A3DEntity* InEntityPtr);
+	template<>
+	CADKERNELENGINE_API A3DStatus TUniqueObject<A3DSurfConeData>::GetData(const A3DEntity* InEntityPtr);
+	template<>
+	CADKERNELENGINE_API A3DStatus TUniqueObject<A3DSurfCylinderData>::GetData(const A3DEntity* InEntityPtr);
+	template<>
+	CADKERNELENGINE_API A3DStatus TUniqueObject<A3DSurfCylindricalData>::GetData(const A3DEntity* InEntityPtr);
+	template<>
+	CADKERNELENGINE_API A3DStatus TUniqueObject<A3DSurfExtrusionData>::GetData(const A3DEntity* InEntityPtr);
+	template<>
+	CADKERNELENGINE_API A3DStatus TUniqueObject<A3DSurfFromCurvesData>::GetData(const A3DEntity* InEntityPtr);
+	template<>
+	CADKERNELENGINE_API A3DStatus TUniqueObject<A3DSurfNurbsData>::GetData(const A3DEntity* InEntityPtr);
+	template<>
+	CADKERNELENGINE_API A3DStatus TUniqueObject<A3DSurfPipeData>::GetData(const A3DEntity* InEntityPtr);
+	template<>
+	CADKERNELENGINE_API A3DStatus TUniqueObject<A3DSurfPlaneData>::GetData(const A3DEntity* InEntityPtr);
+	template<>
+	CADKERNELENGINE_API A3DStatus TUniqueObject<A3DSurfRevolutionData>::GetData(const A3DEntity* InEntityPtr);
+	template<>
+	CADKERNELENGINE_API A3DStatus TUniqueObject<A3DSurfRuledData>::GetData(const A3DEntity* InEntityPtr);
+	template<>
+	CADKERNELENGINE_API A3DStatus TUniqueObject<A3DSurfSphereData>::GetData(const A3DEntity* InEntityPtr);
+	template<>
+	CADKERNELENGINE_API A3DStatus TUniqueObject<A3DSurfTorusData>::GetData(const A3DEntity* InEntityPtr);
+	template<>
+	CADKERNELENGINE_API A3DStatus TUniqueObject<A3DTess3DData>::GetData(const A3DEntity* InEntityPtr);
+	template<>
+	CADKERNELENGINE_API A3DStatus TUniqueObject<A3DTessBaseData>::GetData(const A3DEntity* InEntityPtr);
+	template<>
+	CADKERNELENGINE_API A3DStatus TUniqueObject<A3DTopoBodyData>::GetData(const A3DEntity* InEntityPtr);
+	template<>
+	CADKERNELENGINE_API A3DStatus TUniqueObject<A3DTopoBrepDataData>::GetData(const A3DEntity* InEntityPtr);
+	template<>
+	CADKERNELENGINE_API A3DStatus TUniqueObject<A3DTopoCoEdgeData>::GetData(const A3DEntity* InEntityPtr);
+	template<>
+	CADKERNELENGINE_API A3DStatus TUniqueObject<A3DTopoConnexData>::GetData(const A3DEntity* InEntityPtr);
+	template<>
+	CADKERNELENGINE_API A3DStatus TUniqueObject<A3DTopoContextData>::GetData(const A3DEntity* InEntityPtr);
+	template<>
+	CADKERNELENGINE_API A3DStatus TUniqueObject<A3DTopoEdgeData>::GetData(const A3DEntity* InEntityPtr);
+	template<>
+	CADKERNELENGINE_API A3DStatus TUniqueObject<A3DTopoFaceData>::GetData(const A3DEntity* InEntityPtr);
+	template<>
+	CADKERNELENGINE_API A3DStatus TUniqueObject<A3DTopoLoopData>::GetData(const A3DEntity* InEntityPtr);
+	template<>
+	CADKERNELENGINE_API A3DStatus TUniqueObject<A3DTopoShellData>::GetData(const A3DEntity* InEntityPtr);
+	template<>
+	CADKERNELENGINE_API A3DStatus TUniqueObject<A3DTopoUniqueVertexData>::GetData(const A3DEntity* InEntityPtr);
+	template<>
+	CADKERNELENGINE_API A3DStatus TUniqueObject<A3DTopoMultipleVertexData>::GetData(const A3DEntity* InEntityPtr);
+	template<>
+	CADKERNELENGINE_API A3DStatus TUniqueObject<A3DTopoWireEdgeData>::GetData(const A3DEntity* InEntityPtr);
+	template<>
+	CADKERNELENGINE_API A3DStatus TUniqueObject<A3DVector2dData>::GetData(const A3DEntity* InEntityPtr);
+	template<>
+	CADKERNELENGINE_API A3DStatus TUniqueObject<A3DVector3dData>::GetData(const A3DEntity* InEntityPtr);
+
+	template<>
+	CADKERNELENGINE_API const A3DEntity* TUniqueObject<A3DAsmModelFileData>::GetDefaultIndexerValue() const;
+	template<>
+	CADKERNELENGINE_API const A3DEntity* TUniqueObject<A3DAsmPartDefinitionData>::GetDefaultIndexerValue() const;
+	template<>
+	CADKERNELENGINE_API const A3DEntity* TUniqueObject<A3DAsmProductOccurrenceData>::GetDefaultIndexerValue() const;
+	template<>
+	CADKERNELENGINE_API const A3DEntity* TUniqueObject<A3DAsmProductOccurrenceDataCV5>::GetDefaultIndexerValue() const;
+	template<>
+	CADKERNELENGINE_API const A3DEntity* TUniqueObject<A3DAsmProductOccurrenceDataSLW>::GetDefaultIndexerValue() const;
+	template<>
+	CADKERNELENGINE_API const A3DEntity* TUniqueObject<A3DAsmProductOccurrenceDataUg>::GetDefaultIndexerValue() const;
+	template<>
+	CADKERNELENGINE_API const A3DEntity* TUniqueObject<A3DBoundingBoxData>::GetDefaultIndexerValue() const;
+	template<>
+	CADKERNELENGINE_API const A3DEntity* TUniqueObject<A3DCopyAndAdaptBrepModelData>::GetDefaultIndexerValue() const;
+	template<>
+	CADKERNELENGINE_API const A3DEntity* TUniqueObject<A3DCrvCircleData>::GetDefaultIndexerValue() const;
+	template<>
+	CADKERNELENGINE_API const A3DEntity* TUniqueObject<A3DCrvCompositeData>::GetDefaultIndexerValue() const;
+	template<>
+	CADKERNELENGINE_API const A3DEntity* TUniqueObject<A3DCrvEllipseData>::GetDefaultIndexerValue() const;
+	template<>
+	CADKERNELENGINE_API const A3DEntity* TUniqueObject<A3DCrvHelixData>::GetDefaultIndexerValue() const;
+	template<>
+	CADKERNELENGINE_API const A3DEntity* TUniqueObject<A3DCrvHyperbolaData>::GetDefaultIndexerValue() const;
+	template<>
+	CADKERNELENGINE_API const A3DEntity* TUniqueObject<A3DCrvLineData>::GetDefaultIndexerValue() const;
+	template<>
+	CADKERNELENGINE_API const A3DEntity* TUniqueObject<A3DCrvNurbsData>::GetDefaultIndexerValue() const;
+	template<>
+	CADKERNELENGINE_API const A3DEntity* TUniqueObject<A3DCrvParabolaData>::GetDefaultIndexerValue() const;
+	template<>
+	CADKERNELENGINE_API const A3DEntity* TUniqueObject<A3DCrvPolyLineData>::GetDefaultIndexerValue() const;
+	template<>
+	CADKERNELENGINE_API const A3DEntity* TUniqueObject<A3DCrvTransformData>::GetDefaultIndexerValue() const;
+	template<>
+	CADKERNELENGINE_API const A3DEntity* TUniqueObject<A3DDomainData>::GetDefaultIndexerValue() const;
+	template<>
+	CADKERNELENGINE_API const A3DEntity* TUniqueObject<A3DGlobalData>::GetDefaultIndexerValue() const;
+	template<>
+	CADKERNELENGINE_API const A3DEntity* TUniqueObject<A3DGraphicsData>::GetDefaultIndexerValue() const;
+	template<>
+	CADKERNELENGINE_API const A3DEntity* TUniqueObject<A3DIntervalData>::GetDefaultIndexerValue() const;
+	template<>
+	CADKERNELENGINE_API const A3DEntity* TUniqueObject<A3DMiscAttributeData>::GetDefaultIndexerValue() const;
+	template<>
+	CADKERNELENGINE_API const A3DEntity* TUniqueObject<A3DMiscCartesianTransformationData>::GetDefaultIndexerValue() const;
+	template<>
+	CADKERNELENGINE_API const A3DEntity* TUniqueObject<A3DMiscEntityReferenceData>::GetDefaultIndexerValue() const;
+	template<>
+	CADKERNELENGINE_API const A3DEntity* TUniqueObject<A3DMiscGeneralTransformationData>::GetDefaultIndexerValue() const;
+	template<>
+	CADKERNELENGINE_API const A3DEntity* TUniqueObject<A3DMiscMaterialPropertiesData>::GetDefaultIndexerValue() const;
+	template<>
+	CADKERNELENGINE_API const A3DEntity* TUniqueObject<A3DMiscReferenceOnCsysItemData>::GetDefaultIndexerValue() const;
+	template<>
+	CADKERNELENGINE_API const A3DEntity* TUniqueObject<A3DMiscReferenceOnTessData>::GetDefaultIndexerValue() const;
+	template<>
+	CADKERNELENGINE_API const A3DEntity* TUniqueObject<A3DMiscReferenceOnTopologyData>::GetDefaultIndexerValue() const;
+	template<>
+	CADKERNELENGINE_API const A3DEntity* TUniqueObject<A3DMiscSingleAttributeData>::GetDefaultIndexerValue() const;
+	template<>
+	CADKERNELENGINE_API const A3DEntity* TUniqueObject<A3DRWParamsExportPrcData>::GetDefaultIndexerValue() const;
+	template<>
+	CADKERNELENGINE_API const A3DEntity* TUniqueObject<A3DRiBrepModelData>::GetDefaultIndexerValue() const;
+	template<>
+	CADKERNELENGINE_API const A3DEntity* TUniqueObject<A3DRiCoordinateSystemData>::GetDefaultIndexerValue() const;
+	template<>
+	CADKERNELENGINE_API const A3DEntity* TUniqueObject<A3DRiDirectionData>::GetDefaultIndexerValue() const;
+	template<>
+	CADKERNELENGINE_API const A3DEntity* TUniqueObject<A3DRiPolyBrepModelData>::GetDefaultIndexerValue() const;
+	template<>
+	CADKERNELENGINE_API const A3DEntity* TUniqueObject<A3DRiRepresentationItemData>::GetDefaultIndexerValue() const;
+	template<>
+	CADKERNELENGINE_API const A3DEntity* TUniqueObject<A3DRiSetData>::GetDefaultIndexerValue() const;
+	template<>
+	CADKERNELENGINE_API const A3DEntity* TUniqueObject<A3DRootBaseData>::GetDefaultIndexerValue() const;
+	template<>
+	CADKERNELENGINE_API const A3DEntity* TUniqueObject<A3DRootBaseWithGraphicsData>::GetDefaultIndexerValue() const;
+	template<>
+	CADKERNELENGINE_API const A3DEntity* TUniqueObject<A3DRWParamsTessellationData>::GetDefaultIndexerValue() const;
+	template<>
+	CADKERNELENGINE_API const A3DEntity* TUniqueObject<A3DSewOptionsData>::GetDefaultIndexerValue() const;
+	template<>
+	CADKERNELENGINE_API const A3DEntity* TUniqueObject<A3DSurfBlend01Data>::GetDefaultIndexerValue() const;
+	template<>
+	CADKERNELENGINE_API const A3DEntity* TUniqueObject<A3DSurfBlend02Data>::GetDefaultIndexerValue() const;
+	template<>
+	CADKERNELENGINE_API const A3DEntity* TUniqueObject<A3DSurfBlend03Data>::GetDefaultIndexerValue() const;
+	template<>
+	CADKERNELENGINE_API const A3DEntity* TUniqueObject<A3DSurfConeData>::GetDefaultIndexerValue() const;
+	template<>
+	CADKERNELENGINE_API const A3DEntity* TUniqueObject<A3DSurfCylinderData>::GetDefaultIndexerValue() const;
+	template<>
+	CADKERNELENGINE_API const A3DEntity* TUniqueObject<A3DSurfCylindricalData>::GetDefaultIndexerValue() const;
+	template<>
+	CADKERNELENGINE_API const A3DEntity* TUniqueObject<A3DSurfExtrusionData>::GetDefaultIndexerValue() const;
+	template<>
+	CADKERNELENGINE_API const A3DEntity* TUniqueObject<A3DSurfFromCurvesData>::GetDefaultIndexerValue() const;
+	template<>
+	CADKERNELENGINE_API const A3DEntity* TUniqueObject<A3DSurfNurbsData>::GetDefaultIndexerValue() const;
+	template<>
+	CADKERNELENGINE_API const A3DEntity* TUniqueObject<A3DSurfPipeData>::GetDefaultIndexerValue() const;
+	template<>
+	CADKERNELENGINE_API const A3DEntity* TUniqueObject<A3DSurfPlaneData>::GetDefaultIndexerValue() const;
+	template<>
+	CADKERNELENGINE_API const A3DEntity* TUniqueObject<A3DSurfRevolutionData>::GetDefaultIndexerValue() const;
+	template<>
+	CADKERNELENGINE_API const A3DEntity* TUniqueObject<A3DSurfRuledData>::GetDefaultIndexerValue() const;
+	template<>
+	CADKERNELENGINE_API const A3DEntity* TUniqueObject<A3DSurfSphereData>::GetDefaultIndexerValue() const;
+	template<>
+	CADKERNELENGINE_API const A3DEntity* TUniqueObject<A3DSurfTorusData>::GetDefaultIndexerValue() const;
+	template<>
+	CADKERNELENGINE_API const A3DEntity* TUniqueObject<A3DTess3DData>::GetDefaultIndexerValue() const;
+	template<>
+	CADKERNELENGINE_API const A3DEntity* TUniqueObject<A3DTessBaseData>::GetDefaultIndexerValue() const;
+	template<>
+	CADKERNELENGINE_API const A3DEntity* TUniqueObject<A3DTopoBodyData>::GetDefaultIndexerValue() const;
+	template<>
+	CADKERNELENGINE_API const A3DEntity* TUniqueObject<A3DTopoBrepDataData>::GetDefaultIndexerValue() const;
+	template<>
+	CADKERNELENGINE_API const A3DEntity* TUniqueObject<A3DTopoCoEdgeData>::GetDefaultIndexerValue() const;
+	template<>
+	CADKERNELENGINE_API const A3DEntity* TUniqueObject<A3DTopoConnexData>::GetDefaultIndexerValue() const;
+	template<>
+	CADKERNELENGINE_API const A3DEntity* TUniqueObject<A3DTopoContextData>::GetDefaultIndexerValue() const;
+	template<>
+	CADKERNELENGINE_API const A3DEntity* TUniqueObject<A3DTopoEdgeData>::GetDefaultIndexerValue() const;
+	template<>
+	CADKERNELENGINE_API const A3DEntity* TUniqueObject<A3DTopoFaceData>::GetDefaultIndexerValue() const;
+	template<>
+	CADKERNELENGINE_API const A3DEntity* TUniqueObject<A3DTopoLoopData>::GetDefaultIndexerValue() const;
+	template<>
+	CADKERNELENGINE_API const A3DEntity* TUniqueObject<A3DTopoShellData>::GetDefaultIndexerValue() const;
+	template<>
+	CADKERNELENGINE_API const A3DEntity* TUniqueObject<A3DTopoUniqueVertexData>::GetDefaultIndexerValue() const;
+	template<>
+	CADKERNELENGINE_API const A3DEntity* TUniqueObject<A3DTopoMultipleVertexData>::GetDefaultIndexerValue() const;
+	template<>
+	CADKERNELENGINE_API const A3DEntity* TUniqueObject<A3DTopoWireEdgeData>::GetDefaultIndexerValue() const;
+	template<>
+	CADKERNELENGINE_API const A3DEntity* TUniqueObject<A3DVector2dData>::GetDefaultIndexerValue() const;
+	template<>
+	CADKERNELENGINE_API const A3DEntity* TUniqueObject<A3DVector3dData>::GetDefaultIndexerValue() const;
+
+	// TUniqueObjectFromIndex -----------------------------------
+
+	template<>
+	CADKERNELENGINE_API void TUniqueObjectFromIndex<A3DGraphMaterialData>::InitializeData();
+	template<>
+	CADKERNELENGINE_API void TUniqueObjectFromIndex<A3DGraphPictureData>::InitializeData();
+	template<>
+	CADKERNELENGINE_API void TUniqueObjectFromIndex<A3DGraphRgbColorData>::InitializeData();
+	template<>
+	CADKERNELENGINE_API void TUniqueObjectFromIndex<A3DGraphStyleData>::InitializeData();
+	template<>
+	CADKERNELENGINE_API void TUniqueObjectFromIndex<A3DGraphTextureApplicationData>::InitializeData();
+	template<>
+	CADKERNELENGINE_API void TUniqueObjectFromIndex<A3DGraphTextureDefinitionData>::InitializeData();
+
+	template<>
+	CADKERNELENGINE_API A3DStatus TUniqueObjectFromIndex<A3DGraphMaterialData>::GetData(const uint32 InEntityIndex);
+	template<>
+	CADKERNELENGINE_API A3DStatus TUniqueObjectFromIndex<A3DGraphPictureData>::GetData(const uint32 InEntityIndex);
+	template<>
+	CADKERNELENGINE_API A3DStatus TUniqueObjectFromIndex<A3DGraphRgbColorData>::GetData(const uint32 InEntityIndex);
+	template<>
+	CADKERNELENGINE_API A3DStatus TUniqueObjectFromIndex<A3DGraphStyleData>::GetData(const uint32 InEntityIndex);
+	template<>
+	CADKERNELENGINE_API A3DStatus TUniqueObjectFromIndex<A3DGraphTextureApplicationData>::GetData(const uint32 InEntityIndex);
+	template<>
+	CADKERNELENGINE_API A3DStatus TUniqueObjectFromIndex<A3DGraphTextureDefinitionData>::GetData(const uint32 InEntityIndex);
+
+	template<>
+	CADKERNELENGINE_API uint32 TUniqueObjectFromIndex<A3DGraphMaterialData>::GetDefaultIndexerValue() const;
+	template<>
+	CADKERNELENGINE_API uint32 TUniqueObjectFromIndex<A3DGraphPictureData>::GetDefaultIndexerValue() const;
+	template<>
+	CADKERNELENGINE_API uint32 TUniqueObjectFromIndex<A3DGraphRgbColorData>::GetDefaultIndexerValue() const;
+	template<>
+	CADKERNELENGINE_API uint32 TUniqueObjectFromIndex<A3DGraphStyleData>::GetDefaultIndexerValue() const;
+	template<>
+	CADKERNELENGINE_API uint32 TUniqueObjectFromIndex<A3DGraphTextureApplicationData>::GetDefaultIndexerValue() const;
+	template<>
+	CADKERNELENGINE_API uint32 TUniqueObjectFromIndex<A3DGraphTextureDefinitionData>::GetDefaultIndexerValue() const;
+
+	// A3DUTF8Char* -----------------------------------
+
+	template<>
+	CADKERNELENGINE_API void TUniqueObject<A3DUTF8Char*>::InitializeData();
+	template<>
+	CADKERNELENGINE_API A3DStatus TUniqueObject<A3DUTF8Char*>::GetData(const A3DEntity* InEntityPtr);
+	template<>
+	CADKERNELENGINE_API const A3DEntity* TUniqueObject<A3DUTF8Char*>::GetDefaultIndexerValue() const;
+
+}
+#endif
diff --git a/Engine/Source/Runtime/Datasmith/CADKernel/Engine/Private/TechSoft/TechSoftUtilities.cpp b/Engine/Source/Runtime/Datasmith/CADKernel/Engine/Private/TechSoft/TechSoftUtilities.cpp
new file mode 100644
--- /dev/null
+++ b/Engine/Source/Runtime/Datasmith/CADKernel/Engine/Private/TechSoft/TechSoftUtilities.cpp
@@ -0,0 +1,325 @@
+// Copyright Epic Games, Inc. All Rights Reserved.
+
+#include "TechSoftUtilities.h"
+#include "CADKernelEngine.h"
+
+#ifdef WITH_HOOPS
+
+#include "TechSoftUniqueObjectImpl.h"
+
+#include "MeshUtilities.h"
+
+namespace UE::CADKernel
+{
+	bool FTechSoftUtilities::Save(A3DRiRepresentationItem* Representation, const FString& FilePath)
+	{
+		if (!FTechSoftLibrary::IsInitialized())
+		{
+			return false;
+		}
+
+		return false;
+	}
+
+	bool FTechSoftUtilities::Tessellate(A3DRiRepresentationItem* Representation, const FTessellationContext& Context, FMeshDescription& Mesh)
+	{
+		using namespace UE::CADKernel::MeshUtilities;
+
+		if (!FTechSoftLibrary::IsInitialized())
+		{
+			return false;
+		}
+
+		TSharedPtr<FMeshWrapperAbstract> MeshInterface = FMeshWrapperAbstract::MakeWrapper(Context, Mesh);
+		return TechSoftUtilities::Tessellate(Representation, Context, *MeshInterface);
+	}
+
+	bool FTechSoftUtilities::Tessellate(A3DRiRepresentationItem* Representation, const FTessellationContext& Context, UE::Geometry::FDynamicMesh3& Mesh)
+	{
+		using namespace UE::CADKernel::MeshUtilities;
+
+		if (!FTechSoftLibrary::IsInitialized())
+		{
+			return false;
+		}
+
+		TSharedPtr<FMeshWrapperAbstract> MeshInterface = FMeshWrapperAbstract::MakeWrapper(Context, Mesh);
+		return TechSoftUtilities::Tessellate(Representation, Context, *MeshInterface);
+	}
+
+	A3DRiRepresentationItem* FTechSoftUtilities::CreateRIBRep(A3DTopoShell* TopoShellPtr)
+	{
+		if (!FTechSoftLibrary::IsInitialized())
+		{
+			return nullptr;
+		}
+
+		return nullptr;
+	}
+
+	A3DTopoEdge* FTechSoftUtilities::CreateTopoEdge()
+	{
+		if (!FTechSoftLibrary::IsInitialized())
+		{
+			return nullptr;
+		}
+
+		return nullptr;
+	}
+
+	A3DTopoFace* FTechSoftUtilities::CreateTopoFaceWithNaturalLoop(A3DSurfBase* CarrierSurface)
+	{
+		if (!FTechSoftLibrary::IsInitialized())
+		{
+			return nullptr;
+		}
+
+		return nullptr;
+	}
+
+	A3DCrvNurbs* FTechSoftUtilities::CreateTrimNurbsCurve(A3DCrvNurbs* CurveNurbsPtr, double UMin, double UMax, bool bIs2D)
+	{
+		if (!FTechSoftLibrary::IsInitialized())
+		{
+			return nullptr;
+		}
+
+		return nullptr;
+	}
+
+	A3DTopoShell* FTechSoftUtilities::CreateTopoShell(const A3DTopoShellData& TopoShellData)
+	{
+		if (!FTechSoftLibrary::IsInitialized())
+		{
+			return nullptr;
+		}
+
+		return nullptr;
+	}
+
+	A3DSurfNurbs* FTechSoftUtilities::CreateSurfaceNurbs(const A3DSurfNurbsData& SurfaceNurbsData)
+	{
+		if (!FTechSoftLibrary::IsInitialized())
+		{
+			return nullptr;
+		}
+
+		return nullptr;
+	}
+
+	A3DStatus FTechSoftUtilities::SetEntityGraphicsColor(A3DEntity* InEntity, FColor Color)
+	{
+		if (!FTechSoftLibrary::IsInitialized())
+		{
+			return A3DStatus::A3D_NOT_AVAILABLE;
+		}
+
+		return A3DStatus::A3D_NOT_IMPLEMENTED;
+	}
+
+	A3DTopoFace* FTechSoftUtilities::CreateTopoFace(const A3DTopoFaceData& TopoFaceData)
+	{
+		if (!FTechSoftLibrary::IsInitialized())
+		{
+			return nullptr;
+		}
+
+		return nullptr;
+	}
+
+	A3DTopoLoop* FTechSoftUtilities::CreateTopoLoop(const A3DTopoLoopData& TopoLoopData)
+	{
+		if (!FTechSoftLibrary::IsInitialized())
+		{
+			return nullptr;
+		}
+
+		return nullptr;
+	}
+
+	A3DTopoEdge* FTechSoftUtilities::CreateTopoEdge(const A3DTopoEdgeData& TopoEdgeData)
+	{
+		if (!FTechSoftLibrary::IsInitialized())
+		{
+			return nullptr;
+		}
+
+		return nullptr;
+	}
+
+	A3DTopoCoEdge* FTechSoftUtilities::CreateTopoCoEdge(const A3DTopoCoEdgeData& TopoCoEdgeData)
+	{
+		if (!FTechSoftLibrary::IsInitialized())
+		{
+			return nullptr;
+		}
+
+		return nullptr;
+	}
+
+	A3DCrvNurbs* FTechSoftUtilities::CreateCurveNurbs(const A3DCrvNurbsData& CurveNurbsData)
+	{
+		if (!FTechSoftLibrary::IsInitialized())
+		{
+			return nullptr;
+		}
+
+		return nullptr;
+	}
+
+	A3DStatus FTechSoftUtilities::LinkCoEdges(A3DTopoCoEdge* CoEdgePtr, A3DTopoCoEdge* NeighbourCoEdgePtr)
+	{
+		if (!FTechSoftLibrary::IsInitialized())
+		{
+			return A3DStatus::A3D_NOT_AVAILABLE;
+		}
+
+		return A3DStatus::A3D_NOT_IMPLEMENTED;
+	}
+
+	bool TechSoftUtilities::Tessellate(A3DRiRepresentationItem* Representation, const FTessellationContext& Context, MeshUtilities::FMeshWrapperAbstract& MeshInterface)
+	{
+		if (!FTechSoftLibrary::IsInitialized())
+		{
+			return false;
+		}
+
+		using namespace UE::CADKernel;
+
+		// #cad_import: Create set of helpers IsA[Whatever] in FTechSoftLibrary
+		// Here it would FTechSoftLibrary::IsATriangulatedMesh
+		A3DEEntityType Type;
+		A3DEntityGetType(Representation, &Type);
+		if (Type == kA3DTypeRiPolyBrepModel)
+		{
+			return AddRepresentation(Representation, MeshInterface);
+		}
+
+		TechSoft::TUniqueObject<A3DRiRepresentationItemData> RepresentationData(Representation);
+
+		if (!RepresentationData.IsValid())
+		{
+			return false;
+}
+
+		TArray<A3DRiBrepModel*> NewBReps;
+		const FCADKernelTessellationSettings& TessellationSettings = Context.TessellationSettings;
+
+		if (TessellationSettings.StitchingTechnique == ECADKernelStitchingTechnique::StitchingHeal)
+		{
+			// cad_import: Review unit conversion - mm or cm to model unit
+			const double SewingTolerance = TessellationSettings.GetStitchingTolerance() / TessellationSettings.UnitMultiplier;
+			if (!SewBReps({ Representation }, SewingTolerance, NewBReps))
+			{
+				// Log message
+			}
+		}
+
+		if (!NewBReps.IsEmpty())
+		{
+			for (A3DRiBrepModel* BrepModel : NewBReps)
+			{
+				// It is ok to use BrepModel as a A3DRiRepresentationItem.
+				// TechSoft is feeling the required structure according to the data structure asked for
+				if (TessellateRepresentation(BrepModel, TessellationSettings))
+				{
+					AddRepresentation(BrepModel, MeshInterface);
+				}
+			}
+		}
+		else
+		{
+			if (TessellateRepresentation(Representation, TessellationSettings))
+			{
+				AddRepresentation(Representation, MeshInterface);
+			}
+		}
+
+		return true;
+	}
+
+	bool TechSoftUtilities::TessellateRepresentation(A3DRiRepresentationItem* Representation, const FCADKernelTessellationSettings& Settings)
+	{
+		if (!FTechSoftLibrary::IsInitialized())
+		{
+			return false;
+		}
+
+		// TUniqueTechSoftObj does not work in this case
+		TechSoft::TUniqueObject<A3DRWParamsTessellationData> TessellationParameters;
+
+		TessellationParameters->m_eTessellationLevelOfDetail = kA3DTessLODUserDefined; // Enum to specify predefined values for some following members.
+		TessellationParameters->m_bUseHeightInsteadOfRatio = A3D_TRUE;
+		TessellationParameters->m_dMaxChordHeight = Settings.GetChordTolerance(); // This is centimeters
+		if (!FMath::IsNearlyZero(Settings.UnitMultiplier))
+		{
+			// Convert back to unit of imported data
+			TessellationParameters->m_dMaxChordHeight /= Settings.UnitMultiplier;
+		}
+
+		TessellationParameters->m_dAngleToleranceDeg = Settings.NormalTolerance;
+		TessellationParameters->m_dMaximalTriangleEdgeLength = 0; //Settings.MaxEdgeLength;
+
+		TessellationParameters->m_bAccurateTessellation = A3D_FALSE;  // A3D_FALSE' indicates the tessellation is set for visualization
+		TessellationParameters->m_bAccurateTessellationWithGrid = A3D_FALSE; // Enable accurate tessellation with faces inner points on a grid.
+		TessellationParameters->m_dAccurateTessellationWithGridMaximumStitchLength = 0; 	// Maximal grid stitch length. Disabled if value is 0. Be careful, a too small value can generate a huge tessellation.
+
+		TessellationParameters->m_bKeepUVPoints = A3D_TRUE; // Keep parametric points as texture points.
+
+		// Get the tessellation
+		A3DStatus Status = A3DRiRepresentationItemComputeTessellation(Representation, TessellationParameters.GetPtr());
+		TechSoft::TUniqueObject<A3DRiRepresentationItemData> RepresentationItemData(Representation);
+		if (!RepresentationItemData.IsValid())
+		{
+			return false;
+		}
+
+		A3DEEntityType Type;
+		A3DEntityGetType(RepresentationItemData->m_pTessBase, &Type);
+
+		return Type == kA3DTypeTess3D ? true : false;
+	}
+
+	bool TechSoftUtilities::SewBReps(const TArray<A3DRiBrepModel*>& BRepsIn, double Tolerance, TArray<A3DRiBrepModel*>& BRepsOut)
+	{
+		if (!FTechSoftLibrary::IsInitialized())
+		{
+			return false;
+		}
+
+		TechSoft::TUniqueObject<A3DSewOptionsData> SewData;
+		SewData->m_bComputePreferredOpenShellOrientation = false;
+
+		A3DUns32 NewBRepCount = 0;
+		A3DRiBrepModel** NewBReps = nullptr;
+		A3DRiBrepModel** BRepsToSew = const_cast<A3DRiBrepModel**>(BRepsIn.GetData());
+
+		A3DStatus Status = A3DSewBrep(&BRepsToSew, BRepsIn.Num(), Tolerance, SewData.GetPtr(), &NewBReps, &NewBRepCount);
+
+		if (Status == A3DStatus::A3D_SUCCESS && NewBRepCount > 0)
+		{
+			BRepsOut.Append(NewBReps, NewBRepCount);
+		}
+
+		return Status == A3DStatus::A3D_SUCCESS ? true : false;
+	}
+}
+#else
+namespace UE::CADKernel
+{
+	bool FTechSoftUtilities::Save(A3DRiRepresentationItem* Representation, const FString& FilePath)
+	{
+		return false;
+	}
+
+	bool FTechSoftUtilities::Tessellate(A3DRiRepresentationItem* Representation, const FTessellationContext& Settings, FMeshDescription& Mesh)
+	{
+		return false;
+	}
+
+	bool FTechSoftUtilities::Tessellate(A3DRiRepresentationItem* Representation, const FTessellationContext& Settings, UE::Geometry::FDynamicMesh3& Mesh)
+	{
+		return false;
+	}
+}
+#endif
diff --git a/Engine/Source/Runtime/Datasmith/CADKernel/Engine/Private/TechSoft/TechSoftUtilities.h b/Engine/Source/Runtime/Datasmith/CADKernel/Engine/Private/TechSoft/TechSoftUtilities.h
new file mode 100644
--- /dev/null
+++ b/Engine/Source/Runtime/Datasmith/CADKernel/Engine/Private/TechSoft/TechSoftUtilities.h
@@ -0,0 +1,36 @@
+// Copyright Epic Games, Inc. All Rights Reserved.
+#pragma once
+
+#include "CoreMinimal.h"
+
+#include "TechSoftIncludes.h"
+
+#ifdef WITH_HOOPS
+
+struct FCADKernelTessellationSettings;
+
+namespace UE::CADKernel
+{
+	namespace MeshUtilities
+	{
+		class FMeshWrapperAbstract;
+	}
+
+	struct FTessellationContext;
+
+	namespace TechSoftUtilities
+	{
+		bool Tessellate(A3DRiRepresentationItem* Representation, const FTessellationContext& Settings, MeshUtilities::FMeshWrapperAbstract& MeshInterface);
+		bool TessellateRepresentation(A3DRiRepresentationItem* Representation, const FCADKernelTessellationSettings& Settings);
+		bool AddRepresentation(A3DRiRepresentationItem* Representation, MeshUtilities::FMeshWrapperAbstract& MeshInterface);
+
+		/*
+		 ** Sew the in coming BReps together
+		 * @param: BRepsIn
+		 * @param: Tolerance is file's unit with which to operate the sewing
+		 * @param: BRepsOut
+		 */
+		bool SewBReps(const TArray<A3DRiBrepModel*>& BRepsIn, double Tolerance, TArray<A3DRiBrepModel*>& BRepsOut);
+	};
+}
+#endif
diff --git a/Engine/Source/Runtime/Datasmith/CADKernel/Engine/Public/CADKernelEngine.h b/Engine/Source/Runtime/Datasmith/CADKernel/Engine/Public/CADKernelEngine.h
new file mode 100644
--- /dev/null
+++ b/Engine/Source/Runtime/Datasmith/CADKernel/Engine/Public/CADKernelEngine.h
@@ -0,0 +1,135 @@
+// Copyright Epic Games, Inc. All Rights Reserved.
+#pragma once
+
+#include "CoreMinimal.h"
+
+#include "ParametricSurfaceData.h"
+#include "TechSoftIncludes.h"
+
+namespace UE::CADKernel
+{
+	class FModel;
+	class FModelMesh;
+	class FTopologicalFace;
+}
+
+namespace UE::Geometry
+{
+	class FDynamicMesh3;
+}
+
+namespace UE::CADKernel::MeshUtilities
+{
+	class FMeshWrapperAbstract;
+}
+
+struct FMeshDescription;
+
+namespace UE::CADKernel
+{
+	struct FMeshExtractionContext
+	{
+		FCADKernelModelParameters ModelParams;
+		FCADKernelMeshParameters MeshParams;
+		TSet<int32> FacesToExtract;
+	};
+
+	struct FTessellationContext : public FMeshExtractionContext
+	{
+		FCADKernelTessellationSettings TessellationSettings;
+	};
+
+	class FCADKernelUtilities
+	{
+	public:
+		static bool CADKERNELENGINE_API Save(TSharedPtr<UE::CADKernel::FModel>& Model, const FString& FilePath);
+		static bool CADKERNELENGINE_API Tessellate(TSharedPtr<UE::CADKernel::FModel>& Model, const FTessellationContext& Context, FMeshDescription& MeshOut);
+		static bool CADKERNELENGINE_API Tessellate(TSharedPtr<UE::CADKernel::FModel>& Model, const FTessellationContext& Context, UE::Geometry::FDynamicMesh3& MeshOut);
+
+		static bool CADKERNELENGINE_API Tessellate(const UParametricSurfaceData& SurfaceData, UE::Geometry::FDynamicMesh3& MeshOut);
+		static bool CADKERNELENGINE_API Tessellate(const UParametricSurfaceData& SurfaceData, const FCADKernelRetessellationSettings& Settings, UE::Geometry::FDynamicMesh3& MeshOut);
+
+		static bool CADKERNELENGINE_API Tessellate(const UParametricSurfaceData& SurfaceData, FMeshDescription& MeshOut);
+		static bool CADKERNELENGINE_API Tessellate(const UParametricSurfaceData& SurfaceData, const FCADKernelRetessellationSettings& Settings, FMeshDescription& MeshOut);
+	};
+
+	class FTechSoftLibrary
+	{
+	public:
+		static const TCHAR* GetVersion();
+		static bool Initialize();
+		static bool IsInitialized() { return bIsInitialized; }
+
+	private:
+		static bool bIsInitialized;
+	};
+	class FTechSoftUtilities
+	{
+	public:
+		static CADKERNELENGINE_API bool Save(A3DRiRepresentationItem* Representation, const FString& FilePath);
+		static CADKERNELENGINE_API bool Tessellate(A3DRiRepresentationItem* Representation, const FTessellationContext& Context, FMeshDescription& MeshOut);
+		static CADKERNELENGINE_API bool Tessellate(A3DRiRepresentationItem* Representation, const FTessellationContext& Context, UE::Geometry::FDynamicMesh3& MeshOut);
+
+		static CADKERNELENGINE_API TSharedPtr<UE::CADKernel::FModel> TechSoftToCADKernel(A3DRiRepresentationItem* Representation);
+		/* Limitation: Non of the implicit geometry is supported */
+		static CADKERNELENGINE_API A3DRiRepresentationItem* CADKernelToTechSoft(TSharedPtr<UE::CADKernel::FModel>& Model);
+
+#ifdef WITH_HOOPS
+		static CADKERNELENGINE_API A3DRiRepresentationItem* CreateRIBRep(A3DTopoShell* TopoShellPtr);
+
+		static CADKERNELENGINE_API A3DTopoEdge* CreateTopoEdge();
+		static CADKERNELENGINE_API A3DTopoFace* CreateTopoFaceWithNaturalLoop(A3DSurfBase* CarrierSurface);
+		static CADKERNELENGINE_API A3DTopoShell* CreateTopoShell(const A3DTopoShellData& TopoShellData);
+		static CADKERNELENGINE_API A3DTopoFace* CreateTopoFace(const A3DTopoFaceData& TopoFaceData);
+		static CADKERNELENGINE_API A3DTopoLoop* CreateTopoLoop(const A3DTopoLoopData& TopoLoopData);
+		static CADKERNELENGINE_API A3DTopoEdge* CreateTopoEdge(const A3DTopoEdgeData& TopoEdgeData);
+		static CADKERNELENGINE_API A3DTopoCoEdge* CreateTopoCoEdge(const A3DTopoCoEdgeData& TopoCoEdgeData);
+		static CADKERNELENGINE_API A3DStatus LinkCoEdges(A3DTopoCoEdge* CoEdgePtr, A3DTopoCoEdge* NeighbourCoEdgePtr);
+
+		static CADKERNELENGINE_API A3DCrvNurbs* CreateTrimNurbsCurve(A3DCrvNurbs* CurveNurbsPtr, double UMin, double UMax, bool bIs2D);
+		static CADKERNELENGINE_API A3DSurfNurbs* CreateSurfaceNurbs(const A3DSurfNurbsData& SurfaceNurbsData);
+		static CADKERNELENGINE_API A3DCrvNurbs* CreateCurveNurbs(const A3DCrvNurbsData& CurveNurbsData);
+
+		static CADKERNELENGINE_API A3DStatus SetEntityGraphicsColor(A3DEntity* InEntity, FColor Color);
+#endif
+	};
+
+	namespace MathUtils
+	{
+		/*
+		** Converts a transform from a given coordinates system to UE's one
+		** UE's coordinates system is ZUp_LeftHanded
+		*/
+		FORCEINLINE FTransform ConvertTransform(ECADKernelModelCoordSystem SourceCoordSystem, const FTransform& LocalTransform)
+		{
+			const FTransform RightHanded(FRotator(0.0f, 0.0f, 0.0f), FVector(0.0f, 0.0f, 0.0f), FVector(-1.0f, 1.0f, 1.0f));
+			const FTransform RightHandedLegacy(FRotator(0.0f, 0.0f, 0.0f), FVector(0.0f, 0.0f, 0.0f), FVector(1.0f, -1.0f, 1.0f));
+			const FTransform YUpMatrix(FMatrix(FVector(0.0f, 0.0f, 1.0f), FVector(1.0f, 0.0f, 0.0f), FVector(0.0f, 1.0f, 0.0f), FVector(0.0f, 0.0f, 0.0f)));
+			const FTransform YUpMatrixInv(YUpMatrix.Inverse());
+
+			switch (SourceCoordSystem)
+			{
+				case ECADKernelModelCoordSystem::ZUp_RightHanded:
+				{
+					return RightHanded * LocalTransform * RightHanded;
+				}
+				case ECADKernelModelCoordSystem::YUp_LeftHanded:
+				{
+					return YUpMatrix * LocalTransform * YUpMatrixInv;
+				}
+				case ECADKernelModelCoordSystem::YUp_RightHanded:
+				{
+					return RightHanded * YUpMatrix * LocalTransform * YUpMatrixInv * RightHanded;
+				}
+				case ECADKernelModelCoordSystem::ZUp_RightHanded_FBXLegacy:
+				{
+					return RightHandedLegacy * LocalTransform * RightHandedLegacy;
+				}
+				default:
+				{
+					return LocalTransform;
+				}
+			}
+		}
+	}
+}
diff --git a/Engine/Source/Runtime/Datasmith/CADKernel/Engine/Public/ParametricSurfaceData.h b/Engine/Source/Runtime/Datasmith/CADKernel/Engine/Public/ParametricSurfaceData.h
new file mode 100644
--- /dev/null
+++ b/Engine/Source/Runtime/Datasmith/CADKernel/Engine/Public/ParametricSurfaceData.h
@@ -0,0 +1,295 @@
+// Copyright Epic Games, Inc. All Rights Reserved.
+
+#pragma once
+
+#include "CoreMinimal.h"
+
+#include "ParametricSurfaceData.generated.h"
+
+namespace UE::CADKernel
+{
+	class FModel;
+}
+
+class UStaticMesh;
+struct FMeshDescription;
+
+typedef void A3DRiRepresentationItem;
+
+struct FUnitConverter
+{
+	static const double CentimeterToMillimeter;
+	static const double MillimeterToCentimeter;
+	static const double CentimeterToMeter;
+	static const double MillimeterToMeter;
+};
+
+UENUM()
+enum class ECADKernelModelCoordSystem : uint8
+{
+	ZUp_LeftHanded,
+	ZUp_RightHanded,
+	YUp_LeftHanded,
+	YUp_RightHanded,
+	ZUp_RightHanded_FBXLegacy,
+};
+
+UENUM()
+enum class ECADKernelStitchingTechnique : uint8
+{
+	StitchingNone = 0,
+	StitchingHeal,
+	StitchingSew,
+};
+
+USTRUCT(BlueprintType)
+struct CADKERNELENGINE_API FCADKernelModelParameters
+{
+	GENERATED_BODY()
+
+	// value from FDatasmithUtils::ECADKernelModelCoordSystem
+	UPROPERTY()
+	uint8 ModelCoordSys = uint8(ECADKernelModelCoordSystem::ZUp_LeftHanded);
+
+	UPROPERTY()
+	float ModelUnitToCentimeter = 1.; // Model is in centimeters
+
+	UPROPERTY()
+	float ScaleFactor = 1.0f;
+	
+	bool FromJson(const FString& JsonString) const { return false; }
+	FString ToJson() const { return {}; }
+};
+
+USTRUCT(BlueprintType)
+struct CADKERNELENGINE_API FCADKernelMeshParameters
+{
+	GENERATED_BODY()
+
+	UPROPERTY()
+	bool bNeedSwapOrientation = false;
+
+	UPROPERTY()
+	bool bIsSymmetric = false;
+
+	UPROPERTY()
+	FVector SymmetricOrigin = FVector::ZeroVector;
+
+	UPROPERTY()
+	FVector SymmetricNormal = FVector::ZeroVector;
+
+	UPROPERTY()
+	float ScaleFactor = 1.0f;
+
+	bool FromJson(const FString& JsonString) const { return false; }
+	FString ToJson() const { return {}; }
+};
+
+USTRUCT(BlueprintType)
+struct CADKERNELENGINE_API FCADKernelTessellationSettings
+{
+	GENERATED_BODY()
+
+	FCADKernelTessellationSettings(
+			double InChordTolerance = 0.2,
+			double InMaxEdgeLength = 0.0,
+			double InNormalTolerance = 20.0,
+			ECADKernelStitchingTechnique InStitchingTechnique = ECADKernelStitchingTechnique::StitchingSew,
+			double InGeometricTolerance = 0.001,
+			double InStichingTolerance = 0.001,
+			bool bInUseCADKernel = true
+		)
+		: ChordTolerance(InChordTolerance)
+		, MaxEdgeLength(InMaxEdgeLength)
+		, NormalTolerance(InNormalTolerance)
+		, StitchingTechnique(InStitchingTechnique)
+		, GeometricTolerance(InGeometricTolerance)
+		, StitchingTolerance(InStichingTolerance)
+		, bUseCADKernel(bInUseCADKernel)
+	{
+	}
+
+protected:
+	/**
+	 * Maximum distance between any point on a triangle generated by the tessellation process and the actual surface.
+	 * The lower the value the more triangles.
+	 * Default value is 0.2, minimal value is 0.005 cm.
+	 */
+	UPROPERTY(config, EditAnywhere, BlueprintReadWrite, Category = "Geometry & Tessellation Options", meta = (Units = cm, ToolTip = "Maximum distance between any generated triangle and the original surface. Smaller values make more triangles.", ClampMin = "0.005"))
+	double ChordTolerance;
+
+	/**
+	 * Maximum length of edges of triangles generated by the tessellation process.
+	 * The length is in scene/model unit. The smaller the more triangles are generated.
+	 * Value of 0 means no constraint on length of edges
+	 * Default value is 0 to disable this criteria, and 1. cm is its minimal value if enable.
+	 */
+	UPROPERTY(config, EditAnywhere, BlueprintReadWrite, AdvancedDisplay, Category = "Geometry & Tessellation Options", meta = (Units = cm, DisplayName = "Max Edge Length", ToolTip = "Maximum length of any edge in the generated triangles. Smaller values make more triangles.", ClampMin = "0.0"))
+	double MaxEdgeLength;
+
+public:
+	/**
+	 * Maximum angle between the normal of two triangles generated by the tessellation process.
+	 * The angle is expressed in degree. The smaller the more triangles are generated.
+	 * Default value is 20 degrees, min value is 5 degrees.
+	 */
+	UPROPERTY(config, EditAnywhere, BlueprintReadWrite, AdvancedDisplay, Category = "Geometry & Tessellation Options", meta = (Units = deg, ToolTip = "Maximum angle between adjacent triangles. Smaller values make more triangles.", ClampMin = "5.0", ClampMax = "90.0"))
+	double NormalTolerance;
+
+
+	/**
+	 * Stitching technique applied on neighbouring surfaces before tessellation.
+	 * None : No stitching applied. This is the default.
+	 * Sewing : Connects surfaces which physically share a boundary but not topologically within a set of objects.
+	 *          This technique can modify the structure of the model by removing and adding objects.
+	 * Healing : Connects surfaces which physically share a boundary but not topologically within an object.
+	 */
+	UPROPERTY(config, EditAnywhere, BlueprintReadWrite, AdvancedDisplay, Category = "Geometry & Tessellation Options", meta = (ToolTip = "Stitching technique applied on model before tessellation. Sewing could impact number of objects."))
+	ECADKernelStitchingTechnique StitchingTechnique;
+
+protected:
+	/**
+	 * Tolerance used to determine if a surface should be tessellate or not.
+	 * Any surface which is narrower than the geometric tolerance
+	 * in one of the iso direction will not be tessellated
+	 * Value is in centimeter
+	 */
+	UPROPERTY(config, EditAnywhere, BlueprintReadWrite, AdvancedDisplay, Category = "Geometry & Tessellation Options", meta = (Units = cm, ToolTip = "Tolerance used to determine if a surface should be tessellated or not."))
+	double GeometricTolerance = 0.001;
+
+	/**
+	 * Tolerance used to determine if two surfaces should be stitched.
+	 * Value is in centimeter
+	 */
+	UPROPERTY(config, EditAnywhere, BlueprintReadWrite, AdvancedDisplay, Category = "Geometry & Tessellation Options", meta = (Units = cm, ToolTip = "Tolerance used to determine if two surfaces should be stitched.", editCondition = "StitchingTechnique!=ECADKernelStitchingTechnique::StitchingNone"))
+	double StitchingTolerance = 0.001;
+
+public:
+	/**
+	 * Indicates which tessellator to use CADKernel or TechSoft. By default it is CADKernel.
+	 */
+	UPROPERTY(config, EditAnywhere, BlueprintReadWrite, AdvancedDisplay, Category = "Geometry & Tessellation Options", meta = (DisplayName = "Use CADKernel", Tooltip = "When true uses CADKernel tessellator. Otherwise, TechSoft's tessellator is used. Default is true.", editCondition = "bWithTechSoft == true"))
+	bool bUseCADKernel = true;
+	
+	/*
+	 ** Value to convert imported geometrical values to cm which is UE's unit.
+	 ** Default value is 0.1 as CADKernel's and most CAD formats' unit is millimeter.
+	 */
+	// #cadkernel_check: Verify that this new logic is properly used everywhere.
+	UPROPERTY()
+	double UnitMultiplier = 0.1;
+
+	UPROPERTY(Transient)
+	bool bWithTechSoft = true;
+
+	bool operator == (const FCADKernelTessellationSettings& Other) const
+	{
+		return FMath::IsNearlyEqual(ChordTolerance, Other.ChordTolerance)
+			&& FMath::IsNearlyEqual(MaxEdgeLength, Other.MaxEdgeLength)
+			&& FMath::IsNearlyEqual(NormalTolerance, Other.NormalTolerance)
+			&& StitchingTechnique == Other.StitchingTechnique
+			&& FMath::IsNearlyEqual(GeometricTolerance, Other.GeometricTolerance)
+			&& FMath::IsNearlyEqual(StitchingTolerance, Other.StitchingTolerance);
+	}
+
+	uint32 GetHash() const
+	{
+		uint32 Hash = uint32(StitchingTechnique);
+		for (double Param : {ChordTolerance, MaxEdgeLength, NormalTolerance, GeometricTolerance, StitchingTolerance})
+		{
+			Hash = HashCombine(Hash, GetTypeHash(Param));
+		}
+		return Hash;
+	}
+	/** Helper functions to get geometrical values in the right unit, cm (native) or mm */
+	double GetGeometricTolerance(bool bInMillimeter = false) const { return bInMillimeter ? GeometricTolerance * 10. : GeometricTolerance; }
+	double GetStitchingTolerance(bool bInMillimeter = false) const { return bInMillimeter ? StitchingTolerance * 10. : StitchingTolerance; }
+	double GetChordTolerance(bool bInMillimeter = false) const { return bInMillimeter ? ChordTolerance * 10. : ChordTolerance; }
+	double GetMaxEdgeLength(bool bInMillimeter = false) const { return bInMillimeter ? MaxEdgeLength * 10. : MaxEdgeLength; }
+
+	bool FromJson(const FString& JsonString) const { return false; }
+	FString ToJson() const { return {}; }
+};
+
+UENUM()
+enum class ECADKernelRetessellationRule : uint8
+{
+	All = 0,
+	SkipDeletedSurfaces,
+};
+
+USTRUCT(BlueprintType)
+struct CADKERNELENGINE_API FCADKernelRetessellationSettings : public FCADKernelTessellationSettings
+{
+	GENERATED_BODY()
+
+	FCADKernelRetessellationSettings(double InChordTolerance = 0.2f, double InMaxEdgeLength = 0.0f, double InNormalTolerance = 20.0f, ECADKernelStitchingTechnique InStitchingTechnique = ECADKernelStitchingTechnique::StitchingSew, ECADKernelRetessellationRule InRetessellationRule = ECADKernelRetessellationRule::All)
+		: FCADKernelTessellationSettings(InChordTolerance, InMaxEdgeLength, InNormalTolerance, InStitchingTechnique)
+		, RetessellationRule(InRetessellationRule)
+	{
+	}
+
+	FCADKernelRetessellationSettings(const FCADKernelTessellationSettings& TessellationSettings)
+		: FCADKernelTessellationSettings(TessellationSettings)
+	{
+	}
+
+	UPROPERTY(config, EditAnywhere, BlueprintReadWrite, AdvancedDisplay, Category = "Retessellation Settings", meta = (ToolTip = "Regenerate deleted surfaces during retesselate or ignore them"))
+	ECADKernelRetessellationRule RetessellationRule = ECADKernelRetessellationRule::All;
+};
+
+UCLASS(meta = (DisplayName = "CADKernel Parametric Surface Data"))
+class CADKERNELENGINE_API UParametricSurfaceData : public UObject
+{
+	GENERATED_BODY()
+
+public:
+
+	virtual bool IsValid() { return CADKernelRawData.Num() > 0 || TechSoftRawData.Num() > 0; }
+
+	virtual bool SetFromFile(const TCHAR* FilePath, bool bForTechSoft = false);
+	
+	virtual void SetRawData(const TArray<uint8>& InRawData, bool bForTechSoft = false)
+	{
+		TArray<uint8>& RawData = bForTechSoft ? TechSoftRawData : CADKernelRawData;
+		RawData = InRawData;
+	}
+
+	virtual void SetRawData(TArray<uint8>&& InRawData, bool bForTechSoft = false)
+	{
+		TArray<uint8>& RawData = bForTechSoft ? TechSoftRawData : CADKernelRawData;
+		RawData = MoveTemp(InRawData);
+	}
+
+	void SetModelParameters(const FCADKernelModelParameters& InModelParameters) { ModelParameters = InModelParameters; }
+	const FCADKernelModelParameters& GetModelParameters() const { return ModelParameters; }
+	FCADKernelModelParameters& GetModelParameters() { return ModelParameters; }
+
+	void SetMeshParameters(const FCADKernelMeshParameters& InMeshParameters) { MeshParameters = InMeshParameters; }
+	const FCADKernelMeshParameters& GetMeshParameters() const { return MeshParameters; }
+	FCADKernelMeshParameters& GetMeshParameters() { return MeshParameters; }
+
+	void SetLastTessellationSettings(const FCADKernelTessellationSettings& InTessellationSettings) { LastTessellationSettings = InTessellationSettings; }
+	const FCADKernelTessellationSettings& GetLastTessellationSettings() const { return LastTessellationSettings; }
+
+	const TArray<uint8>& GetRawData(bool bForTechSoft = false) const { return bForTechSoft ? TechSoftRawData : CADKernelRawData; }
+	
+	TSharedPtr<UE::CADKernel::FModel> LoadModel();
+
+protected:
+	virtual void Serialize(FArchive& Ar) override;
+
+protected:
+	UPROPERTY()
+	FCADKernelModelParameters ModelParameters;
+
+	UPROPERTY()
+	FCADKernelMeshParameters MeshParameters;
+
+	UPROPERTY(EditAnywhere, Category = NURBS)
+	FCADKernelTessellationSettings LastTessellationSettings;
+
+	TArray<uint8> CADKernelRawData;
+	TArray<uint8> TechSoftRawData;
+};
+
diff --git a/Engine/Source/Runtime/Datasmith/CADKernel/Engine/Public/TechSoftIncludes.h b/Engine/Source/Runtime/Datasmith/CADKernel/Engine/Public/TechSoftIncludes.h
new file mode 100644
--- /dev/null
+++ b/Engine/Source/Runtime/Datasmith/CADKernel/Engine/Public/TechSoftIncludes.h
@@ -0,0 +1,26 @@
+// Copyright Epic Games, Inc. All Rights Reserved.
+#pragma once
+
+#include "CoreMinimal.h"
+
+#ifdef USE_TECHSOFT_SDK
+
+#if PLATFORM_WINDOWS
+#include "Windows/AllowWindowsPlatformTypes.h"
+#endif
+
+THIRD_PARTY_INCLUDES_START
+#pragma push_macro("TEXT")
+#pragma warning(push)
+#pragma warning(disable:4191) // unsafe sprintf
+#include "A3DSDKIncludes.h"
+#pragma warning(pop)
+#pragma pop_macro("TEXT")
+THIRD_PARTY_INCLUDES_END
+
+#if PLATFORM_WINDOWS
+#include "Windows/HideWindowsPlatformTypes.h"
+#endif
+
+#define WITH_HOOPS
+#endif
diff --git a/Engine/Source/Runtime/Datasmith/CADKernel/Engine/Public/TechSoftUniqueObject.h b/Engine/Source/Runtime/Datasmith/CADKernel/Engine/Public/TechSoftUniqueObject.h
new file mode 100644
--- /dev/null
+++ b/Engine/Source/Runtime/Datasmith/CADKernel/Engine/Public/TechSoftUniqueObject.h
@@ -0,0 +1,229 @@
+// Copyright Epic Games, Inc. All Rights Reserved.
+#pragma once
+
+#include "TechSoftIncludes.h"
+
+#ifdef WITH_HOOPS
+namespace UE::CADKernel::TechSoft
+{
+	// Single-ownership smart TeshSoft object
+	// Use this when you need to manage TechSoft object's lifetime.
+	//
+	// TechSoft give access to void pointers
+	// According to the context, the class name of the void pointer is known but the class is unknown
+	// i.e. A3DSDKTypes.h defines all type like :
+	// 	   typedef void A3DEntity;		
+	// 	   typedef void A3DAsmModelFile; ...
+	// 
+	// From a pointer, TechSoft give access to a copy of the associated structure :
+	//
+	// const A3DXXXXX* pPointer;
+	// A3DXXXXXData sData; // the structure
+	// A3D_INITIALIZE_DATA(A3DXXXXXData, sData); // initialization of the structure
+	// A3DXXXXXXGet(pPointer, &sData); // Copy of the data of the pointer in the structure
+	// ...
+	// A3DXXXXXXGet(NULL, &sData); // Free the structure
+	//
+	// A3D_INITIALIZE_DATA, and all A3DXXXXXXGet methods are TechSoft macro
+	//
+
+	template<class ObjectType, class IndexerType>
+	class TUniqueObjectBase
+	{
+	public:
+
+		/**
+		 * Constructor of an initialized ObjectType object
+		 */
+		TUniqueObjectBase()
+			: bDataFromTechSoft(false)
+			, Status(A3DStatus::A3D_SUCCESS)
+		{
+			InitializeData();
+		}
+
+		/**
+		 * Constructor of an filled ObjectType object with the data of DataPtr
+		 * @param DataPtr: the pointer of the data to copy
+		 */
+		explicit TUniqueObjectBase(IndexerType DataPtr)
+			: bDataFromTechSoft(false)
+		{
+			InitializeData();
+			FillFrom(DataPtr);
+		}
+
+		~TUniqueObjectBase()
+		{
+			ResetData();
+		}
+
+		/**
+		 * Fill the structure with the data of a new DataPtr
+		 */
+		A3DStatus FillFrom(IndexerType EntityPtr)
+		{
+			ResetData();
+
+			if (EntityPtr == GetDefaultIndexerValue())
+			{
+				Status = A3DStatus::A3D_ERROR;
+			}
+			else
+			{
+				Status = GetData(EntityPtr);
+				if (Status == A3DStatus::A3D_SUCCESS)
+				{
+					bDataFromTechSoft = true;
+				}
+			}
+			return Status;
+		}
+
+		template<typename... InArgTypes>
+		A3DStatus FillWith(A3DStatus(*Getter)(const A3DEntity*, ObjectType*, InArgTypes&&...), const A3DEntity* EntityPtr, InArgTypes&&... Args)
+		{
+			ResetData();
+
+			if (EntityPtr == GetDefaultIndexerValue())
+			{
+				Status = A3DStatus::A3D_ERROR;
+			}
+			else
+			{
+				Status = Getter(EntityPtr, &Data, Forward<InArgTypes>(Args)...);
+				if (Status == A3DStatus::A3D_SUCCESS)
+				{
+					bDataFromTechSoft = true;
+				}
+			}
+
+			return Status;
+		}
+
+		/**
+		 * Empty the structure
+		 */
+		void Reset()
+		{
+			ResetData();
+		}
+
+		/**
+		 * Return
+		 *  - A3DStatus::A3D_SUCCESS if the data is filled
+		 *  - A3DStatus::A3D_ERROR if the data is empty
+		 */
+		A3DStatus GetStatus()
+		{
+			return Status;
+		}
+
+		/**
+		 * Return true if the data is filled
+		 */
+		const bool IsValid() const
+		{
+			return Status == A3DStatus::A3D_SUCCESS;
+		}
+
+		// Non-copyable
+		TUniqueObjectBase(const TUniqueObjectBase&) = delete;
+		TUniqueObjectBase& operator=(const TUniqueObjectBase&) = delete;
+
+		// Conversion methods
+
+		const ObjectType& operator*() const
+		{
+			return Data;
+		}
+
+		ObjectType& operator*()
+		{
+			check(IsValid());
+			return Data;
+		}
+
+		const ObjectType* operator->() const
+		{
+			check(IsValid());
+			return &Data;
+		}
+
+		ObjectType* operator->()
+		{
+			check(IsValid());
+			return &Data;
+		}
+
+		/**
+		 * Return the structure pointer
+		 */
+		ObjectType* GetPtr()
+		{
+			if (Status != A3DStatus::A3D_SUCCESS)
+			{
+				return nullptr;
+			}
+			return &Data;
+		}
+
+	private:
+		ObjectType Data;
+		bool bDataFromTechSoft = false;
+		A3DStatus Status = A3DStatus::A3D_ERROR;
+
+		/**
+		 * DefaultValue is used to initialize "Data" with GetData method
+		 * According to IndexerType, the value is either nullptr for const A3DEntity* either something like "A3D_DEFAULT_MATERIAL_INDEX" ((A3DUns16)-1) for uint32
+		 * @see ResetData
+		 */
+		CADKERNELENGINE_API IndexerType GetDefaultIndexerValue() const
+		{
+			return (IndexerType)INDEX_NONE;
+		}
+
+		CADKERNELENGINE_API void InitializeData()
+		{
+		}
+
+		CADKERNELENGINE_API A3DStatus GetData(IndexerType AsmModelFilePtr)
+		{
+			return A3DStatus::A3D_ERROR;
+		}
+
+		void ResetData()
+		{
+			if (bDataFromTechSoft)
+			{
+				GetData(GetDefaultIndexerValue());
+			}
+			else
+			{
+				InitializeData();
+			}
+			Status = A3DStatus::A3D_SUCCESS;
+			bDataFromTechSoft = false;
+		}
+	};
+
+	template<class ObjectType>
+	using TUniqueObject = TUniqueObjectBase<ObjectType, const A3DEntity*>;
+
+	template<class ObjectType>
+	using TUniqueObjectFromIndex = TUniqueObjectBase<ObjectType, uint32>;
+
+	class CADKERNELENGINE_API FTechSoftDefaultValue
+	{
+	public:
+		static const uint32 Material;
+		static const uint32 Picture;
+		static const uint32 RgbColor;
+		static const uint32 Style;
+		static const uint32 TextureApplication;
+		static const uint32 TextureDefinition;
+	};
+
+}
+#endif
+
