Fixed build break on Mac - Issue 702461

#rnx

diff --git a/Engine/Source/Runtime/Datasmith/CADKernel/Engine/Private/DynamicMeshWrapper.h b/Engine/Source/Runtime/Datasmith/CADKernel/Engine/Private/DynamicMeshWrapper.h
--- a/Engine/Source/Runtime/Datasmith/CADKernel/Engine/Private/DynamicMeshWrapper.h
+++ b/Engine/Source/Runtime/Datasmith/CADKernel/Engine/Private/DynamicMeshWrapper.h
@@ -69,6 +69,12 @@ namespace UE::CADKernel::MeshUtilities
 			uint32 MaterialID;
 			FArray3i NormalIndices;
 			FArray3i TexCoordIndices;
+			FTriangleData(int32 InGroupID, uint32 InMaterialID, const FArray3i& InNormalIndices, const FArray3i& InTexCoordIndices)
+				: GroupID(InGroupID)
+				, MaterialID(InMaterialID)
+				, NormalIndices(InNormalIndices)
+				, TexCoordIndices(InTexCoordIndices)
+			{}
 		};
 		TArray<FTriangleData> TriangleDataSet;
 		TSet<int32> GroupIDSet;
