More fixes to Mac build - Issue 702461

#rnx

diff --git a/Engine/Source/Runtime/Datasmith/CADKernel/Engine/CADKernelEngine.Build.cs b/Engine/Source/Runtime/Datasmith/CADKernel/Engine/CADKernelEngine.Build.cs
--- a/Engine/Source/Runtime/Datasmith/CADKernel/Engine/CADKernelEngine.Build.cs
+++ b/Engine/Source/Runtime/Datasmith/CADKernel/Engine/CADKernelEngine.Build.cs
@@ -9,7 +9,7 @@ namespace UnrealBuildTool.Rules
 			: base(Target)
 		{
 			bUseUnity = false; // TechSoftLibrary.cpp cannot be compiled with the other cpp files.
-			OptimizeCode = CodeOptimization.Never;
+			//OptimizeCode = CodeOptimization.Never;
 
 			PublicDependencyModuleNames.AddRange(
 				new string[]
diff --git a/Engine/Source/Runtime/Datasmith/CADKernel/Engine/Private/MeshDescriptionWrapper.h b/Engine/Source/Runtime/Datasmith/CADKernel/Engine/Private/MeshDescriptionWrapper.h
--- a/Engine/Source/Runtime/Datasmith/CADKernel/Engine/Private/MeshDescriptionWrapper.h
+++ b/Engine/Source/Runtime/Datasmith/CADKernel/Engine/Private/MeshDescriptionWrapper.h
@@ -75,8 +75,7 @@ namespace UE::CADKernel::MeshUtilities
 				VertexType ABC[3] = { INDEX_NONE, INDEX_NONE, INDEX_NONE };
 			};
 
-			TMeshDescEntity() = default;
-
+			TMeshDescEntity() {}
 			TMeshDescEntity(VertexType ValA, VertexType ValB, VertexType ValC)
 				: A(ValA), B(ValB), C(ValC)
 			{}
