Fixed false  positive warning from static analysis - Issue 702482

#rnx

diff --git a/Engine/Source/Runtime/Datasmith/CADKernel/Engine/Private/CADKernelEnginePrivate.cpp b/Engine/Source/Runtime/Datasmith/CADKernel/Engine/Private/CADKernelEnginePrivate.cpp
--- a/Engine/Source/Runtime/Datasmith/CADKernel/Engine/Private/CADKernelEnginePrivate.cpp
+++ b/Engine/Source/Runtime/Datasmith/CADKernel/Engine/Private/CADKernelEnginePrivate.cpp
@@ -28,93 +28,6 @@ namespace UE::CADKernel::Private
 	using namespace UE::CADKernel::MeshUtilities;
 	using namespace UE::Geometry;
 
-	void CheckPatchIds(FModelMesh& ModelMesh)
-	{
-		TSet<int32> PatchIds;
-
-		for (FFaceMesh* FaceMesh : ModelMesh.GetFaceMeshes())
-		{
-			FTopologicalFace& Face = static_cast<FTopologicalFace&>(FaceMesh->GetGeometricEntity());
-			PatchIds.Add(Face.GetPatchId());
-		}
-
-		printf(">>> Found %d", PatchIds.Num());
-	}
-
-	void FillPolygroupLayer(FDynamicMesh3& Mesh, TMap<int32, int32>& OutMapping, const FModelMesh& ModelMesh)
-	{
-		if (!ensure(Mesh.HasTriangleGroups()))
-		{
-			return;
-		}
-
-		static TSet<int32> LastPatchIdSet;
-		TSet<int32> PatchIdSet;
-		int32 TriangleCount = 0;
-
-		for (FFaceMesh* FaceMesh : ModelMesh.GetFaceMeshes())
-		{
-			FTopologicalFace& Face = static_cast<FTopologicalFace&>(FaceMesh->GetGeometricEntity());
-			TriangleCount += FaceMesh->TrianglesVerticesIndex.Num() / 3;
-		}
-
-		if (LastPatchIdSet.Num() > 0)
-		{
-			for (int32 PatchId : PatchIdSet)
-			{
-				ensure(LastPatchIdSet.Contains(PatchId));
-			}
-		}
-
-		LastPatchIdSet = PatchIdSet;
-
-		const int32 Mesh3TriangleCount = Mesh.TriangleCount();
-		ensure(Mesh3TriangleCount == TriangleCount);
-
-		TSet<int32> Missings;
-		for (int TriID = 0; TriID < Mesh.TriangleCount(); ++TriID)
-		{
-			if (!PatchIdSet.Contains(Mesh.GetTriangleGroup(TriID)))
-			{
-				Missings.Add(Mesh.GetTriangleGroup(TriID));
-			}
-		}
-
-		ensure(Missings.Num() == 0);
-
-		Mesh.EnableAttributes();
-		Mesh.Attributes()->SetNumPolygroupLayers(PatchIdSet.Num());
-
-		TMap<int32, FDynamicMeshPolygroupAttribute*> LayerMapping;
-		LayerMapping.Reserve(PatchIdSet.Num());
-
-		int32 Index = -1;
-		for (int32 PatchId : PatchIdSet)
-		{
-			FDynamicMeshPolygroupAttribute* Layer = Mesh.Attributes()->GetPolygroupLayer(++Index);
-			Layer->SetName(*FString::Printf(TEXT("Face #%d"), PatchId));
-
-			LayerMapping.Add(PatchId, Layer);
-			OutMapping.Add(PatchId, Index);
-
-			for (int TriID = 0; TriID < Mesh.TriangleCount(); ++TriID)
-			{
-				Layer->SetValue(TriID, -1);
-			}
-		}
-
-		for (int TriID = 0; TriID < Mesh.TriangleCount(); ++TriID)
-		{
-			const int32 PatchId = Mesh.GetTriangleGroup(TriID);
-
-			FDynamicMeshPolygroupAttribute** LayerPtr = LayerMapping.Find(PatchId);
-			if (ensure(LayerPtr))
-			{
-				(*LayerPtr)->SetValue(TriID, PatchId);
-			}
-		}
-	}
-
 	void DefineMeshCriteria(FModelMesh& MeshModel, const FCADKernelTessellationSettings& TessellationSettings)
 	{
 		using namespace UE::CADKernel;
