[In Progress] CADKernel Engine module
- Refined API to tessellate CADKernel model
- Implemented tessellation of CADKernel model to FMeshDescription
- Added method to transfer data from UDatasmithParametricSurfaceData child classes to UParametricSurfaceData
- Tested retessellation

#rnx

diff --git a/Engine/Source/Runtime/Datasmith/CADKernel/Base/Public/Mesh/Structure/ModelMesh.h b/Engine/Source/Runtime/Datasmith/CADKernel/Base/Public/Mesh/Structure/ModelMesh.h
--- a/Engine/Source/Runtime/Datasmith/CADKernel/Base/Public/Mesh/Structure/ModelMesh.h
+++ b/Engine/Source/Runtime/Datasmith/CADKernel/Base/Public/Mesh/Structure/ModelMesh.h
@@ -44,6 +44,8 @@ class CADKERNEL_API FModelMesh : public FEntityGeom
 
 public:
 
+	const TArray<TArray<FVector>*>& GetGlobalPointCloud() const { return GlobalPointCloud; };
+
 	int32 GetFaceCount() const
 	{
 		return FaceMeshes.Num();
diff --git a/Engine/Source/Runtime/Datasmith/CADKernel/Engine/CADKernelEngine.Build.cs b/Engine/Source/Runtime/Datasmith/CADKernel/Engine/CADKernelEngine.Build.cs
--- a/Engine/Source/Runtime/Datasmith/CADKernel/Engine/CADKernelEngine.Build.cs
+++ b/Engine/Source/Runtime/Datasmith/CADKernel/Engine/CADKernelEngine.Build.cs
@@ -9,7 +9,7 @@ namespace UnrealBuildTool.Rules
 			: base(Target)
 		{
 			bUseUnity = false; // TechSoftLibrary.cpp cannot be compiled with the other cpp files.
-			//OptimizeCode = CodeOptimization.Never;
+			OptimizeCode = CodeOptimization.Never;
 
 			PublicDependencyModuleNames.AddRange(
 				new string[]
diff --git a/Engine/Source/Runtime/Datasmith/CADKernel/Engine/Private/CADKernelEngine.cpp b/Engine/Source/Runtime/Datasmith/CADKernel/Engine/Private/CADKernelEngine.cpp
--- a/Engine/Source/Runtime/Datasmith/CADKernel/Engine/Private/CADKernelEngine.cpp
+++ b/Engine/Source/Runtime/Datasmith/CADKernel/Engine/Private/CADKernelEngine.cpp
@@ -2,6 +2,9 @@
 
 #include "CADKernelEngine.h"
 
+#include "CADKernelEnginePrivate.h"
+#include "MeshUtilities.h"
+
 #include "Topo/Body.h"
 #include "Topo/Model.h"
 #include "Topo/Shell.h"
@@ -15,6 +18,8 @@ DEFINE_LOG_CATEGORY(LogCADKernelEngine);
 
 namespace UE::CADKernel
 {
+	using namespace UE::CADKernel::MeshUtilities;
+
 	static bool GUseEngine = false;
 	FAutoConsoleVariableRef GCADKernelDebugUseEngine(
 		TEXT("CADKernel.Debug.UseEngine"),
@@ -22,80 +27,32 @@ namespace UE::CADKernel
 		TEXT(""),
 		ECVF_Default);
 
-	bool FCADKernelUtilities::Save(TSharedPtr<UE::CADKernel::FModel>& Model, const FString& FilePath)
-	{
-		return false;
-	}
-
-	bool FCADKernelUtilities::Tessellate(TSharedPtr<UE::CADKernel::FModel>& Model, const FTessellationContext& Context, FMeshDescription& Mesh)
-	{
-		return false;
-	}
-
-	bool FCADKernelUtilities::Tessellate(TSharedPtr<UE::CADKernel::FModel>& Model, const FTessellationContext& Context, UE::Geometry::FDynamicMesh3& Mesh)
-	{
-		return false;
-	}
-
-	bool FCADKernelUtilities::Tessellate(const UParametricSurfaceData& SurfaceData, UE::Geometry::FDynamicMesh3& Mesh)
+	FTessellationContext::FTessellationContext(const FCADKernelModelParameters& InModelParams, const FCADKernelMeshParameters& InMeshParams, const FCADKernelRetessellationSettings& Settings)
 	{
-#if 0
-		FMeshTranslationContext Context(SurfaceData);
-
-		if (!Context.TessellationSettings.bUseCADKernel)
-		{
-			return false;
-		}
-
-		return Private::TessellateWithCADKernel(SurfaceData.GetRawData(false), Context, Mesh);
-#else
-		return false;
-#endif
+		ModelParams = InModelParams;
+		MeshParams = InMeshParams;
+		TessellationSettings = Settings;
+		bResolveTJunctions = TessellationSettings.bResolveTJunctions;
 	}
 
-	bool FCADKernelUtilities::Tessellate(const UParametricSurfaceData& SurfaceData, const FCADKernelRetessellationSettings& Settings, UE::Geometry::FDynamicMesh3& Mesh)
+	bool FCADKernelUtilities::Save(TSharedPtr<UE::CADKernel::FModel>& Model, const FString& FilePath)
 	{
-#if 0
-		FMeshTranslationContext Context(SurfaceData);
-
-		if (!Context.TessellationSettings.bUseCADKernel)
-		{
-			return false;
-		}
-
-		return Private::TessellateWithCADKernel(SurfaceData.GetRawData(false), Context, Mesh);
-#else
 		return false;
-#endif
 	}
 
-	bool FCADKernelUtilities::Tessellate(const UParametricSurfaceData& SurfaceData, FMeshDescription& Mesh)
+	bool FCADKernelUtilities::Tessellate(UE::CADKernel::FModel& Model, const FTessellationContext& Context, FMeshDescription& Mesh, bool bEmptyMesh)
 	{
-#if 0
-		FMeshTranslationContext Context(SurfaceData);
+		using namespace UE::CADKernel::MeshUtilities;
 
-		if (!Context.TessellationSettings.bUseCADKernel)
-		{
-			return false;
-		}
-
-		return Private::TessellateWithCADKernel(SurfaceData.GetRawData(false), Context, Mesh);
-#else
-		return false;
-#endif 
+		TSharedPtr<FMeshWrapperAbstract> MeshWrapper = FMeshWrapperAbstract::MakeWrapper(Context, Mesh);
+		return Private::Tessellate(Model, Context, *MeshWrapper, bEmptyMesh);
 	}
 
-	bool FCADKernelUtilities::Tessellate(const UParametricSurfaceData& SurfaceData, const FCADKernelRetessellationSettings& Settings, FMeshDescription& Mesh)
+	bool FCADKernelUtilities::Tessellate(UE::CADKernel::FModel& Model, const FTessellationContext& Context, UE::Geometry::FDynamicMesh3& Mesh, bool bEmptyMesh)
 	{
-#if 0
-		if (!Context.TessellationSettings.bUseCADKernel)
-		{
-			return false;
-		}
+		using namespace UE::CADKernel::MeshUtilities;
 
-		return Private::TessellateWithCADKernel(SurfaceData.GetRawData(false), Context, Mesh);
-#else
-		return false;
-#endif
+		TSharedPtr<FMeshWrapperAbstract> MeshWrapper = FMeshWrapperAbstract::MakeWrapper(Context, Mesh);
+		return Private::Tessellate(Model, Context, *MeshWrapper, bEmptyMesh);
 	}
 }
\ No newline at end of file
diff --git a/Engine/Source/Runtime/Datasmith/CADKernel/Engine/Private/CADKernelEnginePrivate.cpp b/Engine/Source/Runtime/Datasmith/CADKernel/Engine/Private/CADKernelEnginePrivate.cpp
--- a/Engine/Source/Runtime/Datasmith/CADKernel/Engine/Private/CADKernelEnginePrivate.cpp
+++ b/Engine/Source/Runtime/Datasmith/CADKernel/Engine/Private/CADKernelEnginePrivate.cpp
@@ -62,29 +62,8 @@ namespace UE::CADKernel::Private
 		}
 	}
 
-	bool TessellateWithCADKernel(const TArray<uint8>& RawData, const FTessellationContext& Context, FDynamicMesh3& Mesh)
+	bool Tessellate(FModel& Model, const FTessellationContext& Context, FMeshWrapperAbstract& MeshWrapper, bool bEmptyMesh)
 	{
-		TSharedPtr<FMeshWrapperAbstract> MeshWrapper = FMeshWrapperAbstract::MakeWrapper(Context, Mesh);
-		return Tessellate(RawData, Context, *MeshWrapper);
-	}
-
-	bool TessellateWithCADKernel(const TArray<uint8>& RawData, const FTessellationContext& Context, FMeshDescription& Mesh)
-	{
-		TSharedPtr<FMeshWrapperAbstract> MeshWrapper = FMeshWrapperAbstract::MakeWrapper(Context, Mesh);
-		return Tessellate(RawData, Context, *MeshWrapper);
-	}
-	
-	bool Tessellate(const TArray<uint8>& RawData, const UE::CADKernel::FTessellationContext& Context, FMeshWrapperAbstract& MeshWrapper)
-	{
-		TSharedRef<FSession> Session = MakeShared<FSession>(Context.TessellationSettings.GetGeometricTolerance(true));
-		Session->AddDatabase(RawData);
-
-		FModel& Model = Session->GetModel();
-		if (Model.GetBodies().IsEmpty())
-		{
-			return false;
-		}
-
 		// Tessellate the model
 		TSharedRef<FModelMesh> ModelMesh = FEntity::MakeShared<FModelMesh>();
 
@@ -92,11 +71,15 @@ namespace UE::CADKernel::Private
 		FMesher Mesher(*ModelMesh, GeometriTolerance, false/*bActivateThinZoneMeshing*/);
 
 		DefineMeshCriteria(*ModelMesh, Context.TessellationSettings);
+		// #cad_import: FMesher::MeshEntity should take a const FModel&
 		Mesher.MeshEntity(Model);
 
-		// #cad_import: Extract mesh data from FModelMesh using MeshWrapper
+		if (bEmptyMesh)
+		{
+			MeshWrapper.ClearMesh();
+		}
 
-		return true;
+		return AddModelMesh(*ModelMesh, MeshWrapper);
 	}
 
 	bool GetFaceTrimmingCurves(const FModel& Model, const FTopologicalFace& Face, TArray<TArray<TArray<FVector>>>& CurvesOut)
diff --git a/Engine/Source/Runtime/Datasmith/CADKernel/Engine/Private/CADKernelEnginePrivate.h b/Engine/Source/Runtime/Datasmith/CADKernel/Engine/Private/CADKernelEnginePrivate.h
--- a/Engine/Source/Runtime/Datasmith/CADKernel/Engine/Private/CADKernelEnginePrivate.h
+++ b/Engine/Source/Runtime/Datasmith/CADKernel/Engine/Private/CADKernelEnginePrivate.h
@@ -3,6 +3,7 @@
 
 #include "CoreMinimal.h"
 
+#include "CADKernelEngine.h"
 #include "MeshUtilities.h"
 
 namespace UE::Geometry
@@ -13,9 +14,6 @@ namespace UE::Geometry
 namespace UE::CADKernel
 {
 	class FModel;
-	class FFaceMesh;
-	class FModelMesh;
-	class FTopologicalFace;
 
 	struct FTessellationContext;
 
@@ -26,14 +24,74 @@ namespace UE::CADKernel
 
 	namespace Private
 	{
-		bool Tessellate(const TArray<uint8>& RawData, const UE::CADKernel::FTessellationContext& Context, UE::CADKernel::MeshUtilities::FMeshWrapperAbstract& MeshWrapper);
-		bool TessellateWithCADKernel(const TArray<uint8>& RawData, const UE::CADKernel::FTessellationContext& Context, UE::Geometry::FDynamicMesh3& Mesh);
-		bool TessellateWithCADKernel(const TArray<uint8>& RawData, const UE::CADKernel::FTessellationContext& Context, FMeshDescription& Mesh);
+
+		template<typename MeshType>
+		bool Tessellate(UParametricSurfaceData& Data, MeshType& MeshOut)
+		{
+			using namespace UE::CADKernel;
+
+			bool bTessellationSuccessful = false;
+
+			FTessellationContext Context(Data.GetModelParameters(), Data.GetMeshParameters(), Data.GetLastTessellationSettings());
+
+			if (Context.TessellationSettings.bUseCADKernel)
+			{
+				TSharedPtr<FModel> Model = Data.GetModel();
+				if (Model.IsValid())
+				{
+					bTessellationSuccessful = FCADKernelUtilities::Tessellate(*Model, Context, MeshOut);
+				}
+			}
+			else if (A3DRiRepresentationItem* Representation = Data.GetRepresentation())
+			{
+				bTessellationSuccessful = FTechSoftUtilities::Tessellate(Representation, Context, MeshOut);
+			}
+
+			return bTessellationSuccessful;
+		}
+
+		template<typename MeshType>
+		bool Retessellate(UParametricSurfaceData& Data, const FCADKernelRetessellationSettings& Settings, MeshType& MeshOut)
+		{
+			using namespace UE::CADKernel;
+			using namespace UE::CADKernel::MeshUtilities;
+
+			bool bTessellationSuccessful = false;
+
+			FTessellationContext Context(Data.GetModelParameters(), Data.GetMeshParameters(), Settings);
+
+			if (Settings.RetessellationRule == ECADKernelRetessellationRule::SkipDeletedFaces)
+			{
+				GetExistingFaceGroups(MeshOut, Context.FaceGroupsToExtract);
+			}
+
+			if (Context.TessellationSettings.bUseCADKernel)
+			{
+				TSharedPtr<FModel> Model = Data.GetModel();
+				if (Model.IsValid())
+				{
+					bTessellationSuccessful = FCADKernelUtilities::Tessellate(*Model, Context, MeshOut, true);
+				}
+			}
+			else
+			{
+				if (A3DRiRepresentationItem* Representation = Data.GetRepresentation())
+				{
+					bTessellationSuccessful = FTechSoftUtilities::Tessellate(Representation, Context, MeshOut, true);
+				}
+			}
+
+			return bTessellationSuccessful;
+		}
+
+		bool Tessellate(FModel& Model, const UE::CADKernel::FTessellationContext& Context, UE::CADKernel::MeshUtilities::FMeshWrapperAbstract& MeshWrapper, bool bEmptyMesh = false);
 
 		// For future use
 		bool GetFaceTrimmingCurves(const UE::CADKernel::FModel& Model, const UE::CADKernel::FTopologicalFace& Face, TArray<TArray<TArray<FVector>>>& CurvesOut);
 		bool GetFaceTrimming2DPolylines(const UE::CADKernel::FModel& Model, const UE::CADKernel::FTopologicalFace& Face, TArray<TArray<FVector2d>>& PolylinesOut);
 		bool GetFaceTrimming3DPolylines(const UE::CADKernel::FModel& Model, const UE::CADKernel::FTopologicalFace& Face, TArray<TArray<FVector>>& PolylinesOut);
+
+		bool AddModelMesh(const UE::CADKernel::FModelMesh& ModelMesh, UE::CADKernel::MeshUtilities::FMeshWrapperAbstract& MeshWrapper);
 	}
 }
 
diff --git a/Engine/Source/Runtime/Datasmith/CADKernel/Engine/Private/CADKernelMesh.cpp b/Engine/Source/Runtime/Datasmith/CADKernel/Engine/Private/CADKernelMesh.cpp
new file mode 100644
--- /dev/null
+++ b/Engine/Source/Runtime/Datasmith/CADKernel/Engine/Private/CADKernelMesh.cpp
@@ -0,0 +1,121 @@
+// Copyright Epic Games, Inc. All Rights Reserved.
+
+#include "CADKernelEngine.h"
+
+#include "CADKernelEnginePrivate.h"
+#include "MeshUtilities.h"
+
+//#include "Geo/Curves/CurveUtilities.h"
+//#include "Mesh/Criteria/Criterion.h"
+//#include "Mesh/Meshers/Mesher.h"
+#include "Mesh/Structure/ModelMesh.h"
+#include "Mesh/Structure/FaceMesh.h"
+//#include "Topo/Model.h"
+//#include "Topo/TopologicalEdge.h"
+#include "Topo/TopologicalFace.h"
+#include "Topo/TopologicalFaceUtilities.h"
+//#include "Topo/TopologicalLoop.h"
+
+namespace UE::CADKernel::Private
+{
+	using namespace UE::CADKernel;
+	using namespace UE::CADKernel::MeshUtilities;
+
+	class FModelMeshConverter
+	{
+	private:
+		const FModelMesh& ModelMesh;
+		FMeshWrapperAbstract& MeshWrapper;
+
+	public:
+		FModelMeshConverter(const FModelMesh& InModelMesh, FMeshWrapperAbstract& InMeshWrapper)
+			: ModelMesh(InModelMesh)
+			, MeshWrapper(InMeshWrapper)
+		{}
+
+		uint32 GetMaterialID(const FTopologicalFace& Face)
+		{
+			const uint32 Prefix = Face.GetMaterialId() != 0 ? PointerHash(TEXT("Material")) : PointerHash(TEXT("Color"));
+			const int32 MaterialID = Face.GetMaterialId() != 0 ? Face.GetMaterialId() : Face.GetColorId();
+
+			return HashCombine(Prefix, MaterialID);
+		}
+
+		void AddTrianglesFromFaceMesh(const FFaceMesh& FaceMesh)
+		{
+			const FTopologicalFace& Face = static_cast<const FTopologicalFace&>(FaceMesh.GetGeometricEntity());
+			const int32 GroupID = Face.GetPatchId();
+
+			if (Face.IsDegenerated() || !MeshWrapper.IsFaceGroupValid(GroupID))
+			{
+				return;
+			}
+
+			const uint32 MaterialID = GetMaterialID(Face);
+
+			const TArray<int32>& TriangleVertexIndices = FaceMesh.TrianglesVerticesIndex;
+
+			TArray<FMeshWrapperAbstract::FFaceTriangle> FaceTriangles;
+			FaceTriangles.Reserve(TriangleVertexIndices.Num() / 3);
+
+			MeshWrapper.StartFaceTriangles(FaceTriangles.Num(), FaceMesh.Normals, FaceMesh.UVMap);
+
+			const int32* GlobalVtxIDs = FaceMesh.VerticesGlobalIndex.GetData();
+			const int32* TriVtxIDs = TriangleVertexIndices.GetData();
+
+			for (int32 Index = 0; Index < TriangleVertexIndices.Num(); Index += 3, TriVtxIDs += 3)
+			{
+				FArray3i Normals(TriVtxIDs[0], TriVtxIDs[1], TriVtxIDs[2]);
+				FArray3i TexCoords = Normals;
+				FArray3i VertexIndices(GlobalVtxIDs[TriVtxIDs[0]], GlobalVtxIDs[TriVtxIDs[1]], GlobalVtxIDs[TriVtxIDs[2]]);
+
+				FaceTriangles.Emplace(GroupID, MaterialID, VertexIndices, Normals, TexCoords);
+			}
+
+			MeshWrapper.AddFaceTriangles(FaceTriangles);
+
+			MeshWrapper.EndFaceTriangles();
+		}
+
+		bool Convert()
+		{
+			{
+				const TArray<TArray<FVector>*>& GlobalPointCloud = ModelMesh.GetGlobalPointCloud();
+
+				int32 VertexCount = 0;
+				for (const TArray<FVector>* PointArray : GlobalPointCloud)
+				{
+					VertexCount += PointArray->Num();
+				}
+
+				TArray<FVector> Vertices;
+				Vertices.Reserve(VertexCount);
+
+				for (const TArray<FVector>* PointArray : GlobalPointCloud)
+				{
+					Vertices.Append(*PointArray);
+				}
+
+				MeshWrapper.SetVertices(MoveTemp(Vertices));
+			}
+
+			{
+
+				for (const FFaceMesh* FaceMesh : ModelMesh.GetFaceMeshes())
+				{
+					AddTrianglesFromFaceMesh(*FaceMesh);
+				}
+			}
+
+			MeshWrapper.Complete();
+
+			return true;
+		}
+	};
+
+	bool AddModelMesh(const FModelMesh& ModelMesh, FMeshWrapperAbstract& MeshWrapper)
+	{
+		FModelMeshConverter Converter(ModelMesh, MeshWrapper);
+		return Converter.Convert();
+	}
+}
\ No newline at end of file
diff --git a/Engine/Source/Runtime/Datasmith/CADKernel/Engine/Private/DynamicMeshWrapper.cpp b/Engine/Source/Runtime/Datasmith/CADKernel/Engine/Private/DynamicMeshWrapper.cpp
--- a/Engine/Source/Runtime/Datasmith/CADKernel/Engine/Private/DynamicMeshWrapper.cpp
+++ b/Engine/Source/Runtime/Datasmith/CADKernel/Engine/Private/DynamicMeshWrapper.cpp
@@ -4,15 +4,13 @@
 
 #include "CADKernelEngine.h"
 
-#include "DynamicMesh/DynamicMesh3.h"
-
 #include "Algo/Reverse.h"
 
 namespace UE::CADKernel::MeshUtilities
 {
 	using namespace UE::CADKernel;
 	using namespace UE::Geometry;
-	
+
 	TSharedPtr<FMeshWrapperAbstract> FMeshWrapperAbstract::MakeWrapper(const FMeshExtractionContext& Context, FDynamicMesh3& Mesh)
 	{
 		return MakeShared<FDynamicMeshWrapper>(Context, Mesh);
@@ -24,21 +22,27 @@ namespace UE::CADKernel::MeshUtilities
 	{
 	}
 
-	bool FDynamicMeshWrapper::AddNewVertices(TArray<FVector>&& VertexArray)
+	void FDynamicMeshWrapper::ClearMesh()
+	{
+		// #cad_import: Clear mesh
+		MeshOut.Clear();
+	}
+
+	bool FDynamicMeshWrapper::AddNewVertices(TArray<FVector>&& InVertexArray)
 	{
 		if (bAreVerticesSet)
 		{
 			return false;
 		}
 
+		TArray<FVector> VertexArray = MoveTemp(InVertexArray);
+
 		VertexIdOffset = VertexMapping.Num();
 
 		const int32 VertexCount = VertexArray.Num();
 		VertexMapping.Reserve(VertexIdOffset + VertexCount);
 		VertIDMap.SetNumUninitialized(VertexIdOffset + VertexCount);
 
-		ConvertVectorArray(Context.ModelParams.ModelCoordSys, VertexArray);
-
 		for (const FVector& Vertex : VertexArray)
 		{
 			const int32 NewVertIdx = MeshOut.AppendVertex(Vertex);
@@ -56,7 +60,7 @@ namespace UE::CADKernel::MeshUtilities
 		VertexIdOffset = VertexMapping.Num();
 		ensure(VertexIdOffset == 0);
 
-		TArray<FVector> Vertices = MoveTemp(InVertices);;
+		TArray<FVector> Vertices = MoveTemp(InVertices);
 		AddNewVertices(MoveTemp(Vertices));
 
 		bAreVerticesSet = true;
@@ -66,8 +70,20 @@ namespace UE::CADKernel::MeshUtilities
 
 	void FDynamicMeshWrapper::FinalizeMesh()
 	{
+		TDynamicVector<FVector3d>& Positions = const_cast<TDynamicVector<FVector3d>&>(MeshOut.GetVerticesBuffer());
+		ConvertVectorArray(Context.ModelParams.ModelCoordSys, Positions);
+
+		if (!FMath::IsNearlyEqual(Context.ModelParams.ModelUnitToCentimeter, 1.f))
+		{
+			for (FVector3d& Position : Positions)
+			{
+				Position *= Context.ModelParams.ModelUnitToCentimeter;
+			}
+		}
+
+		ConvertVectorArray(Context.ModelParams.ModelCoordSys, Normals);
+
 		AddTriangles();
-		AddSymmetry();
 	}
 
 	bool FDynamicMeshWrapper::ReserveNewTriangles(int32 TriangleCount)
@@ -309,6 +325,8 @@ namespace UE::CADKernel::MeshUtilities
 
 			++TriangleDataIter;
 		}
+
+		TriangleDataSet.Empty(0);
 	}
 
 	bool FDynamicMeshWrapper::StartFaceTriangles(int32 TriangleCount, const TArray<FVector3f>& InNormals, const TArray<FVector2f>& InTexCoords)
@@ -322,6 +340,25 @@ namespace UE::CADKernel::MeshUtilities
 		return true;
 	}
 
+	bool FDynamicMeshWrapper::StartFaceTriangles(const TArrayView<FVector>& InNormals, const TArrayView<FVector2d>& InTexCoords)
+	{
+		LastNormalIndex = Normals.Num();
+
+		const int32 ArraySize = InNormals.Num();
+		ensure(ArraySize == InTexCoords.Num());
+
+		Normals.Reserve(LastNormalIndex + ArraySize);
+		TexCoords.Reserve(LastNormalIndex + ArraySize);
+
+		for (int32 Index = 0; Index < ArraySize; ++Index)
+		{
+			Normals.Emplace(InNormals[Index].X, InNormals[Index].Y, InNormals[Index].Z);
+			TexCoords.Emplace(InTexCoords[Index].X, InTexCoords[Index].Y);
+		}
+
+		return true;
+	}
+
 	int32 FDynamicMeshWrapper::AppendTriangle(FIndex3i& VertexIDs, int32 GroupID)
 	{
 		int32 NewTriangleID = MeshOut.AppendTriangle(VertexIDs, GroupID);
@@ -427,4 +464,25 @@ namespace UE::CADKernel::MeshUtilities
 	void FDynamicMeshWrapper::EndFaceTriangles()
 	{
 	}
+
+	void FDynamicMeshWrapper::RecomputeNullNormal()
+	{
+		// #cad_import: TO DO
+	}
+
+	void FDynamicMeshWrapper::OrientMesh()
+	{
+		// #cad_import: TO DO
+	}
+
+	void FDynamicMeshWrapper::ResolveTJunctions()
+	{
+		// #cad_import: TO DO
+	}
+
+
+	void GetExistingFaceGroups(UE::Geometry::FDynamicMesh3& Mesh, TSet<int32>& FaceGroupsOut)
+	{
+		// #cad_import: TO DO
+	}
 }
diff --git a/Engine/Source/Runtime/Datasmith/CADKernel/Engine/Private/DynamicMeshWrapper.h b/Engine/Source/Runtime/Datasmith/CADKernel/Engine/Private/DynamicMeshWrapper.h
--- a/Engine/Source/Runtime/Datasmith/CADKernel/Engine/Private/DynamicMeshWrapper.h
+++ b/Engine/Source/Runtime/Datasmith/CADKernel/Engine/Private/DynamicMeshWrapper.h
@@ -6,6 +6,7 @@
 #include "MeshUtilities.h"
 
 #include "DynamicMesh/DynamicMeshAttributeSet.h"
+#include "DynamicMesh/DynamicMesh3.h"
 
 namespace UE::CADKernel::MeshUtilities
 {
@@ -13,14 +14,14 @@ namespace UE::CADKernel::MeshUtilities
 	{
 	public:
 		FDynamicMeshWrapper(const FMeshExtractionContext& InContext, UE::Geometry::FDynamicMesh3& InMesh);
-
-		~FDynamicMeshWrapper() { FinalizeMesh(); }
 		
+		virtual void ClearMesh() override;
 		virtual bool ReserveNewTriangles(int32 AdditionalTriangleCount) override;
 		virtual bool SetVertices(TArray<FVector>&& Vertices) override;
 		virtual bool AddNewVertices(TArray<FVector>&& Vertices) override;
 		virtual bool AddTriangle(int32 GroupID, uint32 MaterialID, const FArray3i& VertexIndices, const TArrayView<FVector3f>& Normals, const TArrayView<FVector2f>& TexCoords) override;
 		virtual bool StartFaceTriangles(int32 TriangleCount, const TArray<FVector3f>& Normals, const TArray<FVector2f>& TexCoords) override;
+		virtual bool StartFaceTriangles(const TArrayView<FVector>& Normals, const TArrayView<FVector2d>& TexCoords) override;
 		virtual bool AddFaceTriangles(const TArray<FFaceTriangle>& FaceTriangles) override;
 		virtual bool AddFaceTriangle(const FFaceTriangle& FaceTriangle) override
 		{
@@ -29,9 +30,12 @@ namespace UE::CADKernel::MeshUtilities
 		virtual void EndFaceTriangles() override;
 
 		virtual void FinalizeMesh() override;
+		virtual void AddSymmetry() override;
+		virtual void RecomputeNullNormal() override;
+		virtual void OrientMesh() override;
+		virtual void ResolveTJunctions() override;
 
 	private:
-		void AddSymmetry();
 		void InitializeAttributes();
 		void AddTriangles();
 		int32 AppendTriangle(UE::Geometry::FIndex3i& VertexIDs, int32 GroupID);
diff --git a/Engine/Source/Runtime/Datasmith/CADKernel/Engine/Private/MeshDescriptionWrapper.cpp b/Engine/Source/Runtime/Datasmith/CADKernel/Engine/Private/MeshDescriptionWrapper.cpp
--- a/Engine/Source/Runtime/Datasmith/CADKernel/Engine/Private/MeshDescriptionWrapper.cpp
+++ b/Engine/Source/Runtime/Datasmith/CADKernel/Engine/Private/MeshDescriptionWrapper.cpp
@@ -5,7 +5,9 @@
 #include "CADKernelEngine.h"
 #include "CADKernelEngineLog.h"
 #include "MeshUtilities.h"
+#include "MeshTopologyHelper.h"
 
+#include "Containers/Queue.h"
 #include "MeshDescriptionBuilder.h"
 
 namespace UE::CADKernel::MeshUtilities
@@ -24,6 +26,7 @@ namespace UE::CADKernel::MeshUtilities
 		if (Attributes.IsValid())
 		{
 			VertexPositions = Mesh.GetVertexPositions();
+			VertexInstanceToVertex = Attributes.GetVertexInstanceVertexIndices();
 			VertexInstanceNormals = Attributes.GetVertexInstanceNormals();
 			VertexInstanceTangents = Attributes.GetVertexInstanceTangents();
 			VertexInstanceBinormalSigns = Attributes.GetVertexInstanceBinormalSigns();
@@ -34,6 +37,12 @@ namespace UE::CADKernel::MeshUtilities
 		}
 	}
 
+	void FMeshDescriptionWrapper::ClearMesh()
+	{
+		// # cad_import: Maybe cache the material sections???
+		Mesh.Empty();
+	}
+
 	void FMeshDescriptionWrapper::FinalizeMesh()
 	{
 		if (bIsFinalized)
@@ -41,7 +50,16 @@ namespace UE::CADKernel::MeshUtilities
 			return;
 		}
 
-		AddSymmetry();
+		TArrayView<FVector3f> Positions = VertexPositions.GetRawArray();
+		ConvertVectorArray(Context.ModelParams.ModelCoordSys, Positions);
+
+		if (!FMath::IsNearlyEqual(Context.ModelParams.ModelUnitToCentimeter, 1.f))
+		{
+			for (FVector3f& Position : Positions)
+			{
+				Position *= Context.ModelParams.ModelUnitToCentimeter;
+			}
+		}
 
 		if (Context.MeshParams.bNeedSwapOrientation)
 		{
@@ -63,11 +81,6 @@ namespace UE::CADKernel::MeshUtilities
 
 	void FMeshDescriptionWrapper::AddSymmetry()
 	{
-		if (!Context.MeshParams.bIsSymmetric)
-		{
-			return;
-		}
-
 		FMatrix44f SymmetricMatrix = (FMatrix44f)GetSymmetricMatrix(Context.MeshParams.SymmetricOrigin, Context.MeshParams.SymmetricNormal);
 
 		TMap<FVertexID, FVertexID> VertexMapping;
@@ -214,9 +227,6 @@ namespace UE::CADKernel::MeshUtilities
 		Mesh.ReserveNewVertices(Context.MeshParams.bIsSymmetric ? VertexCount * 2 : VertexCount);
 		VertexIDs.Empty(VertexCount);
 
-		// Make Mesh.VertexPositions and VertexID
-		ConvertVectorArray(Context.ModelParams.ModelCoordSys, InVertices);
-
 		for (const FVector& Vertex : InVertices)
 		{
 			const FVertexID VertexID = Mesh.CreateVertex();
@@ -304,15 +314,45 @@ namespace UE::CADKernel::MeshUtilities
 
 	bool FMeshDescriptionWrapper::StartFaceTriangles(int32 TriangleCount, const TArray<FVector3f>& InNormals, const TArray<FVector2f>& InTexCoords)
 	{
-		ensure(InNormals.Num() == InTexCoords.Num());
+		const int32 ArraySize = InNormals.Num();
+		ensure(ArraySize == InTexCoords.Num());
+
 		Normals = InNormals;
+		for (FVector3f& Normal : Normals)
+		{
+			Normal = Normal.GetSafeNormal();
+		}
+
+		ConvertVectorArray(Context.ModelParams.ModelCoordSys, Normals);
+
 		TexCoords = InTexCoords;
 
 		return true;
 	}
 
+	bool FMeshDescriptionWrapper::StartFaceTriangles(const TArrayView<FVector>& InNormals, const TArrayView<FVector2d>& InTexCoords)
+	{
+		const int32 ArraySize = InNormals.Num();
+		ensure(ArraySize == InTexCoords.Num());
+
+		Normals.Reserve(ArraySize);
+		TexCoords.Reserve(ArraySize);
+
+		for (int32 Index = 0; Index < ArraySize; ++Index)
+		{
+			Normals.Emplace(FVector3f(InNormals[Index].X, InNormals[Index].Y, InNormals[Index].Z).GetSafeNormal());
+			TexCoords.Emplace(InTexCoords[Index].X, InTexCoords[Index].Y);
+		}
+
+		ConvertVectorArray(Context.ModelParams.ModelCoordSys, Normals);
+
+		return true;
+	}
+
 	bool FMeshDescriptionWrapper::AddFaceTriangles(const TArray<FFaceTriangle>& FaceTriangles)
 	{
+		TRACE_CPUPROFILER_EVENT_SCOPE(FMeshDescriptionWrapper::AddFaceTriangles);
+
 		const bool bNeedSwapOrientation = Context.MeshParams.bNeedSwapOrientation;
 		const FArray3i& Orientation = bNeedSwapOrientation ? CounterClockwise : Clockwise;
 		
@@ -337,9 +377,9 @@ namespace UE::CADKernel::MeshUtilities
 			VertexInstanceUVs.Set(VertexInstanceIDs[1], 0/*UVChannel*/, FVector2f(TexCoords[FaceTriangle.TexCoords[Orientation[1]]] * ScaleUV));
 			VertexInstanceUVs.Set(VertexInstanceIDs[2], 0/*UVChannel*/, FVector2f(TexCoords[FaceTriangle.TexCoords[Orientation[2]]] * ScaleUV));
 
-			VertexInstanceNormals[VertexInstanceIDs[0]] = FVector3f(Normals[FaceTriangle.Normals[Orientation[0]]]);
-			VertexInstanceNormals[VertexInstanceIDs[1]] = FVector3f(Normals[FaceTriangle.Normals[Orientation[1]]]);
-			VertexInstanceNormals[VertexInstanceIDs[2]] = FVector3f(Normals[FaceTriangle.Normals[Orientation[2]]]);
+			VertexInstanceNormals[VertexInstanceIDs[0]] = Normals[FaceTriangle.Normals[Orientation[0]]];
+			VertexInstanceNormals[VertexInstanceIDs[1]] = Normals[FaceTriangle.Normals[Orientation[1]]];
+			VertexInstanceNormals[VertexInstanceIDs[2]] = Normals[FaceTriangle.Normals[Orientation[2]]];
 
 			// Add the triangle as a polygon to the mesh description
 			PolygonGroupID = GetPolygonGroupID(FaceTriangle.MaterialID);
@@ -353,6 +393,313 @@ namespace UE::CADKernel::MeshUtilities
 	}
 
 	void FMeshDescriptionWrapper::EndFaceTriangles()
+	{
+		Normals.Reset(Normals.Num());
+		TexCoords.Reset(TexCoords.Num());
+	}
+
+	void FMeshDescriptionWrapper::RecomputeNullNormal()
+	{
+		TRACE_CPUPROFILER_EVENT_SCOPE(FMeshDescriptionWrapper::RecomputeNullNormal);
+
+		constexpr float SquareNormalThreshold = UE_KINDA_SMALL_NUMBER * UE_KINDA_SMALL_NUMBER;
+
+		constexpr int32 TriangleIndex[3][3] = { {0,1,2},{1,2,0},{2,0,1} };
+		FVector NewNormal;
+
+		for (FTriangleID Triangle : Mesh.Triangles().GetElementIDs())
+		{
+			TArrayView<const FVertexInstanceID> Vertices = Mesh.GetTriangleVertexInstances(Triangle);
+			for (int32 VertexIndex = 0; VertexIndex < 3; VertexIndex++)
+			{
+				FVector3f& Normal = VertexInstanceNormals[Vertices[VertexIndex]];
+
+				if (Normal.IsNearlyZero(UE_KINDA_SMALL_NUMBER))
+				{
+					NewNormal = FVector::ZeroVector;
+
+					const FVertexInstanceID VertexInstanceID = Vertices[VertexIndex];
+					FVertexID VertexID = Mesh.GetVertexInstanceVertex(VertexInstanceID);
+					const TArrayView<const FTriangleID> VertexConnectedTriangles = Mesh.GetVertexInstanceConnectedTriangleIDs(VertexInstanceID);
+
+					bool bNormalComputed = false;
+
+					// compute the weighted sum of normals of the "partition star" according to the corner angle
+					for (const FTriangleID& TriangleID : VertexConnectedTriangles)
+					{
+						// compute face normal at the vertex
+						TArrayView<const FVertexID> TriangleVertices = Mesh.GetTriangleVertices(TriangleID);
+
+						int32 ApexIndex = 0;
+						for (; ApexIndex < 3; ++ApexIndex)
+						{
+							if (TriangleVertices[ApexIndex] == VertexID)
+							{
+								break;
+							}
+						}
+
+						const FVector Position0 = FVector(VertexPositions[TriangleVertices[TriangleIndex[ApexIndex][0]]]);
+						FVector DPosition1 = FVector(VertexPositions[TriangleVertices[TriangleIndex[ApexIndex][1]]]) - Position0;
+						FVector DPosition2 = FVector(VertexPositions[TriangleVertices[TriangleIndex[ApexIndex][2]]]) - Position0;
+
+						// to avoid numerical issue due to small vector
+						DPosition1.Normalize();
+						DPosition2.Normalize();
+
+						// We have a left-handed coordinate system, but a counter-clockwise winding order
+						// Hence normal calculation has to take the triangle vectors cross product in reverse.
+						FVector TriangleNormal = DPosition2 ^ DPosition1;
+						double SinOfApexAngle = TriangleNormal.Length();
+						double ApexAngle = FMath::Asin(SinOfApexAngle);
+
+						if (TriangleNormal.Normalize(SquareNormalThreshold))
+						{
+							NewNormal += (TriangleNormal * ApexAngle);
+							bNormalComputed = true;
+						}
+					}
+
+					if (bNormalComputed)
+					{
+						if (NewNormal.Normalize(SquareNormalThreshold))
+						{
+							Normal = FVector3f(NewNormal);
+						}
+					}
+					else
+					{
+						// the vertex is a vertex of degenerated triangles, the vertex normal doesn't matter as the triangle is flat, but a non null normal is needed.
+						Normal = FVector3f::UpVector;
+					}
+				}
+			}
+		}
+	}
+
+	void FMeshDescriptionWrapper::OrientMesh()
+	{
+		static const FVector MaxVector(-MAX_flt, -MAX_flt, -MAX_flt);
+		static const FVector MinVector( MAX_flt, MAX_flt, MAX_flt );
+		static const FIntVector UnInitVec( INDEX_NONE, INDEX_NONE, INDEX_NONE );
+
+		FMeshTopologyHelper MeshHelper(Mesh);
+
+		TQueue<FTriangleID> Front;
+		TQueue<FTriangleID> BadOrientationFront;
+
+		FTriangleID AdjacentTriangle;
+		FEdgeID Edge;
+
+		TArrayView < const FVertexInstanceID> Vertices;
+
+		FVector MinCorner;
+		FVector MaxCorner;
+		FIntVector HighestVertex;
+		FIntVector LowestVertex;
+
+		uint32 NbTriangles = Mesh.Triangles().Num();
+
+		TArray<FTriangleID> ConnectedTriangles;
+		ConnectedTriangles.Reserve(NbTriangles);
+
+		for (FTriangleID Triangle : Mesh.Triangles().GetElementIDs())
+		{
+			if (MeshHelper.IsTriangleMarked(Triangle))
+			{
+				continue;
+			}
+
+			MaxCorner = MaxVector;
+			MinCorner = MinVector;
+			HighestVertex = UnInitVec;
+			LowestVertex = UnInitVec;
+
+			MeshHelper.SetTriangleMarked(Triangle);
+
+			MeshHelper.GetTriangleVertexExtremities(Triangle, MinCorner, MaxCorner, HighestVertex, LowestVertex);
+
+			Front.Enqueue(Triangle);
+			ConnectedTriangles.Add(Triangle);
+
+			int32 NbConnectedFaces = 1;
+			int32 NbBorderEdges = 0;
+			int32 NbSurfaceEdges = 0;
+			int32 NbSwappedTriangles = 0;
+			while (!Front.IsEmpty())
+			{
+				while (!Front.IsEmpty())
+				{
+					Front.Dequeue(Triangle);
+
+					TArrayView<const FEdgeID> EdgeSet = Mesh.GetTriangleEdges(Triangle);
+
+					for (int32 IEdge = 0; IEdge < 3; IEdge++)
+					{
+						Edge = EdgeSet[IEdge];
+
+						if (!MeshHelper.IsEdgeOfType(Edge, EElementType::Surface))
+						{
+							NbBorderEdges++;
+							continue;
+						}
+
+						AdjacentTriangle = MeshHelper.GetOtherTriangleAtEdge(Edge, Triangle);
+						if (MeshHelper.IsTriangleMarked(AdjacentTriangle))
+						{
+							continue;
+						}
+
+						NbSurfaceEdges++;
+						NbConnectedFaces++;
+
+						ConnectedTriangles.Add(AdjacentTriangle);
+
+						MeshHelper.SetTriangleMarked(AdjacentTriangle);
+						MeshHelper.GetTriangleVertexExtremities(AdjacentTriangle, MinCorner, MaxCorner, HighestVertex, LowestVertex);
+
+						if (MeshHelper.GetEdgeDirectionInTriangle(Edge, 0) == MeshHelper.GetEdgeDirectionInTriangle(Edge, 1))
+						{
+							MeshHelper.SwapTriangleOrientation(AdjacentTriangle);
+							NbSwappedTriangles++;
+							BadOrientationFront.Enqueue(AdjacentTriangle);
+						}
+						else
+						{
+							Front.Enqueue(AdjacentTriangle);
+						}
+					}
+				}
+
+				while (!BadOrientationFront.IsEmpty())
+				{
+					BadOrientationFront.Dequeue(Triangle);
+
+					TArrayView<const FEdgeID> EdgeSet = Mesh.GetTriangleEdges(Triangle);
+
+					for (int32 IEdge = 0; IEdge < 3; IEdge++)
+					{
+						Edge = EdgeSet[IEdge];
+
+						if (!MeshHelper.IsEdgeOfType(Edge, EElementType::Surface))
+						{
+							NbBorderEdges++;
+							continue;
+						}
+
+						AdjacentTriangle = MeshHelper.GetOtherTriangleAtEdge(Edge, Triangle);
+						if (MeshHelper.IsTriangleMarked(AdjacentTriangle))
+						{
+							continue;
+						}
+
+						NbSurfaceEdges++;
+						NbConnectedFaces++;
+
+						ConnectedTriangles.Add(AdjacentTriangle);
+
+						MeshHelper.SetTriangleMarked(AdjacentTriangle);
+						MeshHelper.GetTriangleVertexExtremities(AdjacentTriangle, MinCorner, MaxCorner, HighestVertex, LowestVertex);
+						if (MeshHelper.GetEdgeDirectionInTriangle(Edge, 0) == MeshHelper.GetEdgeDirectionInTriangle(Edge, 1))
+						{
+							BadOrientationFront.Enqueue(AdjacentTriangle);
+							MeshHelper.SwapTriangleOrientation(AdjacentTriangle);
+							NbSwappedTriangles++;
+						}
+						else
+						{
+							Front.Enqueue(AdjacentTriangle);
+						}
+					}
+				}
+			}
+
+			// Check if the mesh orientation need to be swapped
+			int NbInverted = 0;
+			int NbNotInverted = 0;
+			if (NbBorderEdges == 0 || NbBorderEdges * 20 < NbSurfaceEdges)  // NbBorderEdges * 20 < nbSurfaceEdges => basic rule to define if a mesh is a surface mesh or if the mesh is a volume mesh with gaps 
+			{
+				// case of volume mesh
+				// A vertex can have many normal (one by vertex instance) e.g. the corner of a box that have 3 normal and could be the highest 
+				// vertex of a mesh. At the highest vertex, a folding of the mesh can give a vertex with two opposite normal.
+				// The normal most parallel to the axis is preferred
+				// To avoid mistake, we check for each highest vertex and trust the majority 
+
+				if (HighestVertex[0] != INDEX_NONE)
+				{
+					for (int32 VertexIndex = 0; VertexIndex < 3; VertexIndex++)
+					{
+						if (MeshHelper.IsVertexOfType(HighestVertex[VertexIndex], EElementType::Surface))
+						{
+							FVertexID VertexID = Mesh.GetVertexInstanceVertex(HighestVertex[VertexIndex]);
+							TArrayView<const FVertexInstanceID> CoincidentVertexInstanceIdSet = Mesh.GetVertexVertexInstanceIDs(VertexID);
+							float MaxComponent = 0;
+							for (const FVertexInstanceID VertexInstanceID : CoincidentVertexInstanceIdSet)
+							{
+								FVector Normal = (FVector)VertexInstanceNormals[VertexInstanceID];
+								if (FMath::Abs(MaxComponent) < FMath::Abs(Normal[VertexIndex]))
+								{
+									MaxComponent = Normal[VertexIndex];
+								}
+							}
+
+							if (0 > MaxComponent)
+							{
+								NbInverted++;
+							}
+							else
+							{
+								NbNotInverted++;
+							}
+						}
+
+						if (MeshHelper.IsVertexOfType(LowestVertex[VertexIndex], EElementType::Surface))
+						{
+							FVertexID VertexID = Mesh.GetVertexInstanceVertex(LowestVertex[VertexIndex]);
+							TArrayView<const FVertexInstanceID> CoincidentVertexInstanceIdSet = Mesh.GetVertexVertexInstanceIDs(VertexID);
+							float MaxComponent = 0;
+							for (const FVertexInstanceID VertexInstanceID : CoincidentVertexInstanceIdSet)
+							{
+								FVector Normal = (FVector)Normals[VertexInstanceID];
+								if (FMath::Abs(MaxComponent) < FMath::Abs(Normal[VertexIndex]))
+								{
+									MaxComponent = Normal[VertexIndex];
+								}
+							}
+
+							if (0 < MaxComponent)
+							{
+								NbInverted++;
+							}
+							else
+							{
+								NbNotInverted++;
+							}
+						}
+					}
+				}
+			}
+			else if (NbSwappedTriangles * 2 > NbConnectedFaces)
+			{
+				// case of surface mesh
+				// this means that more triangles of surface shape have been swapped than no swapped, the good orientation has been reversed.
+				// The mesh need to be re swapped
+				NbInverted++;
+			}
+
+			// if needed swap all the mesh
+			if (NbInverted > NbNotInverted)
+			{
+				for (auto Tri : ConnectedTriangles)
+				{
+					MeshHelper.SwapTriangleOrientation(Tri);
+				}
+			}
+			ConnectedTriangles.Empty(NbTriangles);
+		}
+	}
+
+	void FMeshDescriptionWrapper::ResolveTJunctions()
 	{
 	}
 
@@ -381,4 +728,20 @@ namespace UE::CADKernel::MeshUtilities
 		using namespace ExtendedMeshAttribute;
 		return MeshDescription.PolygonAttributes().GetAttributesRef<int32>(PolyTriGroups);
 	}
+
+	void GetExistingFaceGroups(FMeshDescription& Mesh, TSet<int32>& FaceGroupsOut)
+	{
+		using namespace ExtendedMeshAttribute;
+		TPolygonAttributesRef<int32> ElementToGroups = Mesh.PolygonAttributes().GetAttributesRef<int32>(PolyTriGroups);
+		int32 LastPatchId = -1;
+		for (const FPolygonID TriangleID : Mesh.Polygons().GetElementIDs())
+		{
+			int32 PatchId = ElementToGroups[TriangleID];
+			if (PatchId != LastPatchId)
+			{
+				FaceGroupsOut.Add(PatchId);
+				LastPatchId = PatchId;
+			}
+		}
+	}
 }
diff --git a/Engine/Source/Runtime/Datasmith/CADKernel/Engine/Private/MeshDescriptionWrapper.h b/Engine/Source/Runtime/Datasmith/CADKernel/Engine/Private/MeshDescriptionWrapper.h
--- a/Engine/Source/Runtime/Datasmith/CADKernel/Engine/Private/MeshDescriptionWrapper.h
+++ b/Engine/Source/Runtime/Datasmith/CADKernel/Engine/Private/MeshDescriptionWrapper.h
@@ -43,13 +43,13 @@ namespace UE::CADKernel::MeshUtilities
 	public:
 		FMeshDescriptionWrapper(const FMeshExtractionContext& InContext, FMeshDescription& InMesh);
 
-		~FMeshDescriptionWrapper() { FinalizeMesh(); }
-
+		virtual void ClearMesh() override;
 		virtual bool ReserveNewTriangles(int32 AdditionalTriangleCount) override;
 		virtual bool SetVertices(TArray<FVector>&& Vertices) override;
 		virtual bool AddNewVertices(TArray<FVector>&& Vertices) override;
 		virtual bool AddTriangle(int32 GroupID, uint32 MaterialID, const FArray3i& VertexIndices, const TArrayView<FVector3f>& Normals, const TArrayView<FVector2f>& TexCoords) override;
 		virtual bool StartFaceTriangles(int32 TriangleCount, const TArray<FVector3f>& Normals, const TArray<FVector2f>& TexCoords) override;
+		virtual bool StartFaceTriangles(const TArrayView<FVector>& Normals, const TArrayView<FVector2d>& TexCoords) override;
 		virtual bool AddFaceTriangles(const TArray<FFaceTriangle>& FaceTriangles) override;
 		virtual bool AddFaceTriangle(const FFaceTriangle& FaceTriangle) override
 		{
@@ -58,6 +58,10 @@ namespace UE::CADKernel::MeshUtilities
 		virtual void EndFaceTriangles() override;
 
 		virtual void FinalizeMesh() override;
+		virtual void AddSymmetry() override;
+		virtual void RecomputeNullNormal() override;
+		virtual void OrientMesh() override;
+		virtual void ResolveTJunctions() override;
 
 	private:
 		template<typename VertexType>
@@ -126,12 +130,12 @@ namespace UE::CADKernel::MeshUtilities
 		}
 
 		bool GetVertexInstances(const FArray3i& Vertices, FVertexInstanceID3& VertexInstances);
-		void AddSymmetry();
 
 	private:
 		int32 VertexIndexOffset = 0;
 		FCADKernelStaticMeshAttributes Attributes;
 		TVertexAttributesRef<FVector3f> VertexPositions;
+		TVertexInstanceAttributesRef<FVertexID> VertexInstanceToVertex;
 		TVertexInstanceAttributesRef<FVector3f> VertexInstanceNormals;
 		TVertexInstanceAttributesRef<FVector3f> VertexInstanceTangents;
 		TVertexInstanceAttributesRef<float> VertexInstanceBinormalSigns;
diff --git a/Engine/Source/Runtime/Datasmith/CADKernel/Engine/Private/MeshTopologyHelper.cpp b/Engine/Source/Runtime/Datasmith/CADKernel/Engine/Private/MeshTopologyHelper.cpp
new file mode 100644
--- /dev/null
+++ b/Engine/Source/Runtime/Datasmith/CADKernel/Engine/Private/MeshTopologyHelper.cpp
@@ -0,0 +1,277 @@
+// Copyright Epic Games, Inc. All Rights Reserved.
+
+
+#include "MeshTopologyHelper.h"
+
+#include "MeshAttributes.h"
+#include "StaticMeshAttributes.h"
+
+namespace UE::CADKernel::MeshUtilities
+{
+	FMeshTopologyHelper::FMeshTopologyHelper(FMeshDescription& InMeshDescription)
+		: Mesh(InMeshDescription)
+	{
+		UpdateMeshWrapper();
+	}
+
+	void FMeshTopologyHelper::UpdateMeshWrapper()
+	{
+		VertexTypeSet.SetNumZeroed(Mesh.Vertices().GetArraySize());
+		EdgeTypeSet.SetNumZeroed(Mesh.Edges().GetArraySize());
+
+		TriangleMarkers.Init(false, Mesh.Triangles().GetArraySize());
+
+		// If there are more than one polygon groups, edges between groups are considered feature lines
+		for (const FEdgeID EdgeID : Mesh.Edges().GetElementIDs())
+		{
+			DefineEdgeTopology(EdgeID);
+		}
+
+		for (const FVertexID VertexID : Mesh.Vertices().GetElementIDs())
+		{
+			DefineVertexTopologyApproximation(VertexID);
+		}
+
+		VertexInstanceMarker.Init(false, Mesh.VertexInstances().GetArraySize());
+	}
+
+
+	void FMeshTopologyHelper::DefineEdgeTopology(FEdgeID EdgeID)
+	{
+		if (EdgeID == INDEX_NONE)
+		{
+			return;
+		}
+
+		TArrayView<const FTriangleID> EdgeConnectedPolygons = Mesh.GetEdgeConnectedTriangleIDs(EdgeID);
+
+		switch (EdgeConnectedPolygons.Num())
+		{
+		case 0:
+			EdgeTypeSet[EdgeID] = EElementType::Line;
+			break;
+		case 1:
+			EdgeTypeSet[EdgeID] = EElementType::Border;
+			break;
+		case 2:
+			EdgeTypeSet[EdgeID] = EElementType::Surface;
+			break;
+		default:
+			EdgeTypeSet[EdgeID] = EElementType::NonManifold;
+			break;
+		}
+	}
+
+	void FMeshTopologyHelper::DefineVertexTopologyApproximation(FVertexID VertexID)
+	{
+		if (VertexID == INDEX_NONE)
+		{
+			return;
+		}
+
+		TArrayView<const FEdgeID> VertexConnectedEdgeIDs = Mesh.GetVertexConnectedEdgeIDs(VertexID);
+
+		switch (VertexConnectedEdgeIDs.Num())
+		{
+		case 0:
+			VertexTypeSet[VertexID] = EElementType::Free;
+			break;
+
+		case 1:
+			VertexTypeSet[VertexID] = EElementType::Border;
+			break;
+		default:
+		{
+			int CountPerCategory[8] = { 0 };
+
+			for (const FEdgeID EdgeID : VertexConnectedEdgeIDs)
+			{
+				CountPerCategory[type_width(EdgeTypeSet[EdgeID])]++;
+			}
+
+			if (CountPerCategory[type_width(EElementType::NonManifold)] > 0)
+			{
+				VertexTypeSet[VertexID] = EElementType::NonManifold;
+			}
+			else if (CountPerCategory[type_width(EElementType::Border)] > 0)
+			{
+				VertexTypeSet[VertexID] = EElementType::NonSurface;
+
+			}
+			else if (CountPerCategory[type_width(EElementType::Line)] > 0)
+			{
+				if (CountPerCategory[type_width(EElementType::Line)] == 2 && VertexConnectedEdgeIDs.Num() == 2)
+				{
+					VertexTypeSet[VertexID] = EElementType::Line;
+				}
+				else
+				{
+					VertexTypeSet[VertexID] = EElementType::NonManifold;
+				}
+			}
+			else if (CountPerCategory[type_width(EElementType::Surface)] > 0)
+			{
+				const int32 EdgeCount = VertexConnectedEdgeIDs.Num();
+				const FEdgeID FirstEdgeID(VertexConnectedEdgeIDs[0]);
+				FEdgeID EdgeID(FirstEdgeID);
+
+				FTriangleID TriangleID(INDEX_NONE);
+				int32 TriangleCount = 0;
+
+				do
+				{
+					TArrayView<const FTriangleID> EdgeConnectedPolygons = Mesh.GetEdgeConnectedTriangleIDs(EdgeID);
+					// Border edge no more triangles to process, exit from the loop
+					if (EdgeConnectedPolygons.Num() < 2)
+					{
+						break;
+					}
+
+					TriangleID = (TriangleID == EdgeConnectedPolygons[0]) ? EdgeConnectedPolygons[1] : EdgeConnectedPolygons[0];
+					++TriangleCount;
+
+					TArrayView<const FEdgeID> TriangleEdges = Mesh.GetTriangleEdges(TriangleID);
+
+					for (int32 Corner = 0; Corner < 3; ++Corner)
+					{
+						if (EdgeID != TriangleEdges[Corner])
+						{
+							const FVertexID EdgeVertexID0 = Mesh.GetEdgeVertex(TriangleEdges[Corner], 0);
+							const FVertexID EdgeVertexID1 = Mesh.GetEdgeVertex(TriangleEdges[Corner], 1);
+							if (EdgeVertexID0 == VertexID || EdgeVertexID1 == VertexID)
+							{
+								EdgeID = TriangleEdges[Corner];
+								break;
+							}
+						}
+					}
+				} while (EdgeID != FirstEdgeID);
+
+				VertexTypeSet[VertexID] = TriangleCount == EdgeCount ? EElementType::Surface : EElementType::NonManifold;
+			}
+			else
+			{
+				// we don`t need to know exactly if the node is a border or non manifold node. Non surface is enough
+				VertexTypeSet[VertexID] = EElementType::NonSurface;
+			}
+		}
+		break;
+		}
+	}
+
+	FMeshTopologyHelper::~FMeshTopologyHelper()
+	{
+
+	}
+
+	bool FMeshTopologyHelper::IsVertexOfType(FVertexInstanceID VertexInstanceID, EElementType Type) const
+	{
+		FVertexID VertexID = Mesh.GetVertexInstanceVertex(VertexInstanceID);
+		return EnumHasAnyFlags( VertexTypeSet[VertexID], Type);
+	}
+
+	// Triangle
+	void FMeshTopologyHelper::GetTriangleVertexExtremities(FTriangleID Triangle, FVector& MinCorner, FVector& MaxCorner, FIntVector& HighestVertex, FIntVector& LowestVertex) const
+	{
+
+		TArrayView < const FVertexInstanceID> Vertices = Mesh.GetTriangleVertexInstances(Triangle);
+
+		GetVertexExtremities(Vertices[0], MinCorner, MaxCorner, HighestVertex, LowestVertex);
+		GetVertexExtremities(Vertices[1], MinCorner, MaxCorner, HighestVertex, LowestVertex);
+		GetVertexExtremities(Vertices[2], MinCorner, MaxCorner, HighestVertex, LowestVertex);
+	}
+
+	void FMeshTopologyHelper::GetVertexExtremities(FVertexInstanceID VertexInstanceID, FVector& MinCorner, FVector& MaxCorner, FIntVector& HighestVertex, FIntVector& LowestVertex) const
+	{
+		FVertexID VertexID = Mesh.GetVertexInstanceVertex(VertexInstanceID);
+		const FVector VertexPosition = (FVector)Mesh.GetVertexPositions()[VertexID];
+
+		if (MaxCorner[0] < VertexPosition[0])
+		{
+			HighestVertex[0] = VertexInstanceID;
+			MaxCorner[0] = VertexPosition[0];
+		}
+		if (MaxCorner[1] < VertexPosition[1])
+		{
+			HighestVertex[1] = VertexInstanceID;
+			MaxCorner[1] = VertexPosition[1];
+		}
+		if (MaxCorner[2] < VertexPosition[2])
+		{
+			HighestVertex[2] = VertexInstanceID;
+			MaxCorner[2] = VertexPosition[2];
+		}
+
+		if (MinCorner[0] > VertexPosition[0])
+		{
+			LowestVertex[0] = VertexInstanceID;
+			MinCorner[0] = VertexPosition[0];
+		}
+		if (MinCorner[1] > VertexPosition[1])
+		{
+			LowestVertex[1] = VertexInstanceID;
+			MinCorner[1] = VertexPosition[1];
+		}
+		if (MinCorner[2] > VertexPosition[2])
+		{
+			LowestVertex[2] = VertexInstanceID;
+			MinCorner[2] = VertexPosition[2];
+		}
+	}
+
+	void FMeshTopologyHelper::SwapTriangleOrientation(FTriangleID Triangle)
+	{
+		Mesh.ReverseTriangleFacing(Triangle);
+
+		TArrayView<const FVertexInstanceID> TriVertexInstances = Mesh.GetTriangleVertexInstances(Triangle);
+		for (int32 IVertex = 0; IVertex < 3; IVertex++)
+		{
+			const FVertexInstanceID& InstanceID = TriVertexInstances[IVertex];
+			if (!VertexInstanceMarker[InstanceID])
+			{
+				SwapVertexNormal(InstanceID);
+				VertexInstanceMarker[InstanceID] = true;
+			}
+		}
+	}
+
+	// Edge
+	const FTriangleID FMeshTopologyHelper::GetOtherTriangleAtEdge(FEdgeID EdgeID, FTriangleID Triangle) const
+	{
+		TArrayView<const FTriangleID> EdgeConnectedPolygons = Mesh.GetEdgeConnectedTriangleIDs(EdgeID);
+		return EdgeConnectedPolygons.Num() < 2 ? INDEX_NONE : (EdgeConnectedPolygons[0] == Triangle ? EdgeConnectedPolygons[1] : EdgeConnectedPolygons[0]);
+	}
+
+	bool FMeshTopologyHelper::GetEdgeDirectionInTriangle(FEdgeID EdgeID, int32 TriangleIndex) const
+	{
+		TArrayView<const FTriangleID> EdgeConnectedTriangles = Mesh.GetEdgeConnectedTriangleIDs(EdgeID);
+
+		if (EdgeConnectedTriangles.Num() > TriangleIndex)
+		{
+			const FTriangleID TriangleID = EdgeConnectedTriangles[TriangleIndex];
+			TArrayView<const FVertexInstanceID> VertexInstanceIDs = Mesh.GetTriangleVertexInstances(TriangleID);
+			TArrayView<const FEdgeID> TriangleEdges = Mesh.GetTriangleEdges(TriangleID);
+
+			for (int32 Corner = 0; Corner < 3; ++Corner)
+			{
+				if (TriangleEdges[Corner] == EdgeID)
+				{
+					const FVertexID VertexID = Mesh.GetVertexInstanceVertex(VertexInstanceIDs[Corner]);
+					return Mesh.GetEdgeVertex(EdgeID, 0) == VertexID ? true : false;
+				}
+			}
+		}
+		return true;
+	}
+
+	// Vertex
+	void FMeshTopologyHelper::SwapVertexNormal(FVertexInstanceID VertexInstanceID)
+	{
+		FStaticMeshAttributes StaticMeshAttributes(Mesh);
+		FVector3f Normal = StaticMeshAttributes.GetVertexInstanceNormals()[VertexInstanceID];
+		Normal *= -1;
+
+		StaticMeshAttributes.GetVertexInstanceNormals()[VertexInstanceID] = Normal;
+	}
+}
+
diff --git a/Engine/Source/Runtime/Datasmith/CADKernel/Engine/Private/MeshTopologyHelper.h b/Engine/Source/Runtime/Datasmith/CADKernel/Engine/Private/MeshTopologyHelper.h
new file mode 100644
--- /dev/null
+++ b/Engine/Source/Runtime/Datasmith/CADKernel/Engine/Private/MeshTopologyHelper.h
@@ -0,0 +1,105 @@
+// Copyright Epic Games, Inc. All Rights Reserved.
+
+#pragma once
+
+
+#include "CoreMinimal.h"
+
+#include "MeshTypes.h"
+#include "MeshDescription.h"
+
+#include <bit>
+
+#define ELEMENT_FIRST_MARKER_MASK	0x02
+
+namespace UE::CADKernel::MeshUtilities
+{
+	struct FElementMetaData
+	{
+		uint16 Category : 4;
+		uint16 Markers : 4;
+		uint16 Extras : 4;
+	};
+
+	enum class EElementType : uint8_t
+	{
+		Unused       = 0,
+		Free         = 1 << 0,
+		Line         = 1 << 1,
+		Surface      = 1 << 2,
+		Border       = 1 << 3,
+		NonManifold  = 1 << 4,
+		NonSurface   = 1 << 5,
+		Max			 = 1 << 7,
+	};
+	ENUM_CLASS_FLAGS(EElementType)
+
+	class FMeshTopologyHelper
+	{
+
+	private:
+		FMeshDescription& Mesh;
+
+		TArray<EElementType> VertexTypeSet;
+		TArray<EElementType> EdgeTypeSet;
+
+		TBitArray<> VertexInstanceMarker;
+		TBitArray<> TriangleMarkers;
+
+	private:
+		static int32 type_width(EElementType Type)
+		{
+			return std::bit_width((unsigned char)Type);
+		}
+
+	public:
+		FMeshTopologyHelper(FMeshDescription& InMeshDescription);
+		~FMeshTopologyHelper();
+
+		void UpdateMeshWrapper();
+
+		bool IsTriangleMarked(FTriangleID Triangle) const
+		{
+			return TriangleMarkers[Triangle];
+		}
+		void SetTriangleMarked(FTriangleID Triangle)
+		{
+			TriangleMarkers[Triangle] = true;
+		}
+
+		EElementType GetEdgeType(FEdgeID Edge) const
+		{
+			return EdgeTypeSet[Edge];
+		}
+
+		bool IsEdgeOfType(FEdgeID Edge, EElementType Type) const
+		{
+			return EnumHasAnyFlags(EdgeTypeSet[Edge], Type);
+		}
+
+		bool IsVertexOfType(FVertexInstanceID Vertex, EElementType Type) const;
+
+		// Triangle
+		void GetTriangleVertexExtremities(FTriangleID Triangle, FVector& MinCorner, FVector& MaxCorner, FIntVector& HighestVertex, FIntVector& LowestVertex) const;
+
+		void SwapTriangleOrientation(FTriangleID Triangle);
+
+		// Edge
+		const FTriangleID GetOtherTriangleAtEdge(FEdgeID Edge, FTriangleID Triangle) const;
+
+		/**
+		 * @param Edge           The edge defined by its FEdgeID to get direction in its connected triangle
+		 * @param TriangleIndex  The triangle connected to the edge, 0 for the first, 1 for the second, ...
+		 *
+		 */
+		bool GetEdgeDirectionInTriangle(FEdgeID Edge, int32 TriangleIndex) const;
+		void DefineEdgeTopology(FEdgeID EdgeID);
+
+		// Vertex
+		void SwapVertexNormal(FVertexInstanceID VertexID);
+		void GetVertexExtremities(FVertexInstanceID Vertex, FVector& MinCorner, FVector& MaxCorner, FIntVector& HighestVertex, FIntVector& LowestVertex) const;
+
+		void DefineVertexTopologyApproximation(FVertexID VertexID);
+	};
+}
+
diff --git a/Engine/Source/Runtime/Datasmith/CADKernel/Engine/Private/MeshUtilities.cpp b/Engine/Source/Runtime/Datasmith/CADKernel/Engine/Private/MeshUtilities.cpp
--- a/Engine/Source/Runtime/Datasmith/CADKernel/Engine/Private/MeshUtilities.cpp
+++ b/Engine/Source/Runtime/Datasmith/CADKernel/Engine/Private/MeshUtilities.cpp
@@ -4,24 +4,45 @@
 
 #include "CADKernelEngineLog.h"
 
+#include "MeshDescription.h"
 #include "MeshDescriptionBuilder.h"
+#include "StaticMeshAttributes.h"
 #include "StaticMeshOperations.h"
 
-namespace UE::CADKernel::MeshUtilities
+namespace UE::CADKernel
 {
-	bool FMeshOperations::OrientMesh(FMeshDescription& MeshDescription)
+	void ApplyExtractionContext(const FMeshExtractionContext& Context, FMeshDescription& MeshInOut)
 	{
-		// #cad_import: TODO - Refactor MeshOperator::OrientMesh
-		return true;
+		using namespace MeshUtilities;
+
+		TSharedPtr<FMeshWrapperAbstract> MeshWrapper = FMeshWrapperAbstract::MakeWrapper(Context, MeshInOut);
+		MeshWrapper->Complete();
 	}
 
-	void FMeshOperations::ResolveTJunctions(FMeshDescription& MeshDescription, double Tolerance)
+	void ApplyExtractionContext(const FMeshExtractionContext& Context, UE::Geometry::FDynamicMesh3& MeshInOut)
 	{
-		// #cad_import: TODO - Refactor MeshOperator::ResolveTJunctions
+		using namespace MeshUtilities;
+
+		TSharedPtr<FMeshWrapperAbstract> MeshWrapper = FMeshWrapperAbstract::MakeWrapper(Context, MeshInOut);
+		MeshWrapper->Complete();
 	}
 
-	void FMeshOperations::RecomputeNullNormal(FMeshDescription& MeshDescriptione)
+	namespace MeshUtilities
 	{
-		// #cad_import: TODO Refactor MeshOperator::RecomputeNullNormal
+		bool FMeshOperations::OrientMesh(FMeshDescription& MeshDescription)
+		{
+			// #cad_import: TODO - Refactor MeshOperator::OrientMesh
+			return true;
+		}
+
+		void FMeshOperations::ResolveTJunctions(FMeshDescription& MeshDescription, double Tolerance)
+		{
+			// #cad_import: TODO - Refactor MeshOperator::ResolveTJunctions
+		}
+
+		void FMeshOperations::RecomputeNullNormal(FMeshDescription& MeshDescriptione)
+		{
+			// #cad_import: TODO Refactor MeshOperator::RecomputeNullNormal
+		}
 	}
 }
\ No newline at end of file
diff --git a/Engine/Source/Runtime/Datasmith/CADKernel/Engine/Private/MeshUtilities.h b/Engine/Source/Runtime/Datasmith/CADKernel/Engine/Private/MeshUtilities.h
--- a/Engine/Source/Runtime/Datasmith/CADKernel/Engine/Private/MeshUtilities.h
+++ b/Engine/Source/Runtime/Datasmith/CADKernel/Engine/Private/MeshUtilities.h
@@ -2,9 +2,11 @@
 
 #pragma once
 
+#include "CADKernelEngine.h"
 #include "ParametricSurfaceData.h"
 
 #include "IntVectorTypes.h"
+#include "Util/DynamicVector.h"
 
 namespace UE::CADKernel
 {
@@ -39,13 +41,52 @@ namespace UE::CADKernel::MeshUtilities
 
 			// Indices referencing the array of texture coordinates, TexCoords, set calling StartFaceTriangles.
 			const FArray3i TexCoords;
+
+			FFaceTriangle(int32 InGroupID, uint32 InMaterialID, const FArray3i& InVertexIndices, const FArray3i& InNormals, const FArray3i& InTexCoords)
+				: GroupID(InGroupID)
+				, MaterialID(InMaterialID)
+				, VertexIndices(InVertexIndices)
+				, Normals(InNormals)
+				, TexCoords(InTexCoords)
+			{
+			}
 		};
 
-		FMeshWrapperAbstract(const FMeshExtractionContext& InContext) : Context(InContext) {}
+		FMeshWrapperAbstract(const FMeshExtractionContext& InContext) : Context(InContext)
+		{
+			bHasFaceGroupsToSkip = !InContext.FaceGroupsToExtract.IsEmpty();
+		}
 
-		virtual ~FMeshWrapperAbstract() {}
+		void Complete()
+		{
+			if (bIsComplete)
+			{
+				return;
+			}
 
-		virtual void FinalizeMesh() = 0;
+			if (Context.MeshParams.bIsSymmetric)
+			{
+				AddSymmetry();
+			}
+
+			FinalizeMesh();
+
+			// Workaround for SDHE-19725 (Declined): Compute any null normals.
+			RecomputeNullNormal();
+
+			OrientMesh();
+			
+			if (Context.bResolveTJunctions)
+			{
+				ResolveTJunctions();
+			}
+
+			bIsComplete = true;
+		}
+
+		virtual ~FMeshWrapperAbstract() { Complete(); }
+
+		virtual void ClearMesh() = 0;
 
 		virtual bool SetVertices(TArray<FVector>&& Vertices) = 0;
 
@@ -57,19 +98,26 @@ namespace UE::CADKernel::MeshUtilities
 		 ** Expected Normals.Num() == TexCoords.Num() == 3 * TriangleCount
 		 */
 		virtual bool StartFaceTriangles(int32 TriangleCount, const TArray<FVector3f>& Normals, const TArray<FVector2f>& TexCoords) = 0;
-		
+		virtual bool StartFaceTriangles(const TArrayView<FVector>& Normals, const TArrayView<FVector2d>& TexCoords) = 0;
+
 		virtual bool AddFaceTriangles(const TArray<FFaceTriangle>& FaceTriangles) = 0;
 
 		virtual bool AddFaceTriangle(const FFaceTriangle& FaceTriangle) = 0;
 
 		virtual void EndFaceTriangles() = 0;
 
+
 		/*
 		 * Normals and TexCoords array are expected to be arrays of 3 elements
 		 * Each value in those arrays is associated to the vertex in VertexIndices at the same index value, 0, 1 and 2.
 		 */
 		virtual bool AddTriangle(int32 GroupID, uint32 MaterialID, const FArray3i& VertexIndices, const TArrayView<FVector3f>& Normals, const TArrayView<FVector2f>& TexCoords) = 0;
 
+		bool IsFaceGroupValid(int32 GroupID)
+		{
+			return bHasFaceGroupsToSkip ? Context.FaceGroupsToExtract.Contains(GroupID) : true;
+		}
+
 		static TSharedPtr<FMeshWrapperAbstract> MakeWrapper(const FMeshExtractionContext& InContext, FMeshDescription& Mesh);
 		static TSharedPtr<FMeshWrapperAbstract> MakeWrapper(const FMeshExtractionContext& Context, UE::Geometry::FDynamicMesh3& Mesh);
 
@@ -82,10 +130,24 @@ namespace UE::CADKernel::MeshUtilities
 
 		// #cad_import: The UV scaling should be done in CADKernel
 		static constexpr double ScaleUV = 0.001; // mm to m
+
+	protected:
+		virtual void AddSymmetry() = 0;
+		virtual void FinalizeMesh() = 0;
+		virtual void RecomputeNullNormal() = 0;
+		virtual void OrientMesh() = 0;
+		virtual void ResolveTJunctions() = 0;
+
+	private:
+		bool bHasFaceGroupsToSkip = false;
+		bool bIsComplete = false;
 	};
 
+	void GetExistingFaceGroups(FMeshDescription& Mesh, TSet<int32>& FaceGroupsOut);
+	void GetExistingFaceGroups(UE::Geometry::FDynamicMesh3& Mesh, TSet<int32>& FaceGroupsOut);
+
 	template<typename VecType>
-	void ConvertVectorArray(uint8 ModelCoordSys, TArray<VecType>& Array)
+	void ConvertVectorArray(uint8 ModelCoordSys, TArrayView<VecType>& Array)
 	{
 		switch (ECADKernelModelCoordSystem(ModelCoordSys))
 		{
@@ -123,6 +185,52 @@ namespace UE::CADKernel::MeshUtilities
 		}
 	}
 
+	template<typename VecType>
+	void ConvertVectorArray(uint8 ModelCoordSys, TArray<VecType>& Array)
+	{
+		TArrayView<VecType> ArrayView(Array);
+		ConvertVectorArray(ModelCoordSys, ArrayView);
+	}
+
+	template<typename VecType>
+	void ConvertVectorArray(uint8 ModelCoordSys, UE::Geometry::TDynamicVector<VecType>& Array)
+	{
+		switch (ECADKernelModelCoordSystem(ModelCoordSys))
+		{
+			case ECADKernelModelCoordSystem::YUp_LeftHanded:
+			for (VecType& Vector : Array)
+			{
+				Vector.Set(Vector[2], Vector[0], Vector[1]);
+			}
+			break;
+
+			case ECADKernelModelCoordSystem::YUp_RightHanded:
+			for (VecType& Vector : Array)
+			{
+				Vector.Set(-Vector[2], Vector[0], Vector[1]);
+			}
+			break;
+
+			case ECADKernelModelCoordSystem::ZUp_RightHanded:
+			for (VecType& Vector : Array)
+			{
+				Vector.Set(-Vector[0], Vector[1], Vector[2]);
+			}
+			break;
+
+			case ECADKernelModelCoordSystem::ZUp_RightHanded_FBXLegacy:
+			for (VecType& Vector : Array)
+			{
+				Vector.Set(Vector[0], -Vector[1], Vector[2]);
+			}
+			break;
+
+			case ECADKernelModelCoordSystem::ZUp_LeftHanded:
+			default:
+			break;
+		}
+	}
+
 	template<typename VecType>
 	VecType ConvertVector(uint8 ModelCoordSys, const VecType& V)
 	{
diff --git a/Engine/Source/Runtime/Datasmith/CADKernel/Engine/Private/ParametricSurfaceData.cpp b/Engine/Source/Runtime/Datasmith/CADKernel/Engine/Private/ParametricSurfaceData.cpp
--- a/Engine/Source/Runtime/Datasmith/CADKernel/Engine/Private/ParametricSurfaceData.cpp
+++ b/Engine/Source/Runtime/Datasmith/CADKernel/Engine/Private/ParametricSurfaceData.cpp
@@ -2,6 +2,10 @@
 
 #include "ParametricSurfaceData.h"
 
+#include "CADKernelEnginePrivate.h"
+#include "MeshUtilities.h"
+#include "TechSoft/TechSoftUtilities.h"
+
 #include "Core/Session.h"
 #include "Topo/Model.h"
 
@@ -11,10 +15,12 @@
 
 #define LOCTEXT_NAMESPACE "ParametricSurfaceData"
 
-const double FUnitConverter::CentimeterToMillimeter = 0.1;
-const double FUnitConverter::MillimeterToCentimeter = 10.;
+const double FUnitConverter::CentimeterToMillimeter = 10.;
+const double FUnitConverter::MillimeterToCentimeter = 0.1;
 const double FUnitConverter::CentimeterToMeter = 0.01;
+const double FUnitConverter::MeterToCentimeter = 100.;
 const double FUnitConverter::MillimeterToMeter = 0.001;
+const double FUnitConverter::MeterToMillimeter = 1000.;
 
 bool UParametricSurfaceData::SetFromFile(const TCHAR* FilePath, bool bForTechSoft)
 {
@@ -49,7 +55,7 @@ void UParametricSurfaceData::Serialize(FArchive& Ar)
 	Ar << TechSoftRawData;
 }
 
-TSharedPtr<UE::CADKernel::FModel> UParametricSurfaceData::LoadModel()
+TSharedPtr<UE::CADKernel::FModel> UParametricSurfaceData::GetModel()
 {
 	using namespace UE::CADKernel;
 
@@ -64,5 +70,38 @@ TSharedPtr<UE::CADKernel::FModel> UParametricSurfaceData::LoadModel()
 	return Session->GetModelAsShared();
 }
 
+A3DRiRepresentationItem* UParametricSurfaceData::GetRepresentation()
+{
+	using namespace UE::CADKernel;
+
+	if (TechSoftRawData.IsEmpty())
+	{
+		return {};
+	}
+
+	return TechSoftUtilities::GetRepresentation(TechSoftRawData);
+}
+
+
+bool UParametricSurfaceData::Tessellate(UE::Geometry::FDynamicMesh3& MeshOut)
+{
+	return UE::CADKernel::Private::Tessellate(*this, MeshOut);
+}
+
+bool UParametricSurfaceData::Tessellate(FMeshDescription& MeshOut)
+{
+	return UE::CADKernel::Private::Tessellate(*this, MeshOut);
+}
+
+bool UParametricSurfaceData::Retessellate(const FCADKernelRetessellationSettings& Settings, UE::Geometry::FDynamicMesh3& MeshOut)
+{
+	return UE::CADKernel::Private::Retessellate(*this, Settings, MeshOut);
+}
+
+bool UParametricSurfaceData::Retessellate(const FCADKernelRetessellationSettings& Settings, FMeshDescription& MeshOut)
+{
+	return UE::CADKernel::Private::Retessellate(*this, Settings, MeshOut);
+}
+
 #undef LOCTEXT_NAMESPACE // "ParametricSurfaceData"
 
diff --git a/Engine/Source/Runtime/Datasmith/CADKernel/Engine/Private/TechSoft/TechSoftMesh.cpp b/Engine/Source/Runtime/Datasmith/CADKernel/Engine/Private/TechSoft/TechSoftMesh.cpp
--- a/Engine/Source/Runtime/Datasmith/CADKernel/Engine/Private/TechSoft/TechSoftMesh.cpp
+++ b/Engine/Source/Runtime/Datasmith/CADKernel/Engine/Private/TechSoft/TechSoftMesh.cpp
@@ -10,240 +10,86 @@
 namespace UE::CADKernel::TechSoftUtilities
 {
 	using namespace UE::CADKernel;
+	using namespace UE::CADKernel::MeshUtilities;
 
-	class FRepresentationToMeshHelper
+	class FFaceTriangleCollector
 	{
-	public:
-		FRepresentationToMeshHelper(const A3DRiRepresentationItem* RepresentationItem, MeshUtilities::FMeshWrapperAbstract& InMeshInterface)
-			: MeshInterface(InMeshInterface)
-		{
-			TechSoft::TUniqueObject<A3DRiRepresentationItemData> RepresentationItemData(RepresentationItem);
-			if (!RepresentationItemData.IsValid())
-			{
-				return;
-			}
-
-			A3DEEntityType Type;
-			A3DEntityGetType(RepresentationItemData->m_pTessBase, &Type);
-			if (Type != kA3DTypeTess3D)
-			{
-				return;
-			}
-
-			TessellationBase = RepresentationItemData->m_pTessBase;
-
-			if (TessellationBase)
-			{
-				A3DEntityGetType(RepresentationItem, &Type);
-				if (Type == kA3DTypeRiBrepModel)
-				{
-					TechSoft::TUniqueObject<A3DRiBrepModelData> BRepModelData(RepresentationItem);
-					if (BRepModelData.IsValid())
-					{
-						BRepData = BRepModelData->m_pBrepData;
-					}
-				}
-			}
-		}
-
-		bool IsValid() const { return TessellationBase != nullptr; }
-
-		bool Convert();
-
 	private:
-		struct FTriangleContext
-		{
-			int32 GroupID = 0;
-			uint32 MaterialID = 0;
-			FVector2f UVScale;
-		};
-
-		bool CollectBRepFaces(TArray<A3DTopoFace*>& BRepFaces, int32& TriangleCount);
-		int32 GetTriangleCount(const A3DTessFaceData& FaceTessData);
-		bool ParseTriangles(const TArray<A3DTopoFace*>& BRepFaces);
-
-		void GetUVScale(const A3DTopoFace* TopoFace, FVector2f& UVScale);
-
-		void AddFaceTriangle(const FTriangleContext& Context, const uint32 TriangleCount, uint32& StartIndexInOut);
-		void AddFaceTriangleWithUniqueNormal(const FTriangleContext& Context, const uint32 TriCount, uint32& StartIndexInOut);
-		void AddFaceTriangleWithUniqueNormalAndTexture(const FTriangleContext& Context, const uint32 TriCount, const uint32 TexOffset, uint32& StartIndexInOut);
-		void AddFaceTriangleWithTexture(const FTriangleContext& Context, const uint32 TriCount, const uint32 TexOffset, uint32& StartIndexInOut);
-		
-		void AddFaceTriangleFan(const FTriangleContext& Context, const uint32 TriCount, uint32& StartIndexInOut);
-		void AddFaceTriangleFanWithUniqueNormal(const FTriangleContext& Context, const uint32 TriCount, uint32& StartIndexInOut);
-		void AddFaceTriangleFanWithUniqueNormalAndTexture(const FTriangleContext& Context, const uint32 TriCount, const uint32 TexOffset, uint32& StartIndexInOut);
-		void AddFaceTriangleFanWithTexture(const FTriangleContext& Context, const uint32 TriCount, const uint32 TexOffset, uint32& StartIndexInOut);
+		const A3DUns32* TriangulatedIndexes = nullptr;
+		const TArrayView<FVector>& TessellationNormals;
+		const TArrayView<FVector2d>& TessellationTexCoords;
 		
-		void AddFaceTriangleStripe(const FTriangleContext& Context, const uint32 TriCount, uint32& StartIndexInOut);
-		void AddFaceTriangleStripeWithTexture(const FTriangleContext& Context, const uint32 TriCount, const uint32 TexOffset, uint32& StartIndexInOut);
-		void AddFaceTriangleStripeWithUniqueNormal(const FTriangleContext& Context, const uint32 TriCount, uint32& StartIndexInOut);
-		void AddFaceTriangleStripeWithUniqueNormalAndTexture(const FTriangleContext& Context, const uint32 TriCount, const uint32 TexOffset, uint32& StartIndexInOut);
-
-		void AddTriangle(const FTriangleContext& Context, const FArray3i& VertexIndices, const FArray3i& NormalIndices, const FArray3i& TexCoordIndicies)
-		{
-			FVector3f Normals[3]{
-				{(float)TessellationNormals[NormalIndices[0]], (float)TessellationNormals[NormalIndices[0] + 1], (float)TessellationNormals[NormalIndices[0] + 2]},
-				{(float)TessellationNormals[NormalIndices[1]], (float)TessellationNormals[NormalIndices[1] + 1], (float)TessellationNormals[NormalIndices[1] + 2]},
-				{(float)TessellationNormals[NormalIndices[2]], (float)TessellationNormals[NormalIndices[2] + 1], (float)TessellationNormals[NormalIndices[2] + 2]},
-			};
-			FVector2f TexCoords[3]{
-				{(float)TessellationTexCoords[TexCoordIndicies[0]] * Context.UVScale[0], (float)TessellationTexCoords[TexCoordIndicies[0] + 1] * Context.UVScale[1]},
-				{(float)TessellationTexCoords[TexCoordIndicies[1]] * Context.UVScale[0], (float)TessellationTexCoords[TexCoordIndicies[1] + 1] * Context.UVScale[1]},
-				{(float)TessellationTexCoords[TexCoordIndicies[2]] * Context.UVScale[0], (float)TessellationTexCoords[TexCoordIndicies[2] + 1] * Context.UVScale[1]},
-			};
-
-			MeshInterface.AddTriangle(Context.GroupID, Context.MaterialID, VertexIndices, {Normals, 3 }, {TexCoords, 3});
-		}
-
-	private:
-		MeshUtilities::FMeshWrapperAbstract& MeshInterface;
-		const A3DTess3D* TessellationBase = nullptr;
-		const A3DTopoBrepData* BRepData = nullptr;
-
 		int32 LastVertexIndex = 0;
-		A3DUns32* TriangulatedIndexes = nullptr;
-		double* TessellationNormals = nullptr;
-		double* TessellationTexCoords = nullptr;				/*!< Array of \ref A3DDouble, as texture coordinates. */
-	};
-
-	bool AddRepresentation(A3DRiRepresentationItem* RepresentationItem, MeshUtilities::FMeshWrapperAbstract& MeshInterface)
-	{
-		FRepresentationToMeshHelper Helper(RepresentationItem, MeshInterface);
-
-		return Helper.IsValid() ? Helper.Convert() : false;
-	}
-
-	bool FRepresentationToMeshHelper::CollectBRepFaces(TArray<A3DTopoFace*>& BRepFaces, int32& TriangleCount)
-	{
-		BRepFaces.Empty();
-		TriangleCount = 0;
-
-		if (BRepData == nullptr)
-		{
-			return false;
-		}
-
-		// #cad_import: Check what this is all about
-		//TechSoft::TUniqueObject<A3DTopoBodyData> TopoBodyData(BRepData);
-		//if (TopoBodyData.IsValid())
-		//{
-		//	if (TopoBodyData->m_pContext)
-		//	{
-		//		TechSoft::TUniqueObject<A3DTopoContextData> TopoContextData(TopoBodyData->m_pContext);
-		//		if (TopoContextData.IsValid())
-		//		{
-		//			if (TopoContextData->m_bHaveScale)
-		//			{
-		//				TextureUnit *= TopoContextData->m_dScale;
-		//			}
-		//		}
-		//	}
-		//}
+		int32 GroupID = INDEX_NONE;
+		uint32 MaterialID = INDEX_NONE;
+		TArray<FVector3f> FaceNormals;
+		TArray<FVector2f> FaceTexCoords;
+		TArray<FMeshWrapperAbstract::FFaceTriangle> FaceTriangles;
+		FVector2d UVScale;
 
-		TechSoft::TUniqueObject<A3DTopoBrepDataData> TopoBrepData(BRepData);
-		if (!TopoBrepData.IsValid())
+	public:
+		struct FCollectionResult
 		{
-			return false;
-		}
+			TArray<FVector3f> Normals;
+			TArray<FVector2f> TexCoords;
+			TArray<FMeshWrapperAbstract::FFaceTriangle> FaceTriangles;
+		};
 
-		for (A3DUns32 Index = 0; Index < TopoBrepData->m_uiConnexSize; ++Index)
+		FFaceTriangleCollector(const A3DUns32* InTriangulatedIndexes, const TArrayView<FVector>& InNormals, const TArrayView<FVector2d>& InTexCoords)
+			: TriangulatedIndexes(InTriangulatedIndexes)
+			, TessellationNormals(InNormals)
+			, TessellationTexCoords(InTexCoords)
 		{
-			TechSoft::TUniqueObject<A3DTopoConnexData> TopoConnexData(TopoBrepData->m_ppConnexes[Index]);
-			if (TopoConnexData.IsValid())
-			{
-				for (A3DUns32 Sndex = 0; Sndex < TopoConnexData->m_uiShellSize; ++Sndex)
-				{
-					TechSoft::TUniqueObject<A3DTopoShellData> ShellData(TopoConnexData->m_ppShells[Sndex]);
-					if (ShellData.IsValid())
-					{
-						for (A3DUns32 Fndex = 0; Fndex < ShellData->m_uiFaceSize; ++Fndex)
-						{
-							BRepFaces.Add(ShellData->m_ppFaces[Fndex]);
-						}
-					}
-				}
-			}
 		}
 
-		if (BRepFaces.IsEmpty())
-		{
-			return false;
-		}
+		bool CollectTriangles(const A3DTessFaceData& FaceTessData, const A3DTopoFace* TopoFace, int32 InGroupID, FCollectionResult& Result);
 
-		TechSoft::TUniqueObject<A3DTess3DData> TessellationData(TessellationBase);
+		void AddFaceTriangle(const uint32 TriangleCount, uint32& StartIndexInOut);
+		void AddFaceTriangleWithUniqueNormal(const uint32 TriCount, uint32& StartIndexInOut);
+		void AddFaceTriangleWithUniqueNormalAndTexture(const uint32 TriCount, const uint32 TexOffset, uint32& StartIndexInOut);
+		void AddFaceTriangleWithTexture(const uint32 TriCount, const uint32 TexOffset, uint32& StartIndexInOut);
 
-		if (!TessellationData.IsValid() || TessellationData->m_uiFaceTessSize == 0)
-		{
-			return false;
-		}
+		void AddFaceTriangleFan(const uint32 TriCount, uint32& StartIndexInOut);
+		void AddFaceTriangleFanWithUniqueNormal(const uint32 TriCount, uint32& StartIndexInOut);
+		void AddFaceTriangleFanWithUniqueNormalAndTexture(const uint32 TriCout, const uint32 TexOffset, uint32& StartIndexInOut);
+		void AddFaceTriangleFanWithTexture(const uint32 TriCount, const uint32 TexOffset, uint32& StartIndexInOut);
 
-		if (TessellationData->m_uiFaceTessSize != BRepFaces.Num())
-		{
-			BRepFaces.Empty();
-			return false;
-		}
+		void AddFaceTriangleStripe(const uint32 TriCount, uint32& StartIndexInOut);
+		void AddFaceTriangleStripeWithTexture(const uint32 TriCount, const uint32 TexOffset, uint32& StartIndexInOut);
+		void AddFaceTriangleStripeWithUniqueNormal(const uint32 TriCount, uint32& StartIndexInOut);
+		void AddFaceTriangleStripeWithUniqueNormalAndTexture(const uint32 TriCount, const uint32 TexOffset, uint32& StartIndexInOut);
 
-		for (A3DUns32 Index = 0; Index < TessellationData->m_uiFaceTessSize; ++Index)
+		void AddTriangle(const FArray3i& VertexIndices, const FArray3i& NormalIndices, const FArray3i& TexCoordIndicies)
 		{
-			TriangleCount += GetTriangleCount(TessellationData->m_psFaceTessData[Index]);
-		}
-
-		return TriangleCount > 0;
-	}
+			const int32 NextIndex = FaceNormals.Num();
+			const FArray3i OtherIndices{ NextIndex + 0, NextIndex + 1, NextIndex + 2 };
 
-	int32 FRepresentationToMeshHelper::GetTriangleCount(const A3DTessFaceData& FaceTessData)
-	{
-		constexpr int32 TessellationFaceDataWithTriangle   = 0x2222;
-		constexpr int32 TessellationFaceDataWithFan        = 0x4444;
-		constexpr int32 TessellationFaceDataWithStripe     = 0x8888;
-		constexpr int32 TessellationFaceDataWithOneNormal  = 0xE0E0;
-
-		uint32 UsedEntitiesFlags = FaceTessData.m_usUsedEntitiesFlags;
-
-		int32 TriangleCount = 0;
-		uint32 FaceSetIndex = 0;
-		
-		if (UsedEntitiesFlags & TessellationFaceDataWithTriangle)
-		{
-			TriangleCount += FaceTessData.m_puiSizesTriangulated[FaceSetIndex];
-			FaceSetIndex++;
-		}
+			const FVector3f Normals[3]{
+				{FVector3f(TessellationNormals[NormalIndices[0]])},
+				{FVector3f(TessellationNormals[NormalIndices[1]])},
+				{FVector3f(TessellationNormals[NormalIndices[2]])},
+			};
+			FaceNormals.Append(Normals, 3);
 
-		if (FaceTessData.m_uiSizesTriangulatedSize > FaceSetIndex)
-		{
-			if (UsedEntitiesFlags & TessellationFaceDataWithFan)
-			{
-				uint32 LastFanIndex = 1 + FaceSetIndex + FaceTessData.m_puiSizesTriangulated[FaceSetIndex];
-				FaceSetIndex++;
-				for (; FaceSetIndex < LastFanIndex; FaceSetIndex++)
-				{
-					int32 FanSize = (int32)(FaceTessData.m_puiSizesTriangulated[FaceSetIndex] & kA3DTessFaceDataNormalMask);
-					TriangleCount += (FanSize - 2);
-				}
-			}
-		}
+			const FVector2f TexCoords[3]{
+				{FVector2f(TessellationTexCoords[TexCoordIndicies[0]] * UVScale)},
+				{FVector2f(TessellationTexCoords[TexCoordIndicies[1]] * UVScale)},
+				{FVector2f(TessellationTexCoords[TexCoordIndicies[2]] * UVScale)},
+			};
+			FaceTexCoords.Append(TexCoords, 3);
 
-		if (FaceTessData.m_uiSizesTriangulatedSize > FaceSetIndex)
-		{
-			FaceSetIndex++;
-			for (; FaceSetIndex < FaceTessData.m_uiSizesTriangulatedSize; FaceSetIndex++)
-			{
-				int32 StripeSize = (int32)(FaceTessData.m_puiSizesTriangulated[FaceSetIndex] & kA3DTessFaceDataNormalMask);
-				TriangleCount += (StripeSize - 2);
-			}
+			FaceTriangles.Emplace(GroupID, MaterialID, VertexIndices, OtherIndices, OtherIndices);
 		}
+	};
 
-		return TriangleCount;
-	}
-
-	void FRepresentationToMeshHelper::AddFaceTriangle(const FTriangleContext& Context, const uint32 TriCount, uint32& StartIndexInOut)
+	void FFaceTriangleCollector::AddFaceTriangle(const uint32 TriCount, uint32& StartIndexInOut)
 	{
 		FArray3i VertexIndices;
 		FArray3i NormalIndices;
-		FArray3i TexCoordIndices(0, 0, 0); // #cadkernel_check Why no texture coordinates?
+		FArray3i TexCoordIndices(0, 0, 0);
 
 		// Get Triangles
-		for (uint32 FArray3i = 0; FArray3i < TriCount; FArray3i++)
+		for (uint32 Index = 0; Index < TriCount; ++Index)
 		{
 			NormalIndices[0] = TriangulatedIndexes[StartIndexInOut++];
 			VertexIndices[0] = TriangulatedIndexes[StartIndexInOut++] / 3;
@@ -252,18 +98,17 @@ namespace UE::CADKernel::TechSoftUtilities
 			NormalIndices[2] = TriangulatedIndexes[StartIndexInOut++];
 			VertexIndices[2] = TriangulatedIndexes[StartIndexInOut++] / 3;
 
-			AddTriangle(Context, VertexIndices, NormalIndices, TexCoordIndices);
+			AddTriangle(VertexIndices, NormalIndices, TexCoordIndices);
 		}
 	}
 
-	void FRepresentationToMeshHelper::AddFaceTriangleWithUniqueNormal(const FTriangleContext& Context, const uint32 TriCount, uint32& StartIndexInOut)
+	void FFaceTriangleCollector::AddFaceTriangleWithUniqueNormal(const uint32 TriCount, uint32& StartIndexInOut)
 	{
 		FArray3i VertexIndices;
 		FArray3i NormalIndices;
 		FArray3i TexCoordIndices(0, 0, 0);
 
-		// Get Triangles
-		for (uint32 FArray3i = 0; FArray3i < TriCount; FArray3i++)
+		for (uint32 Index = 0; Index < TriCount; ++Index)
 		{
 			NormalIndices[0] = TriangulatedIndexes[StartIndexInOut++];
 			NormalIndices[1] = NormalIndices[0];
@@ -273,11 +118,11 @@ namespace UE::CADKernel::TechSoftUtilities
 			VertexIndices[1] = TriangulatedIndexes[StartIndexInOut++] / 3;
 			VertexIndices[2] = TriangulatedIndexes[StartIndexInOut++] / 3;
 
-			AddTriangle(Context, VertexIndices, NormalIndices, TexCoordIndices);
+			AddTriangle(VertexIndices, NormalIndices, TexCoordIndices);
 		}
 	}
 
-	void FRepresentationToMeshHelper::AddFaceTriangleWithUniqueNormalAndTexture(const FTriangleContext& Context, const uint32 TriCount, const uint32 TexOffset, uint32& StartIndexInOut)
+	void FFaceTriangleCollector::AddFaceTriangleWithUniqueNormalAndTexture(const uint32 TriCount, const uint32 TexOffset, uint32& StartIndexInOut)
 	{
 		FArray3i VertexIndices;
 		FArray3i NormalIndices;
@@ -302,18 +147,17 @@ namespace UE::CADKernel::TechSoftUtilities
 			StartIndexInOut += TexOffset;
 			VertexIndices[2] = TriangulatedIndexes[StartIndexInOut++] / 3;
 
-			AddTriangle(Context, VertexIndices, NormalIndices, TexCoordIndices);
+			AddTriangle(VertexIndices, NormalIndices, TexCoordIndices);
 		}
 	}
 
-	void FRepresentationToMeshHelper::AddFaceTriangleWithTexture(const FTriangleContext& Context, const uint32 TriCount, const uint32 TexOffset, uint32& StartIndexInOut)
+	void FFaceTriangleCollector::AddFaceTriangleWithTexture(const uint32 TriCount, const uint32 TexOffset, uint32& StartIndexInOut)
 	{
 		FArray3i VertexIndices;
 		FArray3i NormalIndices;
 		FArray3i TexCoordIndices;
 
-		// Get Triangles
-		for (uint32 FArray3i = 0; FArray3i < TriCount; FArray3i++)
+		for (uint32 Index = 0; Index < TriCount; ++Index)
 		{
 			NormalIndices[0] = TriangulatedIndexes[StartIndexInOut++];
 			TexCoordIndices[0] = TriangulatedIndexes[StartIndexInOut];
@@ -330,11 +174,11 @@ namespace UE::CADKernel::TechSoftUtilities
 			StartIndexInOut += TexOffset;
 			VertexIndices[2] = TriangulatedIndexes[StartIndexInOut++] / 3;
 
-			AddTriangle(Context, VertexIndices, NormalIndices, TexCoordIndices);
+			AddTriangle(VertexIndices, NormalIndices, TexCoordIndices);
 		}
 	}
 
-	void FRepresentationToMeshHelper::AddFaceTriangleFan(const FTriangleContext& Context, const uint32 TriCount, uint32& StartIndexInOut)
+	void FFaceTriangleCollector::AddFaceTriangleFan(const uint32 TriCount, uint32& StartIndexInOut)
 	{
 		FArray3i VertexIndices;
 		FArray3i NormalIndices;
@@ -345,18 +189,19 @@ namespace UE::CADKernel::TechSoftUtilities
 		NormalIndices[1] = TriangulatedIndexes[StartIndexInOut++];
 		VertexIndices[1] = TriangulatedIndexes[StartIndexInOut++] / 3;
 
-		for (unsigned long FArray3i = 2; FArray3i < TriCount; FArray3i++)
+		for (uint32 Index = 2; Index < TriCount; ++Index)
 		{
 			NormalIndices[2] = TriangulatedIndexes[StartIndexInOut++];
 			VertexIndices[2] = TriangulatedIndexes[StartIndexInOut++] / 3;
 
-			AddTriangle(Context, VertexIndices, NormalIndices, TexCoordIndices);
+			AddTriangle(VertexIndices, NormalIndices, TexCoordIndices);
 
 			NormalIndices[1] = NormalIndices[2];
 			VertexIndices[1] = VertexIndices[2];
 		}
 	}
-	void FRepresentationToMeshHelper::AddFaceTriangleFanWithUniqueNormal(const FTriangleContext& Context, const uint32 TriCount, uint32& StartIndexInOut)
+
+	void FFaceTriangleCollector::AddFaceTriangleFanWithUniqueNormal(const uint32 TriCount, uint32& StartIndexInOut)
 	{
 		FArray3i VertexIndices;
 		FArray3i NormalIndices;
@@ -371,17 +216,17 @@ namespace UE::CADKernel::TechSoftUtilities
 		VertexIndices[1] = TriangulatedIndexes[StartIndexInOut++] / 3;
 
 		// Get Triangles
-		for (uint32 FArray3i = 2; FArray3i < TriCount; FArray3i++)
+		for (uint32 Index = 2; Index < TriCount; ++Index)
 		{
 			VertexIndices[2] = TriangulatedIndexes[StartIndexInOut++] / 3;
 
-			AddTriangle(Context, VertexIndices, NormalIndices, TexCoordIndices);
+			AddTriangle(VertexIndices, NormalIndices, TexCoordIndices);
 
 			VertexIndices[1] = VertexIndices[2];
 		}
 	}
 
-	void FRepresentationToMeshHelper::AddFaceTriangleFanWithUniqueNormalAndTexture(const FTriangleContext& Context, const uint32 TriCount, const uint32 TexOffset, uint32& StartIndexInOut)
+	void FFaceTriangleCollector::AddFaceTriangleFanWithUniqueNormalAndTexture(const uint32 TriCount, const uint32 TexOffset, uint32& StartIndexInOut)
 	{
 		FArray3i VertexIndices;
 		FArray3i NormalIndices;
@@ -399,20 +244,20 @@ namespace UE::CADKernel::TechSoftUtilities
 		StartIndexInOut += TexOffset;
 		VertexIndices[1] = TriangulatedIndexes[StartIndexInOut++] / 3;
 
-		for (uint32 FArray3i = 2; FArray3i < TriCount; FArray3i++)
+		for (uint32 Index = 2; Index < TriCount; ++Index)
 		{
 			TexCoordIndices[2] = TriangulatedIndexes[StartIndexInOut];
 			StartIndexInOut += TexOffset;
 			VertexIndices[2] = TriangulatedIndexes[StartIndexInOut++] / 3;
 
-			AddTriangle(Context, VertexIndices, NormalIndices, TexCoordIndices);
+			AddTriangle(VertexIndices, NormalIndices, TexCoordIndices);
 
 			VertexIndices[1] = VertexIndices[2];
 			TexCoordIndices[1] = TexCoordIndices[2];
 		}
 	}
 
-	void FRepresentationToMeshHelper::AddFaceTriangleFanWithTexture(const FTriangleContext& Context, const uint32 TriCount, const uint32 TexOffset, uint32& StartIndexInOut)
+	void FFaceTriangleCollector::AddFaceTriangleFanWithTexture(const uint32 TriCount, const uint32 TexOffset, uint32& StartIndexInOut)
 	{
 		FArray3i VertexIndices;
 		FArray3i NormalIndices;
@@ -428,14 +273,14 @@ namespace UE::CADKernel::TechSoftUtilities
 		StartIndexInOut += TexOffset;
 		VertexIndices[1] = TriangulatedIndexes[StartIndexInOut++] / 3;
 
-		for (uint32 FArray3i = 2; FArray3i < TriCount; FArray3i++)
+		for (uint32 Index = 2; Index < TriCount; ++Index)
 		{
 			NormalIndices[2] = TriangulatedIndexes[StartIndexInOut++];
 			TexCoordIndices[2] = TriangulatedIndexes[StartIndexInOut];
 			StartIndexInOut += TexOffset;
 			VertexIndices[2] = TriangulatedIndexes[StartIndexInOut++] / 3;
 
-			AddTriangle(Context, VertexIndices, NormalIndices, TexCoordIndices);
+			AddTriangle(VertexIndices, NormalIndices, TexCoordIndices);
 
 			NormalIndices[1] = NormalIndices[2];
 			TexCoordIndices[1] = TexCoordIndices[2];
@@ -443,7 +288,7 @@ namespace UE::CADKernel::TechSoftUtilities
 		}
 	}
 
-	void FRepresentationToMeshHelper::AddFaceTriangleStripe(const FTriangleContext& Context, const uint32 TriCount, uint32& StartIndexInOut)
+	void FFaceTriangleCollector::AddFaceTriangleStripe(const uint32 TriCount, uint32& StartIndexInOut)
 	{
 		FArray3i VertexIndices;
 		FArray3i NormalIndices;
@@ -454,15 +299,15 @@ namespace UE::CADKernel::TechSoftUtilities
 		NormalIndices[1] = TriangulatedIndexes[StartIndexInOut++];
 		VertexIndices[1] = TriangulatedIndexes[StartIndexInOut++] / 3;
 
-		for (unsigned long FArray3i = 2; FArray3i < TriCount; FArray3i++)
+		for (uint32 Index = 2; Index < TriCount; ++Index)
 		{
 			NormalIndices[2] = TriangulatedIndexes[StartIndexInOut++];
 			VertexIndices[2] = TriangulatedIndexes[StartIndexInOut++] / 3;
 
-			AddTriangle(Context, VertexIndices, NormalIndices, TexCoordIndices);
+			AddTriangle(VertexIndices, NormalIndices, TexCoordIndices);
 
-			FArray3i++;
-			if (FArray3i == TriCount)
+			Index++;
+			if (Index == TriCount)
 			{
 				break;
 			}
@@ -473,14 +318,14 @@ namespace UE::CADKernel::TechSoftUtilities
 			NormalIndices[0] = TriangulatedIndexes[StartIndexInOut++];
 			VertexIndices[0] = TriangulatedIndexes[StartIndexInOut++] / 3;
 
-			AddTriangle(Context, VertexIndices, NormalIndices, TexCoordIndices);
+			AddTriangle(VertexIndices, NormalIndices, TexCoordIndices);
 
 			Swap(VertexIndices[0], VertexIndices[1]);
 			Swap(NormalIndices[0], NormalIndices[1]);
 		}
 	}
 
-	void FRepresentationToMeshHelper::AddFaceTriangleStripeWithTexture(const FTriangleContext& Context, const uint32 TriCount, const uint32 TexOffset, uint32& StartIndexInOut)
+	void FFaceTriangleCollector::AddFaceTriangleStripeWithTexture(const uint32 TriCount, const uint32 TexOffset, uint32& StartIndexInOut)
 	{
 		FArray3i VertexIndices;
 		FArray3i NormalIndices;
@@ -495,17 +340,17 @@ namespace UE::CADKernel::TechSoftUtilities
 		StartIndexInOut += TexOffset;
 		VertexIndices[1] = TriangulatedIndexes[StartIndexInOut++] / 3;
 
-		for (unsigned long FArray3i = 2; FArray3i < TriCount; FArray3i++)
+		for (uint32 Index = 2; Index < TriCount; ++Index)
 		{
 			NormalIndices[2] = TriangulatedIndexes[StartIndexInOut++];
 			TexCoordIndices[2] = TriangulatedIndexes[StartIndexInOut];
 			StartIndexInOut += TexOffset;
 			VertexIndices[2] = TriangulatedIndexes[StartIndexInOut++] / 3;
 
-			AddTriangle(Context, VertexIndices, NormalIndices, TexCoordIndices);
+			AddTriangle(VertexIndices, NormalIndices, TexCoordIndices);
 
-			FArray3i++;
-			if (FArray3i == TriCount)
+			Index++;
+			if (Index == TriCount)
 			{
 				break;
 			}
@@ -517,7 +362,7 @@ namespace UE::CADKernel::TechSoftUtilities
 			NormalIndices[0] = TriangulatedIndexes[StartIndexInOut++];
 			VertexIndices[0] = TriangulatedIndexes[StartIndexInOut++] / 3;
 
-			AddTriangle(Context, VertexIndices, NormalIndices, TexCoordIndices);
+			AddTriangle(VertexIndices, NormalIndices, TexCoordIndices);
 
 			Swap(VertexIndices[0], VertexIndices[1]);
 			Swap(NormalIndices[0], NormalIndices[1]);
@@ -525,7 +370,7 @@ namespace UE::CADKernel::TechSoftUtilities
 		}
 	}
 
-	void FRepresentationToMeshHelper::AddFaceTriangleStripeWithUniqueNormal(const FTriangleContext& Context, const uint32 TriCount, uint32& StartIndexInOut)
+	void FFaceTriangleCollector::AddFaceTriangleStripeWithUniqueNormal(const uint32 TriCount, uint32& StartIndexInOut)
 	{
 		FArray3i VertexIndices;
 		FArray3i NormalIndices;
@@ -538,14 +383,14 @@ namespace UE::CADKernel::TechSoftUtilities
 		VertexIndices[0] = TriangulatedIndexes[StartIndexInOut++] / 3;
 		VertexIndices[1] = TriangulatedIndexes[StartIndexInOut++] / 3;
 
-		for (unsigned long FArray3i = 2; FArray3i < TriCount; FArray3i++)
+		for (uint32 Index = 2; Index < TriCount; ++Index)
 		{
 			VertexIndices[2] = TriangulatedIndexes[StartIndexInOut++] / 3;
 
-			AddTriangle(Context, VertexIndices, NormalIndices, TexCoordIndices);
+			AddTriangle(VertexIndices, NormalIndices, TexCoordIndices);
 
-			FArray3i++;
-			if (FArray3i == TriCount)
+			Index++;
+			if (Index == TriCount)
 			{
 				break;
 			}
@@ -554,13 +399,13 @@ namespace UE::CADKernel::TechSoftUtilities
 
 			VertexIndices[0] = TriangulatedIndexes[StartIndexInOut++] / 3;
 
-			AddTriangle(Context, VertexIndices, NormalIndices, TexCoordIndices);
+			AddTriangle(VertexIndices, NormalIndices, TexCoordIndices);
 
 			Swap(VertexIndices[0], VertexIndices[1]);
 		}
 	}
 
-	void FRepresentationToMeshHelper::AddFaceTriangleStripeWithUniqueNormalAndTexture(const FTriangleContext& Context, const uint32 TriCount, const uint32 TexOffset, uint32& StartIndexInOut)
+	void FFaceTriangleCollector::AddFaceTriangleStripeWithUniqueNormalAndTexture(const uint32 TriCount, const uint32 TexOffset, uint32& StartIndexInOut)
 	{
 		FArray3i VertexIndices;
 		FArray3i NormalIndices;
@@ -578,16 +423,16 @@ namespace UE::CADKernel::TechSoftUtilities
 		StartIndexInOut += TexOffset;
 		VertexIndices[1] = TriangulatedIndexes[StartIndexInOut++] / 3;
 
-		for (unsigned long FArray3i = 2; FArray3i < TriCount; FArray3i++)
+		for (uint32 Index = 2; Index < TriCount; ++Index)
 		{
 			TexCoordIndices[2] = TriangulatedIndexes[StartIndexInOut];
 			StartIndexInOut += TexOffset;
 			VertexIndices[2] = TriangulatedIndexes[StartIndexInOut++] / 3;
 
-			AddTriangle(Context, VertexIndices, NormalIndices, TexCoordIndices);
+			AddTriangle(VertexIndices, NormalIndices, TexCoordIndices);
 
-			FArray3i++;
-			if (FArray3i == TriCount)
+			++Index;
+			if (Index == TriCount)
 			{
 				break;
 			}
@@ -599,190 +444,395 @@ namespace UE::CADKernel::TechSoftUtilities
 			StartIndexInOut += TexOffset;
 			VertexIndices[0] = TriangulatedIndexes[StartIndexInOut++] / 3;
 
-			AddTriangle(Context, VertexIndices, NormalIndices, TexCoordIndices);
+			AddTriangle(VertexIndices, NormalIndices, TexCoordIndices);
 
 			Swap(VertexIndices[0], VertexIndices[1]);
 			Swap(TexCoordIndices[0], TexCoordIndices[1]);
 		}
 	}
 
-	void FRepresentationToMeshHelper::GetUVScale(const A3DTopoFace* TopoFace, FVector2f& UVScale)
+	bool FFaceTriangleCollector::CollectTriangles(const A3DTessFaceData& FaceTessData, const A3DTopoFace* TopoFace, int32 InGroupID, FFaceTriangleCollector::FCollectionResult& Result)
 	{
-		// #cad_import: Something to solve based on how this is handle in legacy
-		static FVector2f OneVector(1.f, 1.f);
-		UVScale = OneVector;
-	}
+		FaceNormals.Empty(FaceNormals.Num());
+		FaceTexCoords.Empty(FaceTexCoords.Num());
+		FaceTriangles.Empty(FaceTriangles.Num());
+		UVScale = FVector2d::UnitVector;
+
+		GroupID = InGroupID;
+		MaterialID = 0;
+		if (FaceTessData.m_uiStyleIndexesSize != 0)
+		{
+			// Store the StyleIndex on the MaterialName. It will be processed after tessellation
+			MaterialID = FaceTessData.m_puiStyleIndexes[0];
+		}
 
+		double TextureUnit = 1.;
+		// #cad_import: Where TextureUnit should come from?
+		UVScale = GetUVScale(TopoFace, TextureUnit);
 
-	bool FRepresentationToMeshHelper::ParseTriangles(const TArray<A3DTopoFace*>& BRepFaces)
-	{
-		TechSoft::TUniqueObject<A3DTess3DData> TessellationData(TessellationBase);
+		LastVertexIndex = 0;
+		uint32 FaceSetIndex = 0;
+		bool bMustProcess = true;
 
-		TessellationNormals = TessellationData->m_pdNormals;
-		TessellationTexCoords = TessellationData->m_pdTextureCoords;
-		TriangulatedIndexes = TessellationData->m_puiTriangulatedIndexes;
+		uint32 LastTrianguleIndex = FaceTessData.m_uiStartTriangulated;
+		uint32 UsedEntitiesFlags = FaceTessData.m_usUsedEntitiesFlags;
+		LastVertexIndex = 0;
 
-		FTriangleContext Context;
+		if (bMustProcess && UsedEntitiesFlags & kA3DTessFaceDataTriangle)
+		{
+			AddFaceTriangle(FaceTessData.m_puiSizesTriangulated[FaceSetIndex++], LastTrianguleIndex);
+			bMustProcess = FaceTessData.m_uiSizesTriangulatedSize > FaceSetIndex;
+		}
 
-		for (A3DUns32 Index = 0; Index < TessellationData->m_uiFaceTessSize; ++Index)
+		if (bMustProcess && UsedEntitiesFlags & kA3DTessFaceDataTriangleOneNormal)
 		{
-			const A3DTessFaceData& FaceTessData = TessellationData->m_psFaceTessData[Index];
+			AddFaceTriangleWithUniqueNormal(FaceTessData.m_puiSizesTriangulated[FaceSetIndex++], LastTrianguleIndex);
+			bMustProcess = FaceTessData.m_uiSizesTriangulatedSize > FaceSetIndex;
+		}
 
-			if (FaceTessData.m_uiSizesTriangulatedSize == 0)
-			{
-				continue;
-			}
+		if (bMustProcess && UsedEntitiesFlags & kA3DTessFaceDataTriangleTextured)
+		{
+			AddFaceTriangleWithTexture(FaceTessData.m_puiSizesTriangulated[FaceSetIndex++], FaceTessData.m_uiTextureCoordIndexesSize, LastTrianguleIndex);
+			bMustProcess = FaceTessData.m_uiSizesTriangulatedSize > FaceSetIndex;
+		}
 
-			uint32 FaceSetIndex = 0;
-			bool bMustProcess = true;
+		if (bMustProcess && UsedEntitiesFlags & kA3DTessFaceDataTriangleOneNormalTextured)
+		{
+			AddFaceTriangleWithUniqueNormalAndTexture(FaceTessData.m_puiSizesTriangulated[FaceSetIndex++], FaceTessData.m_uiTextureCoordIndexesSize, LastTrianguleIndex);
+			bMustProcess = FaceTessData.m_uiSizesTriangulatedSize > FaceSetIndex;
+		}
 
-			Context.GroupID = Index;
-			Context.MaterialID = 0;
-			if (FaceTessData.m_uiStyleIndexesSize != 0)
+		if (bMustProcess && UsedEntitiesFlags & kA3DTessFaceDataTriangleFan)
+		{
+			uint32 FanCount = FaceTessData.m_puiSizesTriangulated[FaceSetIndex++];
+			for (uint32 FanIndex = 0; FanIndex < FanCount; ++FanIndex)
 			{
-				// Store the StyleIndex on the MaterialName. It will be processed after tessellation
-				Context.MaterialID = FaceTessData.m_puiStyleIndexes[0];
+				uint32 VertexCount = FaceTessData.m_puiSizesTriangulated[FaceSetIndex++];
+				AddFaceTriangleFan(VertexCount, LastTrianguleIndex);
 			}
+			bMustProcess = FaceTessData.m_uiSizesTriangulatedSize > FaceSetIndex;
+		}
 
-			GetUVScale(BRepFaces[Index], Context.UVScale);
-
-			uint32 UsedEntitiesFlags = FaceTessData.m_usUsedEntitiesFlags;
-
-			uint32 LastTrianguleIndex = FaceTessData.m_uiStartTriangulated;
-			LastVertexIndex = 0;
-
-			if (bMustProcess && UsedEntitiesFlags & kA3DTessFaceDataTriangle)
+		if (bMustProcess && UsedEntitiesFlags & kA3DTessFaceDataTriangleFanOneNormal)
+		{
+			uint32 FanCount = FaceTessData.m_puiSizesTriangulated[FaceSetIndex++] & kA3DTessFaceDataNormalMask;
+			for (uint32 FanIndex = 0; FanIndex < FanCount; ++FanIndex)
 			{
-				AddFaceTriangle(Context, FaceTessData.m_puiSizesTriangulated[FaceSetIndex++], LastTrianguleIndex);
-				bMustProcess = FaceTessData.m_uiSizesTriangulatedSize > FaceSetIndex;
+				ensure((FaceTessData.m_puiSizesTriangulated[FaceSetIndex] & kA3DTessFaceDataNormalSingle) != 0);
+
+				uint32 VertexCount = FaceTessData.m_puiSizesTriangulated[FaceSetIndex++] & kA3DTessFaceDataNormalMask;
+				AddFaceTriangleFanWithUniqueNormal(VertexCount, LastTrianguleIndex);
 			}
+			bMustProcess = FaceTessData.m_uiSizesTriangulatedSize > FaceSetIndex;
+		}
 
-			if (bMustProcess && UsedEntitiesFlags & kA3DTessFaceDataTriangleOneNormal)
+		if (bMustProcess && UsedEntitiesFlags & kA3DTessFaceDataTriangleFanTextured)
+		{
+			uint32 FanCount = FaceTessData.m_puiSizesTriangulated[FaceSetIndex++];
+			for (uint32 FanIndex = 0; FanIndex < FanCount; ++FanIndex)
 			{
-				AddFaceTriangleWithUniqueNormal(Context, FaceTessData.m_puiSizesTriangulated[FaceSetIndex++], LastTrianguleIndex);
-				bMustProcess = FaceTessData.m_uiSizesTriangulatedSize > FaceSetIndex;
+				uint32 VertexCount = FaceTessData.m_puiSizesTriangulated[FaceSetIndex++];
+				AddFaceTriangleFanWithTexture(VertexCount, FaceTessData.m_uiTextureCoordIndexesSize, LastTrianguleIndex);
 			}
+			bMustProcess = FaceTessData.m_uiSizesTriangulatedSize > FaceSetIndex;
+		}
 
-			if (bMustProcess && UsedEntitiesFlags & kA3DTessFaceDataTriangleTextured)
+		if (bMustProcess && UsedEntitiesFlags & kA3DTessFaceDataTriangleFanOneNormalTextured)
+		{
+			uint32 FanCount = FaceTessData.m_puiSizesTriangulated[FaceSetIndex++] & kA3DTessFaceDataNormalMask;
+			for (uint32 FanIndex = 0; FanIndex < FanCount; ++FanIndex)
 			{
-				AddFaceTriangleWithTexture(Context, FaceTessData.m_puiSizesTriangulated[FaceSetIndex++], FaceTessData.m_uiTextureCoordIndexesSize, LastTrianguleIndex);
-				bMustProcess = FaceTessData.m_uiSizesTriangulatedSize > FaceSetIndex;
+				ensure((FaceTessData.m_puiSizesTriangulated[FaceSetIndex] & kA3DTessFaceDataNormalSingle) != 0);
+
+				uint32 VertexCount = FaceTessData.m_puiSizesTriangulated[FaceSetIndex++] & kA3DTessFaceDataNormalMask;
+				AddFaceTriangleFanWithUniqueNormalAndTexture(VertexCount, FaceTessData.m_uiTextureCoordIndexesSize, LastTrianguleIndex);
 			}
+			bMustProcess = FaceTessData.m_uiSizesTriangulatedSize > FaceSetIndex;
+		}
 
-			if (bMustProcess && UsedEntitiesFlags & kA3DTessFaceDataTriangleOneNormalTextured)
+		if (bMustProcess && UsedEntitiesFlags & kA3DTessFaceDataTriangleStripe)
+		{
+			A3DUns32 StripeSize = FaceTessData.m_puiSizesTriangulated[FaceSetIndex++];
+			for (A3DUns32 StripeIndex = 0; StripeIndex < StripeSize; ++StripeIndex)
 			{
-				AddFaceTriangleWithUniqueNormalAndTexture(Context, FaceTessData.m_puiSizesTriangulated[FaceSetIndex++], FaceTessData.m_uiTextureCoordIndexesSize, LastTrianguleIndex);
-				bMustProcess = FaceTessData.m_uiSizesTriangulatedSize > FaceSetIndex;
+				A3DUns32 PointCount = FaceTessData.m_puiSizesTriangulated[FaceSetIndex++];
+				AddFaceTriangleStripe(PointCount, LastTrianguleIndex);
 			}
+			bMustProcess = FaceTessData.m_uiSizesTriangulatedSize > FaceSetIndex;
+		}
 
-			if (bMustProcess && UsedEntitiesFlags & kA3DTessFaceDataTriangleFan)
+		if (bMustProcess && UsedEntitiesFlags & kA3DTessFaceDataTriangleStripeOneNormal)
+		{
+			A3DUns32 StripeSize = FaceTessData.m_puiSizesTriangulated[FaceSetIndex++] & kA3DTessFaceDataNormalMask;
+			for (A3DUns32 StripeIndex = 0; StripeIndex < StripeSize; ++StripeIndex)
 			{
-				uint32 FanCount = FaceTessData.m_puiSizesTriangulated[FaceSetIndex++];
-				for (uint32 FanIndex = 0; FanIndex < FanCount; ++FanIndex)
+				bool bIsOneNormal = (FaceTessData.m_puiSizesTriangulated[FaceSetIndex] & kA3DTessFaceDataNormalSingle) != 0;
+				A3DUns32 PointCount = FaceTessData.m_puiSizesTriangulated[FaceSetIndex++] & kA3DTessFaceDataNormalMask;
+
+				// Is there only one normal for the entire stripe?
+				if (bIsOneNormal == false)
 				{
-					uint32 VertexCount = FaceTessData.m_puiSizesTriangulated[FaceSetIndex++];
-					AddFaceTriangleFan(Context, VertexCount, LastTrianguleIndex);
+					AddFaceTriangleStripe(PointCount, LastTrianguleIndex);
+					continue;
 				}
-				bMustProcess = FaceTessData.m_uiSizesTriangulatedSize > FaceSetIndex;
+
+				AddFaceTriangleStripeWithUniqueNormal(PointCount, LastTrianguleIndex);
 			}
+			bMustProcess = FaceTessData.m_uiSizesTriangulatedSize > FaceSetIndex;
+		}
 
-			if (bMustProcess && UsedEntitiesFlags & kA3DTessFaceDataTriangleFanOneNormal)
+		if (bMustProcess && UsedEntitiesFlags & kA3DTessFaceDataTriangleStripeTextured)
+		{
+			A3DUns32 StripeSize = FaceTessData.m_puiSizesTriangulated[FaceSetIndex++];
+			for (A3DUns32 StripeIndex = 0; StripeIndex < StripeSize; ++StripeIndex)
 			{
-				uint32 FanCount = FaceTessData.m_puiSizesTriangulated[FaceSetIndex++] & kA3DTessFaceDataNormalMask;
-				for (uint32 FanIndex = 0; FanIndex < FanCount; ++FanIndex)
-				{
-					ensure((FaceTessData.m_puiSizesTriangulated[FaceSetIndex] & kA3DTessFaceDataNormalSingle) != 0);
-
-					uint32 VertexCount = FaceTessData.m_puiSizesTriangulated[FaceSetIndex++] & kA3DTessFaceDataNormalMask;
-					AddFaceTriangleFanWithUniqueNormal(Context, VertexCount, LastTrianguleIndex);
-				}
-				bMustProcess = FaceTessData.m_uiSizesTriangulatedSize > FaceSetIndex;
+				A3DUns32 PointCount = FaceTessData.m_puiSizesTriangulated[FaceSetIndex++];
+				AddFaceTriangleStripeWithTexture(PointCount, FaceTessData.m_uiTextureCoordIndexesSize, LastTrianguleIndex);
 			}
+			bMustProcess = FaceTessData.m_uiSizesTriangulatedSize > FaceSetIndex;
+		}
 
-			if (bMustProcess && UsedEntitiesFlags & kA3DTessFaceDataTriangleFanTextured)
+		if (bMustProcess && UsedEntitiesFlags & kA3DTessFaceDataTriangleStripeOneNormalTextured)
+		{
+			A3DUns32 StripeSize = FaceTessData.m_puiSizesTriangulated[FaceSetIndex++] & kA3DTessFaceDataNormalMask;
+			for (A3DUns32 StripeIndex = 0; StripeIndex < StripeSize; ++StripeIndex)
 			{
-				uint32 FanCount = FaceTessData.m_puiSizesTriangulated[FaceSetIndex++];
-				for (uint32 FanIndex = 0; FanIndex < FanCount; ++FanIndex)
-				{
-					uint32 VertexCount = FaceTessData.m_puiSizesTriangulated[FaceSetIndex++];
-					AddFaceTriangleFanWithTexture(Context, VertexCount, FaceTessData.m_uiTextureCoordIndexesSize, LastTrianguleIndex);
-				}
-				bMustProcess = FaceTessData.m_uiSizesTriangulatedSize > FaceSetIndex;
+				A3DUns32 PointCount = FaceTessData.m_puiSizesTriangulated[FaceSetIndex++] & kA3DTessFaceDataNormalMask;
+				AddFaceTriangleStripeWithUniqueNormalAndTexture(PointCount, FaceTessData.m_uiTextureCoordIndexesSize, LastTrianguleIndex);
 			}
+			bMustProcess = FaceTessData.m_uiSizesTriangulatedSize > FaceSetIndex;
+		}
+
+		ensure(!bMustProcess);
+
+		Result.Normals = MoveTemp(FaceNormals);
+		Result.TexCoords = MoveTemp(FaceTexCoords);
+		Result.FaceTriangles = MoveTemp(FaceTriangles);
+
+		return true;
+	}
 
-			if (bMustProcess && UsedEntitiesFlags & kA3DTessFaceDataTriangleFanOneNormalTextured)
+	class FRepresentationConverter
+	{
+	public:
+		FRepresentationConverter(const A3DRiRepresentationItem* RepresentationItem, MeshUtilities::FMeshWrapperAbstract& InMeshWrapper)
+			: MeshWrapper(InMeshWrapper)
+		{
+			TechSoft::TUniqueObject<A3DRiRepresentationItemData> RepresentationItemData(RepresentationItem);
+			if (!RepresentationItemData.IsValid())
 			{
-				uint32 FanCount = FaceTessData.m_puiSizesTriangulated[FaceSetIndex++] & kA3DTessFaceDataNormalMask;
-				for (uint32 FanIndex = 0; FanIndex < FanCount; ++FanIndex)
-				{
-					ensure((FaceTessData.m_puiSizesTriangulated[FaceSetIndex] & kA3DTessFaceDataNormalSingle) != 0);
+				return;
+			}
 
-					uint32 VertexCount = FaceTessData.m_puiSizesTriangulated[FaceSetIndex++] & kA3DTessFaceDataNormalMask;
-					AddFaceTriangleFanWithUniqueNormalAndTexture(Context, VertexCount, FaceTessData.m_uiTextureCoordIndexesSize, LastTrianguleIndex);
-				}
-				bMustProcess = FaceTessData.m_uiSizesTriangulatedSize > FaceSetIndex;
+			A3DEEntityType Type;
+			A3DEntityGetType(RepresentationItemData->m_pTessBase, &Type);
+			if (Type != kA3DTypeTess3D)
+			{
+				return;
 			}
 
-			if (bMustProcess && UsedEntitiesFlags & kA3DTessFaceDataTriangleStripe)
+			TessellationBase = RepresentationItemData->m_pTessBase;
+
+			if (TessellationBase)
 			{
-				A3DUns32 StripeSize = FaceTessData.m_puiSizesTriangulated[FaceSetIndex++];
-				for (A3DUns32 StripeIndex = 0; StripeIndex < StripeSize; ++StripeIndex)
+				A3DEntityGetType(RepresentationItem, &Type);
+				if (Type == kA3DTypeRiBrepModel)
 				{
-					A3DUns32 PointCount = FaceTessData.m_puiSizesTriangulated[FaceSetIndex++];
-					AddFaceTriangleStripe(Context, PointCount, LastTrianguleIndex);
+					TechSoft::TUniqueObject<A3DRiBrepModelData> BRepModelData(RepresentationItem);
+					if (BRepModelData.IsValid())
+					{
+						BRepData = BRepModelData->m_pBrepData;
+					}
 				}
-				bMustProcess = FaceTessData.m_uiSizesTriangulatedSize > FaceSetIndex;
 			}
+		}
 
-			if (bMustProcess && UsedEntitiesFlags & kA3DTessFaceDataTriangleStripeOneNormal)
+		bool IsValid() const { return TessellationBase != nullptr; }
+
+		bool Convert();
+
+	private:
+		struct FTriangleContext
+		{
+			int32 GroupID = 0;
+			uint32 MaterialID = 0;
+			FVector2f UVScale;
+		};
+
+		bool CollectBRepFaces(TArray<A3DTopoFace*>& BRepFaces, int32& TriangleCount);
+		int32 GetTriangleCount(const A3DTessFaceData& FaceTessData);
+		bool ParseTriangles(const TArray<A3DTopoFace*>& BRepFaces);
+
+	private:
+		MeshUtilities::FMeshWrapperAbstract& MeshWrapper;
+		const A3DTess3D* TessellationBase = nullptr;
+		const A3DTopoBrepData* BRepData = nullptr;
+	};
+
+	bool FRepresentationConverter::CollectBRepFaces(TArray<A3DTopoFace*>& BRepFaces, int32& TriangleCount)
+	{
+		BRepFaces.Empty();
+		TriangleCount = 0;
+
+		if (BRepData == nullptr)
+		{
+			return false;
+		}
+
+		// #cad_import: Check what this is all about
+		//TechSoft::TUniqueObject<A3DTopoBodyData> TopoBodyData(BRepData);
+		//if (TopoBodyData.IsValid())
+		//{
+		//	if (TopoBodyData->m_pContext)
+		//	{
+		//		TechSoft::TUniqueObject<A3DTopoContextData> TopoContextData(TopoBodyData->m_pContext);
+		//		if (TopoContextData.IsValid())
+		//		{
+		//			if (TopoContextData->m_bHaveScale)
+		//			{
+		//				TextureUnit *= TopoContextData->m_dScale;
+		//			}
+		//		}
+		//	}
+		//}
+
+		TechSoft::TUniqueObject<A3DTopoBrepDataData> TopoBrepData(BRepData);
+		if (!TopoBrepData.IsValid())
+		{
+			return false;
+		}
+
+		for (A3DUns32 Index = 0; Index < TopoBrepData->m_uiConnexSize; ++Index)
+		{
+			TechSoft::TUniqueObject<A3DTopoConnexData> TopoConnexData(TopoBrepData->m_ppConnexes[Index]);
+			if (TopoConnexData.IsValid())
 			{
-				A3DUns32 StripeSize = FaceTessData.m_puiSizesTriangulated[FaceSetIndex++] & kA3DTessFaceDataNormalMask;
-				for (A3DUns32 StripeIndex = 0; StripeIndex < StripeSize; ++StripeIndex)
+				for (A3DUns32 Sndex = 0; Sndex < TopoConnexData->m_uiShellSize; ++Sndex)
 				{
-					bool bIsOneNormal = (FaceTessData.m_puiSizesTriangulated[FaceSetIndex] & kA3DTessFaceDataNormalSingle) != 0;
-					A3DUns32 PointCount = FaceTessData.m_puiSizesTriangulated[FaceSetIndex++] & kA3DTessFaceDataNormalMask;
-
-					// Is there only one normal for the entire stripe?
-					if (bIsOneNormal == false)
+					TechSoft::TUniqueObject<A3DTopoShellData> ShellData(TopoConnexData->m_ppShells[Sndex]);
+					if (ShellData.IsValid())
 					{
-						AddFaceTriangleStripe(Context, PointCount, LastTrianguleIndex);
-						continue;
+						for (A3DUns32 Fndex = 0; Fndex < ShellData->m_uiFaceSize; ++Fndex)
+						{
+							BRepFaces.Add(ShellData->m_ppFaces[Fndex]);
+						}
 					}
-
-					AddFaceTriangleStripeWithUniqueNormal(Context, PointCount, LastTrianguleIndex);
 				}
-				bMustProcess = FaceTessData.m_uiSizesTriangulatedSize > FaceSetIndex;
 			}
+		}
+
+		if (BRepFaces.IsEmpty())
+		{
+			return false;
+		}
+
+		TechSoft::TUniqueObject<A3DTess3DData> TessellationData(TessellationBase);
+
+		if (!TessellationData.IsValid() || TessellationData->m_uiFaceTessSize == 0)
+		{
+			return false;
+		}
+
+		if (TessellationData->m_uiFaceTessSize != BRepFaces.Num())
+		{
+			BRepFaces.Empty();
+			return false;
+		}
+
+		for (A3DUns32 Index = 0; Index < TessellationData->m_uiFaceTessSize; ++Index)
+		{
+			TriangleCount += GetTriangleCount(TessellationData->m_psFaceTessData[Index]);
+		}
+
+		return TriangleCount > 0;
+	}
 
-			if (bMustProcess && UsedEntitiesFlags & kA3DTessFaceDataTriangleStripeTextured)
+	int32 FRepresentationConverter::GetTriangleCount(const A3DTessFaceData& FaceTessData)
+	{
+		constexpr int32 TessellationFaceDataWithTriangle   = 0x2222;
+		constexpr int32 TessellationFaceDataWithFan        = 0x4444;
+		constexpr int32 TessellationFaceDataWithStripe     = 0x8888;
+		constexpr int32 TessellationFaceDataWithOneNormal  = 0xE0E0;
+
+		uint32 UsedEntitiesFlags = FaceTessData.m_usUsedEntitiesFlags;
+
+		int32 TriangleCount = 0;
+		uint32 FaceSetIndex = 0;
+		
+		if (UsedEntitiesFlags & TessellationFaceDataWithTriangle)
+		{
+			TriangleCount += FaceTessData.m_puiSizesTriangulated[FaceSetIndex];
+			FaceSetIndex++;
+		}
+
+		if (FaceTessData.m_uiSizesTriangulatedSize > FaceSetIndex)
+		{
+			if (UsedEntitiesFlags & TessellationFaceDataWithFan)
 			{
-				A3DUns32 StripeSize = FaceTessData.m_puiSizesTriangulated[FaceSetIndex++];
-				for (A3DUns32 StripeIndex = 0; StripeIndex < StripeSize; ++StripeIndex)
+				uint32 LastFanIndex = 1 + FaceSetIndex + FaceTessData.m_puiSizesTriangulated[FaceSetIndex];
+				FaceSetIndex++;
+				for (; FaceSetIndex < LastFanIndex; FaceSetIndex++)
 				{
-					A3DUns32 PointCount = FaceTessData.m_puiSizesTriangulated[FaceSetIndex++];
-					AddFaceTriangleStripeWithTexture(Context, PointCount, FaceTessData.m_uiTextureCoordIndexesSize, LastTrianguleIndex);
+					int32 FanSize = (int32)(FaceTessData.m_puiSizesTriangulated[FaceSetIndex] & kA3DTessFaceDataNormalMask);
+					TriangleCount += (FanSize - 2);
 				}
-				bMustProcess = FaceTessData.m_uiSizesTriangulatedSize > FaceSetIndex;
 			}
+		}
 
-			if (bMustProcess && UsedEntitiesFlags & kA3DTessFaceDataTriangleStripeOneNormalTextured)
+		if (FaceTessData.m_uiSizesTriangulatedSize > FaceSetIndex)
+		{
+			FaceSetIndex++;
+			for (; FaceSetIndex < FaceTessData.m_uiSizesTriangulatedSize; FaceSetIndex++)
 			{
-				A3DUns32 StripeSize = FaceTessData.m_puiSizesTriangulated[FaceSetIndex++] & kA3DTessFaceDataNormalMask;
-				for (A3DUns32 StripeIndex = 0; StripeIndex < StripeSize; ++StripeIndex)
-				{
-					A3DUns32 PointCount = FaceTessData.m_puiSizesTriangulated[FaceSetIndex++] & kA3DTessFaceDataNormalMask;
-					AddFaceTriangleStripeWithUniqueNormalAndTexture(Context, PointCount, FaceTessData.m_uiTextureCoordIndexesSize, LastTrianguleIndex);
-				}
-				bMustProcess = FaceTessData.m_uiSizesTriangulatedSize > FaceSetIndex;
+				int32 StripeSize = (int32)(FaceTessData.m_puiSizesTriangulated[FaceSetIndex] & kA3DTessFaceDataNormalMask);
+				TriangleCount += (StripeSize - 2);
 			}
+		}
 
-			ensure(!bMustProcess);
+		return TriangleCount;
+	}
+
+	bool FRepresentationConverter::ParseTriangles(const TArray<A3DTopoFace*>& BRepFaces)
+	{
+		TechSoft::TUniqueObject<A3DTess3DData> TessellationData(TessellationBase);
+
+		ensure((TessellationData->m_uiNormalSize % 3) == 0 && (TessellationData->m_uiTextureCoordSize % 2) == 0);
+		TArrayView<FVector> TriangleNormals((FVector*)TessellationData->m_pdNormals, TessellationData->m_uiNormalSize / 3);
+		TArrayView<FVector2d> TriangleTexCoords((FVector2d*)TessellationData->m_pdTextureCoords, TessellationData->m_uiTextureCoordSize / 2);
+
+		ensure(TriangleNormals.Num() == TriangleTexCoords.Num());
+
+		FFaceTriangleCollector TriangleCollector(TessellationData->m_puiTriangulatedIndexes, TriangleNormals, TriangleTexCoords);
+		FFaceTriangleCollector::FCollectionResult CollectionResult;
+
+		for (A3DUns32 Index = 0; Index < TessellationData->m_uiFaceTessSize; ++Index)
+		{
+			const A3DTessFaceData& FaceTessData = TessellationData->m_psFaceTessData[Index];
+
+			if (FaceTessData.m_uiSizesTriangulatedSize == 0)
+			{
+				continue;
+			}
+
+			const A3DTopoFace* TopoFace = BRepFaces[Index];
+
+			TechSoft::TUniqueObject<A3DRootBaseData> RootBaseData(TopoFace);
+			int32 GroupID = RootBaseData.IsValid() ? (int32)RootBaseData->m_uiPersistentId : Index;
+
+			if (MeshWrapper.IsFaceGroupValid(GroupID))
+			{
+				TriangleCollector.CollectTriangles(FaceTessData, TopoFace, GroupID, CollectionResult);
+
+				MeshWrapper.StartFaceTriangles(0, CollectionResult.Normals, CollectionResult.TexCoords);
+				MeshWrapper.AddFaceTriangles(CollectionResult.FaceTriangles);
+				MeshWrapper.EndFaceTriangles();
+			}
 		}
 
 		return true;
 	}
 
-	bool FRepresentationToMeshHelper::Convert()
+	bool FRepresentationConverter::Convert()
 	{
 		using namespace UE::CADKernel;
 
@@ -807,7 +857,7 @@ namespace UE::CADKernel::TechSoftUtilities
 			Vertices.Emplace(Coordinates[0], Coordinates[1], Coordinates[2]);
 		}
 
-		if (!MeshInterface.AddNewVertices(MoveTemp(Vertices)))
+		if (!MeshWrapper.AddNewVertices(MoveTemp(Vertices)))
 		{
 			return false;
 		}
@@ -819,9 +869,16 @@ namespace UE::CADKernel::TechSoftUtilities
 			return false;
 		}
 
-		MeshInterface.ReserveNewTriangles(TriangleCount);
+		ParseTriangles(BRepFaces);
 
 		return false;
 	}
+
+	bool AddRepresentation(A3DRiRepresentationItem* RepresentationItem, MeshUtilities::FMeshWrapperAbstract& MeshWrapper)
+	{
+		FRepresentationConverter Converter(RepresentationItem, MeshWrapper);
+
+		return Converter.IsValid() ? Converter.Convert() : false;
+	}
 }
 #endif
\ No newline at end of file
diff --git a/Engine/Source/Runtime/Datasmith/CADKernel/Engine/Private/TechSoft/TechSoftUtilities.cpp b/Engine/Source/Runtime/Datasmith/CADKernel/Engine/Private/TechSoft/TechSoftUtilities.cpp
--- a/Engine/Source/Runtime/Datasmith/CADKernel/Engine/Private/TechSoft/TechSoftUtilities.cpp
+++ b/Engine/Source/Runtime/Datasmith/CADKernel/Engine/Private/TechSoft/TechSoftUtilities.cpp
@@ -9,6 +9,11 @@
 
 #include "MeshUtilities.h"
 
+#include "Misc/FileHelper.h"
+#include "Misc/Paths.h"
+#include "Serialization/JsonSerializer.h"
+#include "Serialization/JsonWriter.h"
+
 namespace UE::CADKernel
 {
 	bool FTechSoftUtilities::Save(A3DRiRepresentationItem* Representation, const FString& FilePath)
@@ -21,7 +26,7 @@ namespace UE::CADKernel
 		return false;
 	}
 
-	bool FTechSoftUtilities::Tessellate(A3DRiRepresentationItem* Representation, const FTessellationContext& Context, FMeshDescription& Mesh)
+	bool FTechSoftUtilities::Tessellate(A3DRiRepresentationItem* Representation, const FTessellationContext& Context, FMeshDescription& Mesh, bool bEmptyMesh)
 	{
 		using namespace UE::CADKernel::MeshUtilities;
 
@@ -30,11 +35,11 @@ namespace UE::CADKernel
 			return false;
 		}
 
-		TSharedPtr<FMeshWrapperAbstract> MeshInterface = FMeshWrapperAbstract::MakeWrapper(Context, Mesh);
-		return TechSoftUtilities::Tessellate(Representation, Context, *MeshInterface);
+		TSharedPtr<FMeshWrapperAbstract> MeshWrapper = FMeshWrapperAbstract::MakeWrapper(Context, Mesh);
+		return TechSoftUtilities::Tessellate(Representation, Context, *MeshWrapper, bEmptyMesh);
 	}
 
-	bool FTechSoftUtilities::Tessellate(A3DRiRepresentationItem* Representation, const FTessellationContext& Context, UE::Geometry::FDynamicMesh3& Mesh)
+	bool FTechSoftUtilities::Tessellate(A3DRiRepresentationItem* Representation, const FTessellationContext& Context, UE::Geometry::FDynamicMesh3& Mesh, bool bEmptyMesh)
 	{
 		using namespace UE::CADKernel::MeshUtilities;
 
@@ -43,8 +48,8 @@ namespace UE::CADKernel
 			return false;
 		}
 
-		TSharedPtr<FMeshWrapperAbstract> MeshInterface = FMeshWrapperAbstract::MakeWrapper(Context, Mesh);
-		return TechSoftUtilities::Tessellate(Representation, Context, *MeshInterface);
+		TSharedPtr<FMeshWrapperAbstract> MeshWrapper = FMeshWrapperAbstract::MakeWrapper(Context, Mesh);
+		return TechSoftUtilities::Tessellate(Representation, Context, *MeshWrapper, bEmptyMesh);
 	}
 
 	A3DRiRepresentationItem* FTechSoftUtilities::CreateRIBRep(A3DTopoShell* TopoShellPtr)
@@ -177,7 +182,7 @@ namespace UE::CADKernel
 		return A3DStatus::A3D_NOT_IMPLEMENTED;
 	}
 
-	bool TechSoftUtilities::Tessellate(A3DRiRepresentationItem* Representation, const FTessellationContext& Context, MeshUtilities::FMeshWrapperAbstract& MeshInterface)
+	bool TechSoftUtilities::Tessellate(A3DRiRepresentationItem* Representation, const FTessellationContext& Context, MeshUtilities::FMeshWrapperAbstract& MeshWrapper, bool bEmptyMesh)
 	{
 		if (!FTechSoftLibrary::IsInitialized())
 		{
@@ -192,7 +197,7 @@ namespace UE::CADKernel
 		A3DEntityGetType(Representation, &Type);
 		if (Type == kA3DTypeRiPolyBrepModel)
 		{
-			return AddRepresentation(Representation, MeshInterface);
+			return AddRepresentation(Representation, MeshWrapper);
 		}
 
 		TechSoft::TUniqueObject<A3DRiRepresentationItemData> RepresentationData(Representation);
@@ -215,6 +220,11 @@ namespace UE::CADKernel
 			}
 		}
 
+		if (bEmptyMesh)
+		{
+			MeshWrapper.ClearMesh();
+		}
+
 		if (!NewBReps.IsEmpty())
 		{
 			for (A3DRiBrepModel* BrepModel : NewBReps)
@@ -223,7 +233,7 @@ namespace UE::CADKernel
 				// TechSoft is feeling the required structure according to the data structure asked for
 				if (TessellateRepresentation(BrepModel, TessellationSettings))
 				{
-					AddRepresentation(BrepModel, MeshInterface);
+					AddRepresentation(BrepModel, MeshWrapper);
 				}
 			}
 		}
@@ -231,7 +241,7 @@ namespace UE::CADKernel
 		{
 			if (TessellateRepresentation(Representation, TessellationSettings))
 			{
-				AddRepresentation(Representation, MeshInterface);
+				AddRepresentation(Representation, MeshWrapper);
 			}
 		}
 
@@ -303,6 +313,167 @@ namespace UE::CADKernel
 
 		return Status == A3DStatus::A3D_SUCCESS ? true : false;
 	}
+
+	FVector2d TechSoftUtilities::GetUVScale(const A3DTopoFace* TopoFace, double TextureUnit)
+	{
+		TechSoft::TUniqueObject<A3DTopoFaceData> TopoFaceData(TopoFace);
+		if (!TopoFaceData.IsValid())
+		{
+			return FVector2d::ZeroVector;
+		}
+
+		TechSoft::TUniqueObject<A3DDomainData> Domain;
+		if (TopoFaceData->m_bHasTrimDomain)
+		{
+			*Domain = TopoFaceData->m_sSurfaceDomain;
+		}
+		else
+		{
+			A3DStatus Status = A3DSurfGetDomain(TopoFaceData->m_pSurface, Domain.GetPtr());
+			if (Status != A3D_SUCCESS)
+			{
+				return FVector2d::UnitVector;
+			}
+		}
+
+		const int32 IsoCurveCount = 7;
+		const double DeltaU = (Domain->m_sMax.m_dX - Domain->m_sMin.m_dX) / (IsoCurveCount - 1.);
+		const double DeltaV = (Domain->m_sMax.m_dY - Domain->m_sMin.m_dY) / (IsoCurveCount - 1.);
+
+		const A3DSurfBase* A3DSurface = TopoFaceData->m_pSurface;
+
+		FVector NodeMatrix[IsoCurveCount * IsoCurveCount];
+
+		TechSoft::TUniqueObject<A3DVector3dData> Point3D;
+		TechSoft::TUniqueObject<A3DVector2dData> CoordinateObj;
+		A3DVector2dData& Coordinate = *CoordinateObj;
+		Coordinate.m_dX = Domain->m_sMin.m_dX;
+		Coordinate.m_dY = Domain->m_sMin.m_dY;
+
+		for (int32 IndexI = 0; IndexI < IsoCurveCount; IndexI++)
+		{
+			for (int32 IndexJ = 0; IndexJ < IsoCurveCount; IndexJ++)
+			{
+				if (A3DSurfEvaluate(A3DSurface, &Coordinate, 0, Point3D.GetPtr()) == A3D_SUCCESS)
+				{
+					NodeMatrix[IndexI * IsoCurveCount + IndexJ].X = Point3D->m_dX;
+					NodeMatrix[IndexI * IsoCurveCount + IndexJ].Y = Point3D->m_dY;
+					NodeMatrix[IndexI * IsoCurveCount + IndexJ].Z = Point3D->m_dZ;
+					Coordinate.m_dY += DeltaV;
+				}
+			}
+			Coordinate.m_dX += DeltaU;
+			Coordinate.m_dY = Domain->m_sMin.m_dY;
+		}
+
+		// Compute length of 7 iso V line
+		double LengthU[IsoCurveCount];
+		double LengthUMax = 0;
+		double LengthUMed = 0;
+
+		for (int32 IndexJ = 0; IndexJ < IsoCurveCount; IndexJ++)
+		{
+			LengthU[IndexJ] = 0;
+			for (int32 IndexI = 0; IndexI < (IsoCurveCount - 1); IndexI++)
+			{
+				LengthU[IndexJ] += FVector::Distance(NodeMatrix[IndexI * IsoCurveCount + IndexJ], NodeMatrix[(IndexI + 1) * IsoCurveCount + IndexJ]);
+			}
+			LengthUMed += LengthU[IndexJ];
+			LengthUMax = FMath::Max(LengthU[IndexJ], LengthUMax);
+		}
+		LengthUMed /= IsoCurveCount;
+		LengthUMed = LengthUMed * 2 / 3 + LengthUMax / 3;
+
+		// Compute length of 7 iso U line
+		double LengthV[IsoCurveCount];
+		double LengthVMax = 0;
+		double LengthVMed = 0;
+
+		for (int32 IndexI = 0; IndexI < IsoCurveCount; IndexI++)
+		{
+			LengthV[IndexI] = 0;
+			for (int32 IndexJ = 0; IndexJ < (IsoCurveCount - 1); IndexJ++)
+			{
+				LengthV[IndexI] += FVector::Distance(NodeMatrix[IndexI * IsoCurveCount + IndexJ], NodeMatrix[IndexI * IsoCurveCount + IndexJ + 1]);
+			}
+			LengthVMed += LengthV[IndexI];
+			LengthVMax = FMath::Max(LengthV[IndexI], LengthVMax);
+		}
+		LengthVMed /= IsoCurveCount;
+		LengthVMed = LengthVMed * 2 / 3 + LengthVMax / 3;
+
+		// Texture unit is meter, Coord unit from TechSoft is mm, so TextureScale = 0.001 to convert mm into m
+		// #cad_import: Verify assumption above and where TextureUnit can come from
+		constexpr double TextureScale = 0.01;
+
+		return {
+			TextureUnit* TextureScale* LengthUMed / (Domain->m_sMax.m_dX - Domain->m_sMin.m_dX),
+			TextureUnit* TextureScale* LengthVMed / (Domain->m_sMax.m_dY - Domain->m_sMin.m_dY)
+		};
+	}
+
+	TSharedPtr<FJsonObject> TechSoftUtilities::GetJsonObject(A3DAsmProductOccurrence* ProductOcccurence)
+	{
+		TechSoft::TUniqueObject<A3DRootBaseData> RootBaseData(ProductOcccurence);
+
+		if (RootBaseData.IsValid() && RootBaseData->m_uiSize > 0)
+		{
+			TechSoft::TUniqueObject<A3DMiscAttributeData> AttributeData(RootBaseData->m_ppAttributes[0]);
+			if (AttributeData->m_uiSize > 0 && AttributeData->m_asSingleAttributesData[0].m_eType == kA3DModellerAttributeTypeString)
+			{
+				FString JsonString = UTF8_TO_TCHAR(AttributeData->m_asSingleAttributesData[0].m_pcData);
+
+				TSharedRef<TJsonReader<>> JsonReader = TJsonReaderFactory<>::Create(JsonString);
+				TSharedPtr<FJsonObject> JsonObject;
+
+				if (FJsonSerializer::Deserialize(JsonReader, JsonObject))
+				{
+					return JsonObject;
+				}
+			}
+		}
+
+		return {};
+	}
+
+	A3DRiRepresentationItem* TechSoftUtilities::GetRepresentation(const TArray<uint8>& TechSoftRawData)
+	{
+		const FString CachePath = FPaths::Combine(FPaths::ProjectIntermediateDir(), TEXT("Retessellate"));
+		IFileManager::Get().MakeDirectory(*CachePath, true);
+
+		FString ResourceFile = FPaths::CreateTempFilename(*CachePath, TEXT(""), TEXT(".prc"));
+		FPaths::ConvertRelativePathToFull(ResourceFile);
+
+		FFileHelper::SaveArrayToFile(TechSoftRawData, *ResourceFile);
+
+		A3DRWParamsPrcReadHelper* ReadHelper = nullptr;
+
+		A3DAsmModelFile* ModelFile = nullptr;
+		if (A3DAsmModelFileLoadFromPrcFile(TCHAR_TO_UTF8(*ResourceFile), &ReadHelper, &ModelFile) != A3DStatus::A3D_SUCCESS)
+		{
+			return nullptr;
+		}
+
+		TechSoft::TUniqueObject<A3DAsmModelFileData> ModelFileData(ModelFile);
+		if (!ModelFileData.IsValid() || ModelFileData->m_uiPOccurrencesSize != 1)
+		{
+			return nullptr;
+		}
+
+		TechSoft::TUniqueObject<A3DAsmProductOccurrenceData> OccurenceData(ModelFileData->m_ppPOccurrences[0]);
+		if (!OccurenceData.IsValid() || !OccurenceData->m_pPart)
+		{
+			return nullptr;
+		}
+
+		TechSoft::TUniqueObject<A3DAsmPartDefinitionData> PartDefinitionData(OccurenceData->m_pPart);
+		if (!PartDefinitionData.IsValid() || PartDefinitionData->m_uiRepItemsSize == 0)
+		{
+			return nullptr;
+		}
+
+		return PartDefinitionData->m_ppRepItems[0];
+	}
 }
 #else
 namespace UE::CADKernel
diff --git a/Engine/Source/Runtime/Datasmith/CADKernel/Engine/Private/TechSoft/TechSoftUtilities.h b/Engine/Source/Runtime/Datasmith/CADKernel/Engine/Private/TechSoft/TechSoftUtilities.h
--- a/Engine/Source/Runtime/Datasmith/CADKernel/Engine/Private/TechSoft/TechSoftUtilities.h
+++ b/Engine/Source/Runtime/Datasmith/CADKernel/Engine/Private/TechSoft/TechSoftUtilities.h
@@ -7,6 +7,7 @@
 
 #ifdef WITH_HOOPS
 
+class FJsonObject;
 struct FCADKernelTessellationSettings;
 
 namespace UE::CADKernel
@@ -20,9 +21,12 @@ namespace UE::CADKernel
 
 	namespace TechSoftUtilities
 	{
-		bool Tessellate(A3DRiRepresentationItem* Representation, const FTessellationContext& Settings, MeshUtilities::FMeshWrapperAbstract& MeshInterface);
+		bool Tessellate(A3DRiRepresentationItem* Representation, const FTessellationContext& Settings, MeshUtilities::FMeshWrapperAbstract& MeshWrapper, bool bEmptyMesh = false);
 		bool TessellateRepresentation(A3DRiRepresentationItem* Representation, const FCADKernelTessellationSettings& Settings);
-		bool AddRepresentation(A3DRiRepresentationItem* Representation, MeshUtilities::FMeshWrapperAbstract& MeshInterface);
+		bool AddRepresentation(A3DRiRepresentationItem* Representation, MeshUtilities::FMeshWrapperAbstract& MeshWrapper);
+		FVector2d GetUVScale(const A3DTopoFace* TopoFace, double TextureUnit);
+		TSharedPtr<FJsonObject> GetJsonObject(A3DAsmProductOccurrence* ProductOcccurence);
+		A3DRiRepresentationItem* GetRepresentation(const TArray<uint8>& TechSoftRawData);
 
 		/*
 		 ** Sew the in coming BReps together
diff --git a/Engine/Source/Runtime/Datasmith/CADKernel/Engine/Public/CADKernelEngine.h b/Engine/Source/Runtime/Datasmith/CADKernel/Engine/Public/CADKernelEngine.h
--- a/Engine/Source/Runtime/Datasmith/CADKernel/Engine/Public/CADKernelEngine.h
+++ b/Engine/Source/Runtime/Datasmith/CADKernel/Engine/Public/CADKernelEngine.h
@@ -31,26 +31,26 @@ namespace UE::CADKernel
 	{
 		FCADKernelModelParameters ModelParams;
 		FCADKernelMeshParameters MeshParams;
-		TSet<int32> FacesToExtract;
+		TSet<int32> FaceGroupsToExtract;
+		bool bResolveTJunctions = false;
 	};
 
 	struct FTessellationContext : public FMeshExtractionContext
 	{
 		FCADKernelTessellationSettings TessellationSettings;
+		FTessellationContext(const FCADKernelModelParameters& ModelParams, const FCADKernelMeshParameters& MeshParams, const FCADKernelRetessellationSettings& Settings);
 	};
 
 	class FCADKernelUtilities
 	{
 	public:
 		static bool CADKERNELENGINE_API Save(TSharedPtr<UE::CADKernel::FModel>& Model, const FString& FilePath);
-		static bool CADKERNELENGINE_API Tessellate(TSharedPtr<UE::CADKernel::FModel>& Model, const FTessellationContext& Context, FMeshDescription& MeshOut);
-		static bool CADKERNELENGINE_API Tessellate(TSharedPtr<UE::CADKernel::FModel>& Model, const FTessellationContext& Context, UE::Geometry::FDynamicMesh3& MeshOut);
-
-		static bool CADKERNELENGINE_API Tessellate(const UParametricSurfaceData& SurfaceData, UE::Geometry::FDynamicMesh3& MeshOut);
-		static bool CADKERNELENGINE_API Tessellate(const UParametricSurfaceData& SurfaceData, const FCADKernelRetessellationSettings& Settings, UE::Geometry::FDynamicMesh3& MeshOut);
-
-		static bool CADKERNELENGINE_API Tessellate(const UParametricSurfaceData& SurfaceData, FMeshDescription& MeshOut);
-		static bool CADKERNELENGINE_API Tessellate(const UParametricSurfaceData& SurfaceData, const FCADKernelRetessellationSettings& Settings, FMeshDescription& MeshOut);
+		static bool CADKERNELENGINE_API Tessellate(UE::CADKernel::FModel& Model, const FTessellationContext& Context, FMeshDescription& MeshOut, bool bEmptyMesh = false);
+		static bool CADKERNELENGINE_API Tessellate(UE::CADKernel::FModel& Model, const FTessellationContext& Context, UE::Geometry::FDynamicMesh3& MeshOut, bool bEmptyMesh = false);
+		
+		/* Note: Meshes are expected to be in their native coordinates system and unit */
+		static void CADKERNELENGINE_API ApplyExtractionContext(const FMeshExtractionContext& Context, FMeshDescription& MeshInOut);
+		static void CADKERNELENGINE_API ApplyExtractionContext(const FMeshExtractionContext& Context, UE::Geometry::FDynamicMesh3& MeshInOut);
 	};
 
 	class FTechSoftLibrary
@@ -63,12 +63,18 @@ namespace UE::CADKernel
 	private:
 		static bool bIsInitialized;
 	};
+
 	class FTechSoftUtilities
 	{
 	public:
+		/*
+		 ** @param FilePath: Absolute path to the prc file to load
+		 ** @param bSewModel: Optional argument to force a stitching of the model
+		 ** @param StitchingTolerance: Optional stitching tolerance in mm
+		 */
 		static CADKERNELENGINE_API bool Save(A3DRiRepresentationItem* Representation, const FString& FilePath);
-		static CADKERNELENGINE_API bool Tessellate(A3DRiRepresentationItem* Representation, const FTessellationContext& Context, FMeshDescription& MeshOut);
-		static CADKERNELENGINE_API bool Tessellate(A3DRiRepresentationItem* Representation, const FTessellationContext& Context, UE::Geometry::FDynamicMesh3& MeshOut);
+		static CADKERNELENGINE_API bool Tessellate(A3DRiRepresentationItem* Representation, const FTessellationContext& Context, FMeshDescription& MeshOut, bool bEmptyMesh = false);
+		static CADKERNELENGINE_API bool Tessellate(A3DRiRepresentationItem* Representation, const FTessellationContext& Context, UE::Geometry::FDynamicMesh3& MeshOut, bool bEmptyMesh = false);
 
 		static CADKERNELENGINE_API TSharedPtr<UE::CADKernel::FModel> TechSoftToCADKernel(A3DRiRepresentationItem* Representation);
 		/* Limitation: Non of the implicit geometry is supported */
diff --git a/Engine/Source/Runtime/Datasmith/CADKernel/Engine/Public/ParametricSurfaceData.h b/Engine/Source/Runtime/Datasmith/CADKernel/Engine/Public/ParametricSurfaceData.h
--- a/Engine/Source/Runtime/Datasmith/CADKernel/Engine/Public/ParametricSurfaceData.h
+++ b/Engine/Source/Runtime/Datasmith/CADKernel/Engine/Public/ParametricSurfaceData.h
@@ -4,6 +4,8 @@
 
 #include "CoreMinimal.h"
 
+#include "TechSoftIncludes.h"
+
 #include "ParametricSurfaceData.generated.h"
 
 namespace UE::CADKernel
@@ -11,17 +13,25 @@ namespace UE::CADKernel
 	class FModel;
 }
 
+namespace UE::Geometry
+{
+	class FDynamicMesh3;
+}
+
 class UStaticMesh;
 struct FMeshDescription;
 
 typedef void A3DRiRepresentationItem;
+typedef void A3DAsmModelFile;
 
 struct FUnitConverter
 {
-	static const double CentimeterToMillimeter;
-	static const double MillimeterToCentimeter;
-	static const double CentimeterToMeter;
-	static const double MillimeterToMeter;
+	static const double CADKERNELENGINE_API CentimeterToMillimeter;
+	static const double CADKERNELENGINE_API MillimeterToCentimeter;
+	static const double CADKERNELENGINE_API CentimeterToMeter;
+	static const double CADKERNELENGINE_API MeterToCentimeter;
+	static const double CADKERNELENGINE_API MillimeterToMeter;
+	static const double CADKERNELENGINE_API MeterToMillimeter;
 };
 
 UENUM()
@@ -103,9 +113,9 @@ struct CADKERNELENGINE_API FCADKernelTessellationSettings
 		, MaxEdgeLength(InMaxEdgeLength)
 		, NormalTolerance(InNormalTolerance)
 		, StitchingTechnique(InStitchingTechnique)
+		, bUseCADKernel(bInUseCADKernel)
 		, GeometricTolerance(InGeometricTolerance)
 		, StitchingTolerance(InStichingTolerance)
-		, bUseCADKernel(bInUseCADKernel)
 	{
 	}
 
@@ -138,7 +148,7 @@ public:
 
 
 	/**
-	 * Stitching technique applied on neighbouring surfaces before tessellation.
+	 * Stitching technique applied on neighboring surfaces before tessellation.
 	 * None : No stitching applied. This is the default.
 	 * Sewing : Connects surfaces which physically share a boundary but not topologically within a set of objects.
 	 *          This technique can modify the structure of the model by removing and adding objects.
@@ -147,6 +157,18 @@ public:
 	UPROPERTY(config, EditAnywhere, BlueprintReadWrite, AdvancedDisplay, Category = "Geometry & Tessellation Options", meta = (ToolTip = "Stitching technique applied on model before tessellation. Sewing could impact number of objects."))
 	ECADKernelStitchingTechnique StitchingTechnique;
 
+	/**
+	 * Indicates which tessellator to use CADKernel or TechSoft. By default it is CADKernel.
+	 */
+	UPROPERTY(config, EditAnywhere, BlueprintReadWrite, AdvancedDisplay, Category = "Geometry & Tessellation Options", meta = (DisplayName = "Use CADKernel", Tooltip = "When true uses CADKernel tessellator. Otherwise, TechSoft's tessellator is used. Default is true.", editCondition = "bWithTechSoft == true"))
+	bool bUseCADKernel = true;
+
+	/**
+	 * Indicates which tessellator to use CADKernel or TechSoft. By default it is CADKernel.
+	 */
+	UPROPERTY(config, EditAnywhere, BlueprintReadWrite, AdvancedDisplay, Category = "Geometry & Tessellation Options", meta = (Tooltip = "If true, enables the removal of T-Junctions from the resulting mesh. Default is false."))
+	bool bResolveTJunctions = false;
+
 protected:
 	/**
 	 * Tolerance used to determine if a surface should be tessellate or not.
@@ -165,12 +187,6 @@ protected:
 	double StitchingTolerance = 0.001;
 
 public:
-	/**
-	 * Indicates which tessellator to use CADKernel or TechSoft. By default it is CADKernel.
-	 */
-	UPROPERTY(config, EditAnywhere, BlueprintReadWrite, AdvancedDisplay, Category = "Geometry & Tessellation Options", meta = (DisplayName = "Use CADKernel", Tooltip = "When true uses CADKernel tessellator. Otherwise, TechSoft's tessellator is used. Default is true.", editCondition = "bWithTechSoft == true"))
-	bool bUseCADKernel = true;
-	
 	/*
 	 ** Value to convert imported geometrical values to cm which is UE's unit.
 	 ** Default value is 0.1 as CADKernel's and most CAD formats' unit is millimeter.
@@ -179,8 +195,11 @@ public:
 	UPROPERTY()
 	double UnitMultiplier = 0.1;
 
-	UPROPERTY(Transient)
+#ifdef WITH_HOOPS
 	bool bWithTechSoft = true;
+#else
+	bool bWithTechSoft = false;
+#endif
 
 	bool operator == (const FCADKernelTessellationSettings& Other) const
 	{
@@ -214,8 +233,8 @@ public:
 UENUM()
 enum class ECADKernelRetessellationRule : uint8
 {
-	All = 0,
-	SkipDeletedSurfaces,
+	AllFaces = 0,
+	SkipDeletedFaces,
 };
 
 USTRUCT(BlueprintType)
@@ -223,19 +242,20 @@ struct CADKERNELENGINE_API FCADKernelRetessellationSettings : public FCADKernelT
 {
 	GENERATED_BODY()
 
-	FCADKernelRetessellationSettings(double InChordTolerance = 0.2f, double InMaxEdgeLength = 0.0f, double InNormalTolerance = 20.0f, ECADKernelStitchingTechnique InStitchingTechnique = ECADKernelStitchingTechnique::StitchingSew, ECADKernelRetessellationRule InRetessellationRule = ECADKernelRetessellationRule::All)
+	FCADKernelRetessellationSettings(double InChordTolerance = 0.2f, double InMaxEdgeLength = 0.0f, double InNormalTolerance = 20.0f, ECADKernelStitchingTechnique InStitchingTechnique = ECADKernelStitchingTechnique::StitchingSew, ECADKernelRetessellationRule InRetessellationRule = ECADKernelRetessellationRule::AllFaces)
 		: FCADKernelTessellationSettings(InChordTolerance, InMaxEdgeLength, InNormalTolerance, InStitchingTechnique)
 		, RetessellationRule(InRetessellationRule)
 	{
 	}
 
-	FCADKernelRetessellationSettings(const FCADKernelTessellationSettings& TessellationSettings)
+	FCADKernelRetessellationSettings(const FCADKernelTessellationSettings& TessellationSettings, ECADKernelRetessellationRule InRetessellationRule = ECADKernelRetessellationRule::AllFaces)
 		: FCADKernelTessellationSettings(TessellationSettings)
+		, RetessellationRule(InRetessellationRule)
 	{
 	}
 
 	UPROPERTY(config, EditAnywhere, BlueprintReadWrite, AdvancedDisplay, Category = "Retessellation Settings", meta = (ToolTip = "Regenerate deleted surfaces during retesselate or ignore them"))
-	ECADKernelRetessellationRule RetessellationRule = ECADKernelRetessellationRule::All;
+	ECADKernelRetessellationRule RetessellationRule = ECADKernelRetessellationRule::AllFaces;
 };
 
 UCLASS(meta = (DisplayName = "CADKernel Parametric Surface Data"))
@@ -248,18 +268,6 @@ public:
 	virtual bool IsValid() { return CADKernelRawData.Num() > 0 || TechSoftRawData.Num() > 0; }
 
 	virtual bool SetFromFile(const TCHAR* FilePath, bool bForTechSoft = false);
-	
-	virtual void SetRawData(const TArray<uint8>& InRawData, bool bForTechSoft = false)
-	{
-		TArray<uint8>& RawData = bForTechSoft ? TechSoftRawData : CADKernelRawData;
-		RawData = InRawData;
-	}
-
-	virtual void SetRawData(TArray<uint8>&& InRawData, bool bForTechSoft = false)
-	{
-		TArray<uint8>& RawData = bForTechSoft ? TechSoftRawData : CADKernelRawData;
-		RawData = MoveTemp(InRawData);
-	}
 
 	void SetModelParameters(const FCADKernelModelParameters& InModelParameters) { ModelParameters = InModelParameters; }
 	const FCADKernelModelParameters& GetModelParameters() const { return ModelParameters; }
@@ -274,7 +282,14 @@ public:
 
 	const TArray<uint8>& GetRawData(bool bForTechSoft = false) const { return bForTechSoft ? TechSoftRawData : CADKernelRawData; }
 	
-	TSharedPtr<UE::CADKernel::FModel> LoadModel();
+	TSharedPtr<UE::CADKernel::FModel> GetModel();
+	A3DRiRepresentationItem* GetRepresentation();
+
+	bool Tessellate(UE::Geometry::FDynamicMesh3& MeshOut);
+	bool Tessellate(FMeshDescription& MeshOut);
+
+	bool Retessellate(const FCADKernelRetessellationSettings& Settings, UE::Geometry::FDynamicMesh3& MeshOut);
+	bool Retessellate(const FCADKernelRetessellationSettings& Settings, FMeshDescription& MeshOut);
 
 protected:
 	virtual void Serialize(FArchive& Ar) override;
@@ -291,5 +306,18 @@ protected:
 
 	TArray<uint8> CADKernelRawData;
 	TArray<uint8> TechSoftRawData;
+
+private:
+	/*
+	 ** The SetRawData is only for internal use to facilitate the transition
+	 ** out of UDatasmithParametricSurfaceData into the new UParametricSurfaceData class
+	 */
+	virtual void SetRawData(const TArray<uint8>& InRawData, bool bForTechSoft = false)
+	{
+		TArray<uint8>& RawData = bForTechSoft ? TechSoftRawData : CADKernelRawData;
+		RawData = InRawData;
+	}
+	friend class UTechSoftParametricSurfaceData;
+	friend class UCADKernelParametricSurfaceData;
 };
 
