Fixed issue 709654, 709655
Temporarily restricted CADKernelEngine to desktop only

#rnx

diff --git a/Engine/Source/Runtime/Datasmith/CADKernel/Engine/CADKernelEngine.Build.cs b/Engine/Source/Runtime/Datasmith/CADKernel/Engine/CADKernelEngine.Build.cs
--- a/Engine/Source/Runtime/Datasmith/CADKernel/Engine/CADKernelEngine.Build.cs
+++ b/Engine/Source/Runtime/Datasmith/CADKernel/Engine/CADKernelEngine.Build.cs
@@ -9,15 +9,15 @@ namespace UnrealBuildTool.Rules
 			: base(Target)
 		{
 			bUseUnity = false; // TechSoftLibrary.cpp cannot be compiled with the other cpp files.
-			OptimizeCode = CodeOptimization.Never;
+			//OptimizeCode = CodeOptimization.Never;
 
 			PublicDependencyModuleNames.AddRange(
 				new string[]
 				{
 					"Core",
+					"CoreUObject",
 				}
 			);
-			PublicDependencyModuleNames.Add("CoreUObject");
 
 			PrivateDependencyModuleNames.AddRange(
 				new string[]
@@ -32,13 +32,13 @@ namespace UnrealBuildTool.Rules
 			);
 
 			// TechSoft related API needs the TechSoft headers to be accessible
-			bool bHasTechSoft = System.Type.GetType("TechSoft") != null;
+			//bool bHasTechSoft = System.Type.GetType("TechSoft") != null;
 
-			if ((Target.Platform == UnrealTargetPlatform.Win64 || Target.Platform == UnrealTargetPlatform.Linux) && bHasTechSoft)
-			{
-				// #cad_import: Verifiy that this is correct
-				PublicDependencyModuleNames.Add("TechSoft");
-			}
+			//if ((Target.Platform == UnrealTargetPlatform.Win64 || Target.Platform == UnrealTargetPlatform.Linux) && bHasTechSoft)
+			//{
+			//	// #cad_import: Verifiy that this is correct
+			//	PublicDependencyModuleNames.Add("TechSoft");
+			//}
 		}
 	}
 }
\ No newline at end of file
diff --git a/Engine/Source/Runtime/Datasmith/CADKernel/Engine/Private/CADKernelEngine.cpp b/Engine/Source/Runtime/Datasmith/CADKernel/Engine/Private/CADKernelEngine.cpp
--- a/Engine/Source/Runtime/Datasmith/CADKernel/Engine/Private/CADKernelEngine.cpp
+++ b/Engine/Source/Runtime/Datasmith/CADKernel/Engine/Private/CADKernelEngine.cpp
@@ -2,6 +2,7 @@
 
 #include "CADKernelEngine.h"
 
+#if PLATFORM_DESKTOP
 #include "CADKernelEnginePrivate.h"
 #include "MeshUtilities.h"
 
@@ -55,4 +56,5 @@ namespace UE::CADKernel
 		TSharedPtr<FMeshWrapperAbstract> MeshWrapper = FMeshWrapperAbstract::MakeWrapper(Context, Mesh);
 		return Private::Tessellate(Model, Context, *MeshWrapper, bEmptyMesh);
 	}
-}
\ No newline at end of file
+}
+#endif
diff --git a/Engine/Source/Runtime/Datasmith/CADKernel/Engine/Private/CADKernelEnginePrivate.cpp b/Engine/Source/Runtime/Datasmith/CADKernel/Engine/Private/CADKernelEnginePrivate.cpp
--- a/Engine/Source/Runtime/Datasmith/CADKernel/Engine/Private/CADKernelEnginePrivate.cpp
+++ b/Engine/Source/Runtime/Datasmith/CADKernel/Engine/Private/CADKernelEnginePrivate.cpp
@@ -2,6 +2,7 @@
 
 #include "CADKernelEnginePrivate.h"
 
+#if PLATFORM_DESKTOP
 #include "CADKernelEngine.h"
 #include "MeshUtilities.h"
 
@@ -177,3 +178,4 @@ namespace UE::CADKernel::Private
 		return PolylinesOut.Num() > PolylinesCount;
 	}
 }
+#endif
\ No newline at end of file
diff --git a/Engine/Source/Runtime/Datasmith/CADKernel/Engine/Private/CADKernelEnginePrivate.h b/Engine/Source/Runtime/Datasmith/CADKernel/Engine/Private/CADKernelEnginePrivate.h
--- a/Engine/Source/Runtime/Datasmith/CADKernel/Engine/Private/CADKernelEnginePrivate.h
+++ b/Engine/Source/Runtime/Datasmith/CADKernel/Engine/Private/CADKernelEnginePrivate.h
@@ -3,6 +3,7 @@
 
 #include "CoreMinimal.h"
 
+#if PLATFORM_DESKTOP
 #include "CADKernelEngine.h"
 #include "MeshUtilities.h"
 
@@ -94,4 +95,4 @@ namespace UE::CADKernel
 		bool AddModelMesh(const UE::CADKernel::FModelMesh& ModelMesh, UE::CADKernel::MeshUtilities::FMeshWrapperAbstract& MeshWrapper);
 	}
 }
-
+#endif
diff --git a/Engine/Source/Runtime/Datasmith/CADKernel/Engine/Private/CADKernelMesh.cpp b/Engine/Source/Runtime/Datasmith/CADKernel/Engine/Private/CADKernelMesh.cpp
--- a/Engine/Source/Runtime/Datasmith/CADKernel/Engine/Private/CADKernelMesh.cpp
+++ b/Engine/Source/Runtime/Datasmith/CADKernel/Engine/Private/CADKernelMesh.cpp
@@ -1,6 +1,7 @@
 // Copyright Epic Games, Inc. All Rights Reserved.
 
 #include "CADKernelEngine.h"
+#if PLATFORM_DESKTOP
 
 #include "CADKernelEnginePrivate.h"
 #include "MeshUtilities.h"
@@ -118,4 +119,5 @@ namespace UE::CADKernel::Private
 		FModelMeshConverter Converter(ModelMesh, MeshWrapper);
 		return Converter.Convert();
 	}
-}
\ No newline at end of file
+}
+#endif
diff --git a/Engine/Source/Runtime/Datasmith/CADKernel/Engine/Private/DynamicMeshWrapper.cpp b/Engine/Source/Runtime/Datasmith/CADKernel/Engine/Private/DynamicMeshWrapper.cpp
--- a/Engine/Source/Runtime/Datasmith/CADKernel/Engine/Private/DynamicMeshWrapper.cpp
+++ b/Engine/Source/Runtime/Datasmith/CADKernel/Engine/Private/DynamicMeshWrapper.cpp
@@ -2,6 +2,7 @@
 
 #include "DynamicMeshWrapper.h"
 
+#if PLATFORM_DESKTOP
 #include "CADKernelEngine.h"
 
 #include "Algo/Reverse.h"
@@ -486,3 +487,4 @@ namespace UE::CADKernel::MeshUtilities
 		// #cad_import: TO DO
 	}
 }
+#endif
diff --git a/Engine/Source/Runtime/Datasmith/CADKernel/Engine/Private/DynamicMeshWrapper.h b/Engine/Source/Runtime/Datasmith/CADKernel/Engine/Private/DynamicMeshWrapper.h
--- a/Engine/Source/Runtime/Datasmith/CADKernel/Engine/Private/DynamicMeshWrapper.h
+++ b/Engine/Source/Runtime/Datasmith/CADKernel/Engine/Private/DynamicMeshWrapper.h
@@ -2,6 +2,7 @@
 #pragma once
 
 #include "CoreMinimal.h"
+#if PLATFORM_DESKTOP
 
 #include "MeshUtilities.h"
 
@@ -90,3 +91,4 @@ namespace UE::CADKernel::MeshUtilities
 		bool bIsFinalized = false;
 	};
 }
+#endif
diff --git a/Engine/Source/Runtime/Datasmith/CADKernel/Engine/Private/MeshDescriptionWrapper.cpp b/Engine/Source/Runtime/Datasmith/CADKernel/Engine/Private/MeshDescriptionWrapper.cpp
--- a/Engine/Source/Runtime/Datasmith/CADKernel/Engine/Private/MeshDescriptionWrapper.cpp
+++ b/Engine/Source/Runtime/Datasmith/CADKernel/Engine/Private/MeshDescriptionWrapper.cpp
@@ -2,6 +2,7 @@
 
 #include "MeshDescriptionWrapper.h"
 
+#if PLATFORM_DESKTOP
 #include "CADKernelEngine.h"
 #include "CADKernelEngineLog.h"
 #include "MeshUtilities.h"
@@ -745,3 +746,4 @@ namespace UE::CADKernel::MeshUtilities
 		}
 	}
 }
+#endif
diff --git a/Engine/Source/Runtime/Datasmith/CADKernel/Engine/Private/MeshDescriptionWrapper.h b/Engine/Source/Runtime/Datasmith/CADKernel/Engine/Private/MeshDescriptionWrapper.h
--- a/Engine/Source/Runtime/Datasmith/CADKernel/Engine/Private/MeshDescriptionWrapper.h
+++ b/Engine/Source/Runtime/Datasmith/CADKernel/Engine/Private/MeshDescriptionWrapper.h
@@ -3,6 +3,7 @@
 #pragma once
 
 #include "MeshUtilities.h"
+#if PLATFORM_DESKTOP
 
 #include "MeshDescription.h"
 #include "NaniteDefinitions.h"
@@ -155,3 +156,4 @@ namespace UE::CADKernel::MeshUtilities
 		bool bIsFinalized = false;
 	};
 }
+#endif
\ No newline at end of file
diff --git a/Engine/Source/Runtime/Datasmith/CADKernel/Engine/Private/MeshTopologyHelper.cpp b/Engine/Source/Runtime/Datasmith/CADKernel/Engine/Private/MeshTopologyHelper.cpp
--- a/Engine/Source/Runtime/Datasmith/CADKernel/Engine/Private/MeshTopologyHelper.cpp
+++ b/Engine/Source/Runtime/Datasmith/CADKernel/Engine/Private/MeshTopologyHelper.cpp
@@ -1,8 +1,8 @@
 // Copyright Epic Games, Inc. All Rights Reserved.
 
-
 #include "MeshTopologyHelper.h"
 
+#if PLATFORM_DESKTOP
 #include "MeshAttributes.h"
 #include "StaticMeshAttributes.h"
 
@@ -274,4 +274,4 @@ namespace UE::CADKernel::MeshUtilities
 		StaticMeshAttributes.GetVertexInstanceNormals()[VertexInstanceID] = Normal;
 	}
 }
-
+#endif
diff --git a/Engine/Source/Runtime/Datasmith/CADKernel/Engine/Private/MeshTopologyHelper.h b/Engine/Source/Runtime/Datasmith/CADKernel/Engine/Private/MeshTopologyHelper.h
--- a/Engine/Source/Runtime/Datasmith/CADKernel/Engine/Private/MeshTopologyHelper.h
+++ b/Engine/Source/Runtime/Datasmith/CADKernel/Engine/Private/MeshTopologyHelper.h
@@ -2,8 +2,8 @@
 
 #pragma once
 
-
 #include "CoreMinimal.h"
+#if PLATFORM_DESKTOP
 
 #include "MeshTypes.h"
 #include "MeshDescription.h"
@@ -102,4 +102,4 @@ namespace UE::CADKernel::MeshUtilities
 		void DefineVertexTopologyApproximation(FVertexID VertexID);
 	};
 }
-
+#endif
diff --git a/Engine/Source/Runtime/Datasmith/CADKernel/Engine/Private/MeshUtilities.cpp b/Engine/Source/Runtime/Datasmith/CADKernel/Engine/Private/MeshUtilities.cpp
--- a/Engine/Source/Runtime/Datasmith/CADKernel/Engine/Private/MeshUtilities.cpp
+++ b/Engine/Source/Runtime/Datasmith/CADKernel/Engine/Private/MeshUtilities.cpp
@@ -2,6 +2,7 @@
 
 #include "MeshUtilities.h"
 
+#if PLATFORM_DESKTOP
 #include "CADKernelEngineLog.h"
 
 #include "MeshDescription.h"
@@ -45,4 +46,5 @@ namespace UE::CADKernel
 			// #cad_import: TODO Refactor MeshOperator::RecomputeNullNormal
 		}
 	}
-}
\ No newline at end of file
+}
+#endif
diff --git a/Engine/Source/Runtime/Datasmith/CADKernel/Engine/Private/MeshUtilities.h b/Engine/Source/Runtime/Datasmith/CADKernel/Engine/Private/MeshUtilities.h
--- a/Engine/Source/Runtime/Datasmith/CADKernel/Engine/Private/MeshUtilities.h
+++ b/Engine/Source/Runtime/Datasmith/CADKernel/Engine/Private/MeshUtilities.h
@@ -3,6 +3,7 @@
 #pragma once
 
 #include "CADKernelEngine.h"
+#if PLATFORM_DESKTOP
 #include "ParametricSurfaceData.h"
 
 #include "IntVectorTypes.h"
@@ -303,4 +304,5 @@ namespace UE::CADKernel::MeshUtilities
 		static void ResolveTJunctions(FMeshDescription& MeshDescription, double Tolerance);
 		static void RecomputeNullNormal(FMeshDescription& MeshDescriptione);
 	};
-}
\ No newline at end of file
+}
+#endif
\ No newline at end of file
diff --git a/Engine/Source/Runtime/Datasmith/CADKernel/Engine/Private/ParametricSurfaceData.cpp b/Engine/Source/Runtime/Datasmith/CADKernel/Engine/Private/ParametricSurfaceData.cpp
--- a/Engine/Source/Runtime/Datasmith/CADKernel/Engine/Private/ParametricSurfaceData.cpp
+++ b/Engine/Source/Runtime/Datasmith/CADKernel/Engine/Private/ParametricSurfaceData.cpp
@@ -57,6 +57,7 @@ void UParametricSurfaceData::Serialize(FArchive& Ar)
 
 TSharedPtr<UE::CADKernel::FModel> UParametricSurfaceData::GetModel()
 {
+#if PLATFORM_DESKTOP
 	using namespace UE::CADKernel;
 
 	if (CADKernelRawData.IsEmpty())
@@ -68,10 +69,14 @@ TSharedPtr<UE::CADKernel::FModel> UParametricSurfaceData::GetModel()
 	Session->AddDatabase(CADKernelRawData);
 
 	return Session->GetModelAsShared();
+#else
+	return {};
+#endif
 }
 
 A3DRiRepresentationItem* UParametricSurfaceData::GetRepresentation()
 {
+#if PLATFORM_DESKTOP
 	using namespace UE::CADKernel;
 
 	if (TechSoftRawData.IsEmpty())
@@ -79,29 +84,47 @@ A3DRiRepresentationItem* UParametricSurfaceData::GetRepresentation()
 		return {};
 	}
 
-	return TechSoftUtilities::GetRepresentation(TechSoftRawData);
+	return FTechSoftUtilities::GetRepresentation(TechSoftRawData);
+#else
+	return nullptr;
+#endif
 }
 
 
 bool UParametricSurfaceData::Tessellate(UE::Geometry::FDynamicMesh3& MeshOut)
 {
+#if PLATFORM_DESKTOP
 	return UE::CADKernel::Private::Tessellate(*this, MeshOut);
+#else
+	return false;
+#endif
 }
 
 bool UParametricSurfaceData::Tessellate(FMeshDescription& MeshOut)
 {
+#if PLATFORM_DESKTOP
 	return UE::CADKernel::Private::Tessellate(*this, MeshOut);
+#else
+	return false;
+#endif
 }
 
 bool UParametricSurfaceData::Retessellate(const FCADKernelRetessellationSettings& Settings, UE::Geometry::FDynamicMesh3& MeshOut)
 {
+#if PLATFORM_DESKTOP
 	return UE::CADKernel::Private::Retessellate(*this, Settings, MeshOut);
+#else
+	return false;
+#endif
 }
 
 bool UParametricSurfaceData::Retessellate(const FCADKernelRetessellationSettings& Settings, FMeshDescription& MeshOut)
 {
+#if PLATFORM_DESKTOP
 	return UE::CADKernel::Private::Retessellate(*this, Settings, MeshOut);
+#else
+	return false;
+#endif
 }
 
 #undef LOCTEXT_NAMESPACE // "ParametricSurfaceData"
-
diff --git a/Engine/Source/Runtime/Datasmith/CADKernel/Engine/Private/TechSoft/CADKernelToTechSoft.cpp b/Engine/Source/Runtime/Datasmith/CADKernel/Engine/Private/TechSoft/CADKernelToTechSoft.cpp
--- a/Engine/Source/Runtime/Datasmith/CADKernel/Engine/Private/TechSoft/CADKernelToTechSoft.cpp
+++ b/Engine/Source/Runtime/Datasmith/CADKernel/Engine/Private/TechSoft/CADKernelToTechSoft.cpp
@@ -2,6 +2,7 @@
 
 #include "CADKernelEngine.h"
 
+#if PLATFORM_DESKTOP
 #ifdef WITH_HOOPS
 #include "TechSoftUniqueObjectImpl.h"
 
@@ -26,3 +27,5 @@ namespace UE::CADKernel
 	}
 }
 #endif
+#endif
+
diff --git a/Engine/Source/Runtime/Datasmith/CADKernel/Engine/Private/TechSoft/TechSoftLibrary.cpp b/Engine/Source/Runtime/Datasmith/CADKernel/Engine/Private/TechSoft/TechSoftLibrary.cpp
--- a/Engine/Source/Runtime/Datasmith/CADKernel/Engine/Private/TechSoft/TechSoftLibrary.cpp
+++ b/Engine/Source/Runtime/Datasmith/CADKernel/Engine/Private/TechSoft/TechSoftLibrary.cpp
@@ -4,6 +4,7 @@
 #define INITIALIZE_A3D_API
 
 #include "CADKernelEngine.h"
+#if PLATFORM_DESKTOP
 
 #include "HAL/PlatformProcess.h"
 #include "Misc/Paths.h"
@@ -60,4 +61,5 @@ namespace UE::CADKernel
 		return TEXT("TechSoft unavailable");
 	}
 #endif
-}
\ No newline at end of file
+}
+#endif
\ No newline at end of file
diff --git a/Engine/Source/Runtime/Datasmith/CADKernel/Engine/Private/TechSoft/TechSoftToCADKernel.cpp b/Engine/Source/Runtime/Datasmith/CADKernel/Engine/Private/TechSoft/TechSoftToCADKernel.cpp
--- a/Engine/Source/Runtime/Datasmith/CADKernel/Engine/Private/TechSoft/TechSoftToCADKernel.cpp
+++ b/Engine/Source/Runtime/Datasmith/CADKernel/Engine/Private/TechSoft/TechSoftToCADKernel.cpp
@@ -2,6 +2,7 @@
 
 #include "CADKernelEngine.h"
 
+#if PLATFORM_DESKTOP
 #ifdef WITH_HOOPS
 #include "TechSoftUniqueObjectImpl.h"
 
@@ -27,3 +28,4 @@ namespace UE::CADKernel
 	}
 }
 #endif
+#endif
diff --git a/Engine/Source/Runtime/Datasmith/CADKernel/Engine/Private/TechSoft/TechSoftUtilities.cpp b/Engine/Source/Runtime/Datasmith/CADKernel/Engine/Private/TechSoft/TechSoftUtilities.cpp
--- a/Engine/Source/Runtime/Datasmith/CADKernel/Engine/Private/TechSoft/TechSoftUtilities.cpp
+++ b/Engine/Source/Runtime/Datasmith/CADKernel/Engine/Private/TechSoft/TechSoftUtilities.cpp
@@ -1,6 +1,8 @@
 // Copyright Epic Games, Inc. All Rights Reserved.
 
 #include "TechSoftUtilities.h"
+
+#if PLATFORM_DESKTOP
 #include "CADKernelEngine.h"
 
 #ifdef WITH_HOOPS
@@ -182,6 +184,45 @@ namespace UE::CADKernel
 		return A3DStatus::A3D_NOT_IMPLEMENTED;
 	}
 
+	A3DRiRepresentationItem* FTechSoftUtilities::GetRepresentation(const TArray<uint8>& TechSoftRawData)
+	{
+		const FString CachePath = FPaths::Combine(FPaths::ProjectIntermediateDir(), TEXT("Retessellate"));
+		IFileManager::Get().MakeDirectory(*CachePath, true);
+
+		FString ResourceFile = FPaths::CreateTempFilename(*CachePath, TEXT(""), TEXT(".prc"));
+		FPaths::ConvertRelativePathToFull(ResourceFile);
+
+		FFileHelper::SaveArrayToFile(TechSoftRawData, *ResourceFile);
+
+		A3DRWParamsPrcReadHelper* ReadHelper = nullptr;
+
+		A3DAsmModelFile* ModelFile = nullptr;
+		if (A3DAsmModelFileLoadFromPrcFile(TCHAR_TO_UTF8(*ResourceFile), &ReadHelper, &ModelFile) != A3DStatus::A3D_SUCCESS)
+		{
+			return nullptr;
+		}
+
+		TechSoft::TUniqueObject<A3DAsmModelFileData> ModelFileData(ModelFile);
+		if (!ModelFileData.IsValid() || ModelFileData->m_uiPOccurrencesSize != 1)
+		{
+			return nullptr;
+		}
+
+		TechSoft::TUniqueObject<A3DAsmProductOccurrenceData> OccurenceData(ModelFileData->m_ppPOccurrences[0]);
+		if (!OccurenceData.IsValid() || !OccurenceData->m_pPart)
+		{
+			return nullptr;
+		}
+
+		TechSoft::TUniqueObject<A3DAsmPartDefinitionData> PartDefinitionData(OccurenceData->m_pPart);
+		if (!PartDefinitionData.IsValid() || PartDefinitionData->m_uiRepItemsSize == 0)
+		{
+			return nullptr;
+		}
+
+		return PartDefinitionData->m_ppRepItems[0];
+	}
+
 	bool TechSoftUtilities::Tessellate(A3DRiRepresentationItem* Representation, const FTessellationContext& Context, MeshUtilities::FMeshWrapperAbstract& MeshWrapper, bool bEmptyMesh)
 	{
 		if (!FTechSoftLibrary::IsInitialized())
@@ -435,45 +476,6 @@ namespace UE::CADKernel
 
 		return {};
 	}
-
-	A3DRiRepresentationItem* TechSoftUtilities::GetRepresentation(const TArray<uint8>& TechSoftRawData)
-	{
-		const FString CachePath = FPaths::Combine(FPaths::ProjectIntermediateDir(), TEXT("Retessellate"));
-		IFileManager::Get().MakeDirectory(*CachePath, true);
-
-		FString ResourceFile = FPaths::CreateTempFilename(*CachePath, TEXT(""), TEXT(".prc"));
-		FPaths::ConvertRelativePathToFull(ResourceFile);
-
-		FFileHelper::SaveArrayToFile(TechSoftRawData, *ResourceFile);
-
-		A3DRWParamsPrcReadHelper* ReadHelper = nullptr;
-
-		A3DAsmModelFile* ModelFile = nullptr;
-		if (A3DAsmModelFileLoadFromPrcFile(TCHAR_TO_UTF8(*ResourceFile), &ReadHelper, &ModelFile) != A3DStatus::A3D_SUCCESS)
-		{
-			return nullptr;
-		}
-
-		TechSoft::TUniqueObject<A3DAsmModelFileData> ModelFileData(ModelFile);
-		if (!ModelFileData.IsValid() || ModelFileData->m_uiPOccurrencesSize != 1)
-		{
-			return nullptr;
-		}
-
-		TechSoft::TUniqueObject<A3DAsmProductOccurrenceData> OccurenceData(ModelFileData->m_ppPOccurrences[0]);
-		if (!OccurenceData.IsValid() || !OccurenceData->m_pPart)
-		{
-			return nullptr;
-		}
-
-		TechSoft::TUniqueObject<A3DAsmPartDefinitionData> PartDefinitionData(OccurenceData->m_pPart);
-		if (!PartDefinitionData.IsValid() || PartDefinitionData->m_uiRepItemsSize == 0)
-		{
-			return nullptr;
-		}
-
-		return PartDefinitionData->m_ppRepItems[0];
-	}
 }
 #else
 namespace UE::CADKernel
@@ -483,14 +485,20 @@ namespace UE::CADKernel
 		return false;
 	}
 
-	bool FTechSoftUtilities::Tessellate(A3DRiRepresentationItem* Representation, const FTessellationContext& Settings, FMeshDescription& Mesh)
+	bool FTechSoftUtilities::Tessellate(A3DRiRepresentationItem* Representation, const FTessellationContext& Settings, FMeshDescription& Mesh, bool bEmptyMesh)
 	{
 		return false;
 	}
 
-	bool FTechSoftUtilities::Tessellate(A3DRiRepresentationItem* Representation, const FTessellationContext& Settings, UE::Geometry::FDynamicMesh3& Mesh)
+	bool FTechSoftUtilities::Tessellate(A3DRiRepresentationItem* Representation, const FTessellationContext& Settings, UE::Geometry::FDynamicMesh3& Mesh, bool bEmptyMesh)
 	{
 		return false;
 	}
+
+	A3DRiRepresentationItem* FTechSoftUtilities::GetRepresentation(const TArray<uint8>& TechSoftRawData)
+	{
+		return nullptr;
+	}
 }
 #endif
+#endif
diff --git a/Engine/Source/Runtime/Datasmith/CADKernel/Engine/Private/TechSoft/TechSoftUtilities.h b/Engine/Source/Runtime/Datasmith/CADKernel/Engine/Private/TechSoft/TechSoftUtilities.h
--- a/Engine/Source/Runtime/Datasmith/CADKernel/Engine/Private/TechSoft/TechSoftUtilities.h
+++ b/Engine/Source/Runtime/Datasmith/CADKernel/Engine/Private/TechSoft/TechSoftUtilities.h
@@ -3,6 +3,7 @@
 
 #include "CoreMinimal.h"
 
+#if PLATFORM_DESKTOP
 #include "TechSoftIncludes.h"
 
 #ifdef WITH_HOOPS
@@ -26,7 +27,6 @@ namespace UE::CADKernel
 		bool AddRepresentation(A3DRiRepresentationItem* Representation, MeshUtilities::FMeshWrapperAbstract& MeshWrapper);
 		FVector2d GetUVScale(const A3DTopoFace* TopoFace, double TextureUnit);
 		TSharedPtr<FJsonObject> GetJsonObject(A3DAsmProductOccurrence* ProductOcccurence);
-		A3DRiRepresentationItem* GetRepresentation(const TArray<uint8>& TechSoftRawData);
 
 		/*
 		 ** Sew the in coming BReps together
@@ -38,3 +38,4 @@ namespace UE::CADKernel
 	};
 }
 #endif
+#endif
diff --git a/Engine/Source/Runtime/Datasmith/CADKernel/Engine/Public/CADKernelEngine.h b/Engine/Source/Runtime/Datasmith/CADKernel/Engine/Public/CADKernelEngine.h
--- a/Engine/Source/Runtime/Datasmith/CADKernel/Engine/Public/CADKernelEngine.h
+++ b/Engine/Source/Runtime/Datasmith/CADKernel/Engine/Public/CADKernelEngine.h
@@ -3,6 +3,7 @@
 
 #include "CoreMinimal.h"
 
+#if PLATFORM_DESKTOP
 #include "ParametricSurfaceData.h"
 #include "TechSoftIncludes.h"
 
@@ -79,6 +80,7 @@ namespace UE::CADKernel
 		static CADKERNELENGINE_API TSharedPtr<UE::CADKernel::FModel> TechSoftToCADKernel(A3DRiRepresentationItem* Representation);
 		/* Limitation: Non of the implicit geometry is supported */
 		static CADKERNELENGINE_API A3DRiRepresentationItem* CADKernelToTechSoft(TSharedPtr<UE::CADKernel::FModel>& Model);
+		static CADKERNELENGINE_API A3DRiRepresentationItem* GetRepresentation(const TArray<uint8>& TechSoftRawData);
 
 #ifdef WITH_HOOPS
 		static CADKERNELENGINE_API A3DRiRepresentationItem* CreateRIBRep(A3DTopoShell* TopoShellPtr);
@@ -139,3 +141,4 @@ namespace UE::CADKernel
 		}
 	}
 }
+#endif
