Fixed (another) issue 709771

#rnx

diff --git a/Engine/Source/Runtime/Datasmith/CADKernel/Engine/CADKernelEngine.Build.cs b/Engine/Source/Runtime/Datasmith/CADKernel/Engine/CADKernelEngine.Build.cs
--- a/Engine/Source/Runtime/Datasmith/CADKernel/Engine/CADKernelEngine.Build.cs
+++ b/Engine/Source/Runtime/Datasmith/CADKernel/Engine/CADKernelEngine.Build.cs
@@ -32,13 +32,13 @@ namespace UnrealBuildTool.Rules
 			);
 
 			// TechSoft related API needs the TechSoft headers to be accessible
-			//bool bHasTechSoft = System.Type.GetType("TechSoft") != null;
+			bool bHasTechSoft = System.Type.GetType("TechSoft") != null;
 
-			//if ((Target.Platform == UnrealTargetPlatform.Win64 || Target.Platform == UnrealTargetPlatform.Linux) && bHasTechSoft)
-			//{
-			//	// #cad_import: Verifiy that this is correct
-			//	PublicDependencyModuleNames.Add("TechSoft");
-			//}
+			if ((Target.Platform == UnrealTargetPlatform.Win64 || Target.Platform == UnrealTargetPlatform.Linux) && bHasTechSoft)
+			{
+				// #cad_import: Verifiy that this is correct
+				PublicDependencyModuleNames.Add("TechSoft");
+			}
 		}
 	}
 }
\ No newline at end of file
