[In Progress] CADKernelEngine module implementation
- Fixed issues with the retessellation of TechSoft representation
- Refined API to set/get a model or a representation in/from a UParametricSurfaceData object

#rnx

diff --git a/Engine/Source/Runtime/Datasmith/CADKernel/Engine/Private/CADKernelEnginePrivate.cpp b/Engine/Source/Runtime/Datasmith/CADKernel/Engine/Private/CADKernelEnginePrivate.cpp
--- a/Engine/Source/Runtime/Datasmith/CADKernel/Engine/Private/CADKernelEnginePrivate.cpp
+++ b/Engine/Source/Runtime/Datasmith/CADKernel/Engine/Private/CADKernelEnginePrivate.cpp
@@ -83,6 +83,11 @@ namespace UE::CADKernel::Private
 		return AddModelMesh(*ModelMesh, MeshWrapper);
 	}
 
+	bool ToRawData(TSharedPtr<UE::CADKernel::FModel>& Model, TArray<uint8>& RawDataOut)
+	{
+		return false;
+	}
+
 	bool GetFaceTrimmingCurves(const FModel& Model, const FTopologicalFace& Face, TArray<TArray<TArray<FVector>>>& CurvesOut)
 	{
 		using namespace UE::CADKernel;
diff --git a/Engine/Source/Runtime/Datasmith/CADKernel/Engine/Private/CADKernelEnginePrivate.h b/Engine/Source/Runtime/Datasmith/CADKernel/Engine/Private/CADKernelEnginePrivate.h
--- a/Engine/Source/Runtime/Datasmith/CADKernel/Engine/Private/CADKernelEnginePrivate.h
+++ b/Engine/Source/Runtime/Datasmith/CADKernel/Engine/Private/CADKernelEnginePrivate.h
@@ -43,8 +43,9 @@ namespace UE::CADKernel
 					bTessellationSuccessful = FCADKernelUtilities::Tessellate(*Model, Context, MeshOut);
 				}
 			}
-			else if (A3DRiRepresentationItem* Representation = Data.GetRepresentation())
+			else if (FTechSoftLibrary::Initialize())
 			{
+				A3DRiRepresentationItem* Representation = Data.GetRepresentation();
 				bTessellationSuccessful = FTechSoftUtilities::Tessellate(Representation, Context, MeshOut);
 			}
 
@@ -74,7 +75,7 @@ namespace UE::CADKernel
 					bTessellationSuccessful = FCADKernelUtilities::Tessellate(*Model, Context, MeshOut, true);
 				}
 			}
-			else
+			else if(FTechSoftLibrary::Initialize())
 			{
 				if (A3DRiRepresentationItem* Representation = Data.GetRepresentation())
 				{
@@ -93,6 +94,7 @@ namespace UE::CADKernel
 		bool GetFaceTrimming3DPolylines(const UE::CADKernel::FModel& Model, const UE::CADKernel::FTopologicalFace& Face, TArray<TArray<FVector>>& PolylinesOut);
 
 		bool AddModelMesh(const UE::CADKernel::FModelMesh& ModelMesh, UE::CADKernel::MeshUtilities::FMeshWrapperAbstract& MeshWrapper);
+		bool ToRawData(TSharedPtr<UE::CADKernel::FModel>& Model, TArray<uint8>& RawDataOut);
 	}
 }
 #endif
diff --git a/Engine/Source/Runtime/Datasmith/CADKernel/Engine/Private/MeshUtilities.h b/Engine/Source/Runtime/Datasmith/CADKernel/Engine/Private/MeshUtilities.h
--- a/Engine/Source/Runtime/Datasmith/CADKernel/Engine/Private/MeshUtilities.h
+++ b/Engine/Source/Runtime/Datasmith/CADKernel/Engine/Private/MeshUtilities.h
@@ -58,6 +58,10 @@ namespace UE::CADKernel::MeshUtilities
 			bHasFaceGroupsToSkip = !InContext.FaceGroupsToExtract.IsEmpty();
 		}
 
+		/*
+		 ** Call this method when the building of the mesh is completed.
+		 ** Must be called before the pointer to the child class is deleted.
+		 */
 		void Complete()
 		{
 			if (bIsComplete)
@@ -85,7 +89,7 @@ namespace UE::CADKernel::MeshUtilities
 			bIsComplete = true;
 		}
 
-		virtual ~FMeshWrapperAbstract() { Complete(); }
+		virtual ~FMeshWrapperAbstract() {}
 
 		virtual void ClearMesh() = 0;
 
diff --git a/Engine/Source/Runtime/Datasmith/CADKernel/Engine/Private/ParametricSurfaceData.cpp b/Engine/Source/Runtime/Datasmith/CADKernel/Engine/Private/ParametricSurfaceData.cpp
--- a/Engine/Source/Runtime/Datasmith/CADKernel/Engine/Private/ParametricSurfaceData.cpp
+++ b/Engine/Source/Runtime/Datasmith/CADKernel/Engine/Private/ParametricSurfaceData.cpp
@@ -13,6 +13,12 @@
 #include "Misc/Paths.h"
 #include "UObject/EnterpriseObjectVersion.h"
 
+#ifdef WITH_HOOPS
+bool FCADKernelTessellationSettings::bWithHoops = true;
+#else
+bool FCADKernelTessellationSettings::bWithHoops = false;
+#endif
+
 #define LOCTEXT_NAMESPACE "ParametricSurfaceData"
 
 const double FUnitConverter::CentimeterToMillimeter = 10.;
@@ -127,4 +133,24 @@ bool UParametricSurfaceData::Retessellate(const FCADKernelRetessellationSettings
 #endif
 }
 
+bool UParametricSurfaceData::SetModel(TSharedPtr<UE::CADKernel::FModel>& Model, double UnitModelToCentimeter)
+{
+#if PLATFORM_DESKTOP
+	ModelParameters.ModelUnitToCentimeter = UnitModelToCentimeter;
+	return UE::CADKernel::Private::ToRawData(Model, CADKernelRawData);
+#else
+	return false;
+#endif
+}
+
+bool UParametricSurfaceData::SetRepresentation(A3DRiRepresentationItem* Representation, int32 MaterialID, double UnitRepresentationToCentimeter)
+{
+#if PLATFORM_DESKTOP
+	ModelParameters.ModelUnitToCentimeter = UnitRepresentationToCentimeter;
+	return UE::CADKernel::TechSoftUtilities::ToRawData(Representation, MaterialID, TechSoftRawData);
+#else
+	return false;
+#endif
+}
+
 #undef LOCTEXT_NAMESPACE // "ParametricSurfaceData"
diff --git a/Engine/Source/Runtime/Datasmith/CADKernel/Engine/Private/TechSoft/TechSoftLibrary.cpp b/Engine/Source/Runtime/Datasmith/CADKernel/Engine/Private/TechSoft/TechSoftLibrary.cpp
--- a/Engine/Source/Runtime/Datasmith/CADKernel/Engine/Private/TechSoft/TechSoftLibrary.cpp
+++ b/Engine/Source/Runtime/Datasmith/CADKernel/Engine/Private/TechSoft/TechSoftLibrary.cpp
@@ -34,7 +34,7 @@ namespace UE::CADKernel
 			if (Status == A3D_SUCCESS)
 			{
 				Status = A3DDllInitialize(A3D_DLL_MAJORVERSION, A3D_DLL_MINORVERSION);
-				if (Status == A3D_SUCCESS)
+				if (Status == A3D_SUCCESS || Status == A3D_INITIALIZE_ALREADY_CALLED)
 				{
 					bIsInitialized = true;
 					return true;
diff --git a/Engine/Source/Runtime/Datasmith/CADKernel/Engine/Private/TechSoft/TechSoftMesh.cpp b/Engine/Source/Runtime/Datasmith/CADKernel/Engine/Private/TechSoft/TechSoftMesh.cpp
--- a/Engine/Source/Runtime/Datasmith/CADKernel/Engine/Private/TechSoft/TechSoftMesh.cpp
+++ b/Engine/Source/Runtime/Datasmith/CADKernel/Engine/Private/TechSoft/TechSoftMesh.cpp
@@ -42,7 +42,7 @@ namespace UE::CADKernel::TechSoftUtilities
 		{
 		}
 
-		bool CollectTriangles(const A3DTessFaceData& FaceTessData, const A3DTopoFace* TopoFace, int32 InGroupID, FCollectionResult& Result);
+		bool CollectTriangles(const A3DTessFaceData& FaceTessData, const A3DTopoFace* TopoFace, int32 InGroupID, double TextureUnit, FCollectionResult& Result);
 
 		void AddFaceTriangle(const uint32 TriangleCount, uint32& StartIndexInOut);
 		void AddFaceTriangleWithUniqueNormal(const uint32 TriCount, uint32& StartIndexInOut);
@@ -65,16 +65,16 @@ namespace UE::CADKernel::TechSoftUtilities
 			const FArray3i OtherIndices{ NextIndex + 0, NextIndex + 1, NextIndex + 2 };
 
 			const FVector3f Normals[3]{
-				{FVector3f(TessellationNormals[NormalIndices[0]])},
-				{FVector3f(TessellationNormals[NormalIndices[1]])},
-				{FVector3f(TessellationNormals[NormalIndices[2]])},
+				{FVector3f(TessellationNormals[NormalIndices[0] / 3])},
+				{FVector3f(TessellationNormals[NormalIndices[1] / 3])},
+				{FVector3f(TessellationNormals[NormalIndices[2] / 3])},
 			};
 			FaceNormals.Append(Normals, 3);
 
 			const FVector2f TexCoords[3]{
-				{FVector2f(TessellationTexCoords[TexCoordIndicies[0]] * UVScale)},
-				{FVector2f(TessellationTexCoords[TexCoordIndicies[1]] * UVScale)},
-				{FVector2f(TessellationTexCoords[TexCoordIndicies[2]] * UVScale)},
+				{FVector2f(TessellationTexCoords[TexCoordIndicies[0] / 2] * UVScale)},
+				{FVector2f(TessellationTexCoords[TexCoordIndicies[1] / 2] * UVScale)},
+				{FVector2f(TessellationTexCoords[TexCoordIndicies[2] / 2] * UVScale)},
 			};
 			FaceTexCoords.Append(TexCoords, 3);
 
@@ -451,7 +451,7 @@ namespace UE::CADKernel::TechSoftUtilities
 		}
 	}
 
-	bool FFaceTriangleCollector::CollectTriangles(const A3DTessFaceData& FaceTessData, const A3DTopoFace* TopoFace, int32 InGroupID, FFaceTriangleCollector::FCollectionResult& Result)
+	bool FFaceTriangleCollector::CollectTriangles(const A3DTessFaceData& FaceTessData, const A3DTopoFace* TopoFace, int32 InGroupID, double TextureUnit, FFaceTriangleCollector::FCollectionResult& Result)
 	{
 		FaceNormals.Empty(FaceNormals.Num());
 		FaceTexCoords.Empty(FaceTexCoords.Num());
@@ -466,7 +466,6 @@ namespace UE::CADKernel::TechSoftUtilities
 			MaterialID = FaceTessData.m_puiStyleIndexes[0];
 		}
 
-		double TextureUnit = 1.;
 		// #cad_import: Where TextureUnit should come from?
 		UVScale = GetUVScale(TopoFace, TextureUnit);
 
@@ -474,31 +473,31 @@ namespace UE::CADKernel::TechSoftUtilities
 		uint32 FaceSetIndex = 0;
 		bool bMustProcess = true;
 
-		uint32 LastTrianguleIndex = FaceTessData.m_uiStartTriangulated;
+		uint32 LastTriangulatedIndex = FaceTessData.m_uiStartTriangulated;
 		uint32 UsedEntitiesFlags = FaceTessData.m_usUsedEntitiesFlags;
 		LastVertexIndex = 0;
 
 		if (bMustProcess && UsedEntitiesFlags & kA3DTessFaceDataTriangle)
 		{
-			AddFaceTriangle(FaceTessData.m_puiSizesTriangulated[FaceSetIndex++], LastTrianguleIndex);
+			AddFaceTriangle(FaceTessData.m_puiSizesTriangulated[FaceSetIndex++], LastTriangulatedIndex);
 			bMustProcess = FaceTessData.m_uiSizesTriangulatedSize > FaceSetIndex;
 		}
 
 		if (bMustProcess && UsedEntitiesFlags & kA3DTessFaceDataTriangleOneNormal)
 		{
-			AddFaceTriangleWithUniqueNormal(FaceTessData.m_puiSizesTriangulated[FaceSetIndex++], LastTrianguleIndex);
+			AddFaceTriangleWithUniqueNormal(FaceTessData.m_puiSizesTriangulated[FaceSetIndex++], LastTriangulatedIndex);
 			bMustProcess = FaceTessData.m_uiSizesTriangulatedSize > FaceSetIndex;
 		}
 
 		if (bMustProcess && UsedEntitiesFlags & kA3DTessFaceDataTriangleTextured)
 		{
-			AddFaceTriangleWithTexture(FaceTessData.m_puiSizesTriangulated[FaceSetIndex++], FaceTessData.m_uiTextureCoordIndexesSize, LastTrianguleIndex);
+			AddFaceTriangleWithTexture(FaceTessData.m_puiSizesTriangulated[FaceSetIndex++], FaceTessData.m_uiTextureCoordIndexesSize, LastTriangulatedIndex);
 			bMustProcess = FaceTessData.m_uiSizesTriangulatedSize > FaceSetIndex;
 		}
 
 		if (bMustProcess && UsedEntitiesFlags & kA3DTessFaceDataTriangleOneNormalTextured)
 		{
-			AddFaceTriangleWithUniqueNormalAndTexture(FaceTessData.m_puiSizesTriangulated[FaceSetIndex++], FaceTessData.m_uiTextureCoordIndexesSize, LastTrianguleIndex);
+			AddFaceTriangleWithUniqueNormalAndTexture(FaceTessData.m_puiSizesTriangulated[FaceSetIndex++], FaceTessData.m_uiTextureCoordIndexesSize, LastTriangulatedIndex);
 			bMustProcess = FaceTessData.m_uiSizesTriangulatedSize > FaceSetIndex;
 		}
 
@@ -508,7 +507,7 @@ namespace UE::CADKernel::TechSoftUtilities
 			for (uint32 FanIndex = 0; FanIndex < FanCount; ++FanIndex)
 			{
 				uint32 VertexCount = FaceTessData.m_puiSizesTriangulated[FaceSetIndex++];
-				AddFaceTriangleFan(VertexCount, LastTrianguleIndex);
+				AddFaceTriangleFan(VertexCount, LastTriangulatedIndex);
 			}
 			bMustProcess = FaceTessData.m_uiSizesTriangulatedSize > FaceSetIndex;
 		}
@@ -521,7 +520,7 @@ namespace UE::CADKernel::TechSoftUtilities
 				ensure((FaceTessData.m_puiSizesTriangulated[FaceSetIndex] & kA3DTessFaceDataNormalSingle) != 0);
 
 				uint32 VertexCount = FaceTessData.m_puiSizesTriangulated[FaceSetIndex++] & kA3DTessFaceDataNormalMask;
-				AddFaceTriangleFanWithUniqueNormal(VertexCount, LastTrianguleIndex);
+				AddFaceTriangleFanWithUniqueNormal(VertexCount, LastTriangulatedIndex);
 			}
 			bMustProcess = FaceTessData.m_uiSizesTriangulatedSize > FaceSetIndex;
 		}
@@ -532,7 +531,7 @@ namespace UE::CADKernel::TechSoftUtilities
 			for (uint32 FanIndex = 0; FanIndex < FanCount; ++FanIndex)
 			{
 				uint32 VertexCount = FaceTessData.m_puiSizesTriangulated[FaceSetIndex++];
-				AddFaceTriangleFanWithTexture(VertexCount, FaceTessData.m_uiTextureCoordIndexesSize, LastTrianguleIndex);
+				AddFaceTriangleFanWithTexture(VertexCount, FaceTessData.m_uiTextureCoordIndexesSize, LastTriangulatedIndex);
 			}
 			bMustProcess = FaceTessData.m_uiSizesTriangulatedSize > FaceSetIndex;
 		}
@@ -545,7 +544,7 @@ namespace UE::CADKernel::TechSoftUtilities
 				ensure((FaceTessData.m_puiSizesTriangulated[FaceSetIndex] & kA3DTessFaceDataNormalSingle) != 0);
 
 				uint32 VertexCount = FaceTessData.m_puiSizesTriangulated[FaceSetIndex++] & kA3DTessFaceDataNormalMask;
-				AddFaceTriangleFanWithUniqueNormalAndTexture(VertexCount, FaceTessData.m_uiTextureCoordIndexesSize, LastTrianguleIndex);
+				AddFaceTriangleFanWithUniqueNormalAndTexture(VertexCount, FaceTessData.m_uiTextureCoordIndexesSize, LastTriangulatedIndex);
 			}
 			bMustProcess = FaceTessData.m_uiSizesTriangulatedSize > FaceSetIndex;
 		}
@@ -556,7 +555,7 @@ namespace UE::CADKernel::TechSoftUtilities
 			for (A3DUns32 StripeIndex = 0; StripeIndex < StripeSize; ++StripeIndex)
 			{
 				A3DUns32 PointCount = FaceTessData.m_puiSizesTriangulated[FaceSetIndex++];
-				AddFaceTriangleStripe(PointCount, LastTrianguleIndex);
+				AddFaceTriangleStripe(PointCount, LastTriangulatedIndex);
 			}
 			bMustProcess = FaceTessData.m_uiSizesTriangulatedSize > FaceSetIndex;
 		}
@@ -572,11 +571,11 @@ namespace UE::CADKernel::TechSoftUtilities
 				// Is there only one normal for the entire stripe?
 				if (bIsOneNormal == false)
 				{
-					AddFaceTriangleStripe(PointCount, LastTrianguleIndex);
+					AddFaceTriangleStripe(PointCount, LastTriangulatedIndex);
 					continue;
 				}
 
-				AddFaceTriangleStripeWithUniqueNormal(PointCount, LastTrianguleIndex);
+				AddFaceTriangleStripeWithUniqueNormal(PointCount, LastTriangulatedIndex);
 			}
 			bMustProcess = FaceTessData.m_uiSizesTriangulatedSize > FaceSetIndex;
 		}
@@ -587,7 +586,7 @@ namespace UE::CADKernel::TechSoftUtilities
 			for (A3DUns32 StripeIndex = 0; StripeIndex < StripeSize; ++StripeIndex)
 			{
 				A3DUns32 PointCount = FaceTessData.m_puiSizesTriangulated[FaceSetIndex++];
-				AddFaceTriangleStripeWithTexture(PointCount, FaceTessData.m_uiTextureCoordIndexesSize, LastTrianguleIndex);
+				AddFaceTriangleStripeWithTexture(PointCount, FaceTessData.m_uiTextureCoordIndexesSize, LastTriangulatedIndex);
 			}
 			bMustProcess = FaceTessData.m_uiSizesTriangulatedSize > FaceSetIndex;
 		}
@@ -598,7 +597,7 @@ namespace UE::CADKernel::TechSoftUtilities
 			for (A3DUns32 StripeIndex = 0; StripeIndex < StripeSize; ++StripeIndex)
 			{
 				A3DUns32 PointCount = FaceTessData.m_puiSizesTriangulated[FaceSetIndex++] & kA3DTessFaceDataNormalMask;
-				AddFaceTriangleStripeWithUniqueNormalAndTexture(PointCount, FaceTessData.m_uiTextureCoordIndexesSize, LastTrianguleIndex);
+				AddFaceTriangleStripeWithUniqueNormalAndTexture(PointCount, FaceTessData.m_uiTextureCoordIndexesSize, LastTriangulatedIndex);
 			}
 			bMustProcess = FaceTessData.m_uiSizesTriangulatedSize > FaceSetIndex;
 		}
@@ -615,8 +614,9 @@ namespace UE::CADKernel::TechSoftUtilities
 	class FRepresentationConverter
 	{
 	public:
-		FRepresentationConverter(const A3DRiRepresentationItem* RepresentationItem, MeshUtilities::FMeshWrapperAbstract& InMeshWrapper)
+		FRepresentationConverter(const A3DRiRepresentationItem* RepresentationItem, double InTextureUnit, MeshUtilities::FMeshWrapperAbstract& InMeshWrapper)
 			: MeshWrapper(InMeshWrapper)
+			, TextureUnit(InTextureUnit)
 		{
 			TechSoft::TUniqueObject<A3DRiRepresentationItemData> RepresentationItemData(RepresentationItem);
 			if (!RepresentationItemData.IsValid())
@@ -667,6 +667,7 @@ namespace UE::CADKernel::TechSoftUtilities
 		MeshUtilities::FMeshWrapperAbstract& MeshWrapper;
 		const A3DTess3D* TessellationBase = nullptr;
 		const A3DTopoBrepData* BRepData = nullptr;
+		double TextureUnit = 1.;
 	};
 
 	bool FRepresentationConverter::CollectBRepFaces(TArray<A3DTopoFace*>& BRepFaces, int32& TriangleCount)
@@ -821,7 +822,7 @@ namespace UE::CADKernel::TechSoftUtilities
 
 			if (MeshWrapper.IsFaceGroupValid(GroupID))
 			{
-				TriangleCollector.CollectTriangles(FaceTessData, TopoFace, GroupID, CollectionResult);
+				TriangleCollector.CollectTriangles(FaceTessData, TopoFace, GroupID, TextureUnit, CollectionResult);
 
 				MeshWrapper.StartFaceTriangles(0, CollectionResult.Normals, CollectionResult.TexCoords);
 				MeshWrapper.AddFaceTriangles(CollectionResult.FaceTriangles);
@@ -871,12 +872,12 @@ namespace UE::CADKernel::TechSoftUtilities
 
 		ParseTriangles(BRepFaces);
 
-		return false;
+		return true;
 	}
 
-	bool AddRepresentation(A3DRiRepresentationItem* RepresentationItem, MeshUtilities::FMeshWrapperAbstract& MeshWrapper)
+	bool AddRepresentation(A3DRiRepresentationItem* RepresentationItem, double ModelUnitToMeter, MeshUtilities::FMeshWrapperAbstract& MeshWrapper)
 	{
-		FRepresentationConverter Converter(RepresentationItem, MeshWrapper);
+		FRepresentationConverter Converter(RepresentationItem, ModelUnitToMeter, MeshWrapper);
 
 		return Converter.IsValid() ? Converter.Convert() : false;
 	}
diff --git a/Engine/Source/Runtime/Datasmith/CADKernel/Engine/Private/TechSoft/TechSoftUtilities.cpp b/Engine/Source/Runtime/Datasmith/CADKernel/Engine/Private/TechSoft/TechSoftUtilities.cpp
--- a/Engine/Source/Runtime/Datasmith/CADKernel/Engine/Private/TechSoft/TechSoftUtilities.cpp
+++ b/Engine/Source/Runtime/Datasmith/CADKernel/Engine/Private/TechSoft/TechSoftUtilities.cpp
@@ -238,7 +238,7 @@ namespace UE::CADKernel
 		A3DEntityGetType(Representation, &Type);
 		if (Type == kA3DTypeRiPolyBrepModel)
 		{
-			return AddRepresentation(Representation, MeshWrapper);
+			return AddRepresentation(Representation, Context.ModelParams.ModelUnitToCentimeter, MeshWrapper);
 		}
 
 		TechSoft::TUniqueObject<A3DRiRepresentationItemData> RepresentationData(Representation);
@@ -274,7 +274,7 @@ namespace UE::CADKernel
 				// TechSoft is feeling the required structure according to the data structure asked for
 				if (TessellateRepresentation(BrepModel, TessellationSettings))
 				{
-					AddRepresentation(BrepModel, MeshWrapper);
+					AddRepresentation(BrepModel, Context.ModelParams.ModelUnitToCentimeter, MeshWrapper);
 				}
 			}
 		}
@@ -282,10 +282,12 @@ namespace UE::CADKernel
 		{
 			if (TessellateRepresentation(Representation, TessellationSettings))
 			{
-				AddRepresentation(Representation, MeshWrapper);
+				AddRepresentation(Representation, Context.ModelParams.ModelUnitToCentimeter, MeshWrapper);
 			}
 		}
 
+		MeshWrapper.Complete();
+
 		return true;
 	}
 
@@ -453,29 +455,40 @@ namespace UE::CADKernel
 		};
 	}
 
-	TSharedPtr<FJsonObject> TechSoftUtilities::GetJsonObject(A3DAsmProductOccurrence* ProductOcccurence)
+	TSharedPtr<FJsonObject> TechSoftUtilities::GetJsonObject(A3DEntity* Entity, bool bIsLegacy)
 	{
-		TechSoft::TUniqueObject<A3DRootBaseData> RootBaseData(ProductOcccurence);
+		TechSoft::TUniqueObject<A3DRootBaseData> RootBaseData(Entity);
 
 		if (RootBaseData.IsValid() && RootBaseData->m_uiSize > 0)
 		{
-			TechSoft::TUniqueObject<A3DMiscAttributeData> AttributeData(RootBaseData->m_ppAttributes[0]);
-			if (AttributeData->m_uiSize > 0 && AttributeData->m_asSingleAttributesData[0].m_eType == kA3DModellerAttributeTypeString)
+			if (bIsLegacy)
 			{
-				FString JsonString = UTF8_TO_TCHAR(AttributeData->m_asSingleAttributesData[0].m_pcData);
+				TechSoft::TUniqueObject<A3DMiscAttributeData> AttributeData(RootBaseData->m_ppAttributes[0]);
+				if (AttributeData->m_uiSize > 0 && AttributeData->m_asSingleAttributesData[0].m_eType == kA3DModellerAttributeTypeString)
+				{
+					FString JsonString = UTF8_TO_TCHAR(AttributeData->m_asSingleAttributesData[0].m_pcData);
 
-				TSharedRef<TJsonReader<>> JsonReader = TJsonReaderFactory<>::Create(JsonString);
-				TSharedPtr<FJsonObject> JsonObject;
+					TSharedRef<TJsonReader<>> JsonReader = TJsonReaderFactory<>::Create(JsonString);
 
-				if (FJsonSerializer::Deserialize(JsonReader, JsonObject))
-				{
-					return JsonObject;
+					TSharedPtr<FJsonObject> JsonObject;
+					if (FJsonSerializer::Deserialize(JsonReader, JsonObject))
+					{
+						return JsonObject;
+					}
 				}
 			}
 		}
 
 		return {};
 	}
+
+	bool TechSoftUtilities::ToRawData(A3DRiRepresentationItem* Representation, int32 MaterialID, TArray<uint8>& RawDataOut)
+	{
+		// See TechSoftUtils::SaveBodiesToPrcFile
+		// Set attributes on representation not ProductOccurence
+		return false;
+	}
+
 }
 #else
 namespace UE::CADKernel
diff --git a/Engine/Source/Runtime/Datasmith/CADKernel/Engine/Private/TechSoft/TechSoftUtilities.h b/Engine/Source/Runtime/Datasmith/CADKernel/Engine/Private/TechSoft/TechSoftUtilities.h
--- a/Engine/Source/Runtime/Datasmith/CADKernel/Engine/Private/TechSoft/TechSoftUtilities.h
+++ b/Engine/Source/Runtime/Datasmith/CADKernel/Engine/Private/TechSoft/TechSoftUtilities.h
@@ -24,9 +24,11 @@ namespace UE::CADKernel
 	{
 		bool Tessellate(A3DRiRepresentationItem* Representation, const FTessellationContext& Settings, MeshUtilities::FMeshWrapperAbstract& MeshWrapper, bool bEmptyMesh = false);
 		bool TessellateRepresentation(A3DRiRepresentationItem* Representation, const FCADKernelTessellationSettings& Settings);
-		bool AddRepresentation(A3DRiRepresentationItem* Representation, MeshUtilities::FMeshWrapperAbstract& MeshWrapper);
+		bool AddRepresentation(A3DRiRepresentationItem* Representation, double ModelUnitToMeter, MeshUtilities::FMeshWrapperAbstract& MeshWrapper);
 		FVector2d GetUVScale(const A3DTopoFace* TopoFace, double TextureUnit);
-		TSharedPtr<FJsonObject> GetJsonObject(A3DAsmProductOccurrence* ProductOcccurence);
+		TSharedPtr<FJsonObject> GetJsonObject(A3DEntity* Entity, bool bIsLegacy = false);
+		bool ToRawData(A3DRiRepresentationItem* Representation, int32 MaterialID, TArray<uint8>& RawDataOut);
+
 
 		/*
 		 ** Sew the in coming BReps together
diff --git a/Engine/Source/Runtime/Datasmith/CADKernel/Engine/Public/ParametricSurfaceData.h b/Engine/Source/Runtime/Datasmith/CADKernel/Engine/Public/ParametricSurfaceData.h
--- a/Engine/Source/Runtime/Datasmith/CADKernel/Engine/Public/ParametricSurfaceData.h
+++ b/Engine/Source/Runtime/Datasmith/CADKernel/Engine/Public/ParametricSurfaceData.h
@@ -195,11 +195,10 @@ public:
 	UPROPERTY()
 	double UnitMultiplier = 0.1;
 
-#ifdef WITH_HOOPS
-	bool bWithTechSoft = true;
-#else
-	bool bWithTechSoft = false;
-#endif
+	UPROPERTY(Transient)
+	bool bWithTechSoft = bWithHoops;
+
+	static bool bWithHoops;
 
 	bool operator == (const FCADKernelTessellationSettings& Other) const
 	{
@@ -267,8 +266,6 @@ public:
 
 	virtual bool IsValid() { return CADKernelRawData.Num() > 0 || TechSoftRawData.Num() > 0; }
 
-	virtual bool SetFromFile(const TCHAR* FilePath, bool bForTechSoft = false);
-
 	void SetModelParameters(const FCADKernelModelParameters& InModelParameters) { ModelParameters = InModelParameters; }
 	const FCADKernelModelParameters& GetModelParameters() const { return ModelParameters; }
 	FCADKernelModelParameters& GetModelParameters() { return ModelParameters; }
@@ -280,8 +277,12 @@ public:
 	void SetLastTessellationSettings(const FCADKernelTessellationSettings& InTessellationSettings) { LastTessellationSettings = InTessellationSettings; }
 	const FCADKernelTessellationSettings& GetLastTessellationSettings() const { return LastTessellationSettings; }
 
-	const TArray<uint8>& GetRawData(bool bForTechSoft = false) const { return bForTechSoft ? TechSoftRawData : CADKernelRawData; }
+
+	virtual bool SetFromFile(const TCHAR* FilePath, bool bForTechSoft = false);
 	
+	bool SetModel(TSharedPtr<UE::CADKernel::FModel>& Model, double UnitModelToCentimeter = 1.);
+	bool SetRepresentation(A3DRiRepresentationItem* Representation, int32 MaterialID, double UnitRepresentationToCentimeter = 1.);
+
 	TSharedPtr<UE::CADKernel::FModel> GetModel();
 	A3DRiRepresentationItem* GetRepresentation();
 
