Fixed issue 711801
Root cause: Forgot to compile without TechSoft :-(

#rnx

diff --git a/Engine/Source/Runtime/Datasmith/CADKernel/Engine/Private/ParametricSurfaceData.cpp b/Engine/Source/Runtime/Datasmith/CADKernel/Engine/Private/ParametricSurfaceData.cpp
--- a/Engine/Source/Runtime/Datasmith/CADKernel/Engine/Private/ParametricSurfaceData.cpp
+++ b/Engine/Source/Runtime/Datasmith/CADKernel/Engine/Private/ParametricSurfaceData.cpp
@@ -82,7 +82,7 @@ TSharedPtr<UE::CADKernel::FModel> UParametricSurfaceData::GetModel()
 
 A3DRiRepresentationItem* UParametricSurfaceData::GetRepresentation()
 {
-#if PLATFORM_DESKTOP
+#ifdef WITH_HOOPS
 	using namespace UE::CADKernel;
 
 	if (TechSoftRawData.IsEmpty())
@@ -145,7 +145,7 @@ bool UParametricSurfaceData::SetModel(TSharedPtr<UE::CADKernel::FModel>& Model,
 
 bool UParametricSurfaceData::SetRepresentation(A3DRiRepresentationItem* Representation, int32 MaterialID, double UnitRepresentationToCentimeter)
 {
-#if PLATFORM_DESKTOP
+#ifdef WITH_HOOPS
 	ModelParameters.ModelUnitToCentimeter = UnitRepresentationToCentimeter;
 	return UE::CADKernel::TechSoftUtilities::ToRawData(Representation, MaterialID, TechSoftRawData);
 #else
