[In Progress] CADKernelEngine module
- Miscellaneous changes to accomodate InterchangeCADTranslator plugin

#rnx

diff --git a/Engine/Source/Runtime/Datasmith/CADKernel/Engine/Private/TechSoft/TechSoftMesh.cpp b/Engine/Source/Runtime/Datasmith/CADKernel/Engine/Private/TechSoft/TechSoftMesh.cpp
--- a/Engine/Source/Runtime/Datasmith/CADKernel/Engine/Private/TechSoft/TechSoftMesh.cpp
+++ b/Engine/Source/Runtime/Datasmith/CADKernel/Engine/Private/TechSoft/TechSoftMesh.cpp
@@ -849,8 +849,6 @@ namespace UE::CADKernel::TechSoftUtilities
 		
 		Vertices.Reserve(VertexCount);
 
-		//const double ScaleFactor = BodyUnit * FImportParameters::GUnitScale;
-		const double ScaleFactor = 1.;
 		double* Coordinates = TessellationBaseData->m_pdCoords;
 
 		for (int32 Index = 0; Index < VertexCount; ++Index, Coordinates += 3)
diff --git a/Engine/Source/Runtime/Datasmith/CADKernel/Engine/Private/TechSoft/TechSoftToCADKernel.cpp b/Engine/Source/Runtime/Datasmith/CADKernel/Engine/Private/TechSoft/TechSoftToCADKernel.cpp
--- a/Engine/Source/Runtime/Datasmith/CADKernel/Engine/Private/TechSoft/TechSoftToCADKernel.cpp
+++ b/Engine/Source/Runtime/Datasmith/CADKernel/Engine/Private/TechSoft/TechSoftToCADKernel.cpp
@@ -2,30 +2,45 @@
 
 #include "CADKernelEngine.h"
 
-#if PLATFORM_DESKTOP
 #ifdef WITH_HOOPS
 #include "TechSoftUniqueObjectImpl.h"
 
 namespace UE::CADKernel
 {
+	namespace TechSoftUtilities
+	{
+		using namespace UE::CADKernel;
+
+		class FRepresentationToModel
+		{
+		public:
+			FRepresentationToModel(A3DRiRepresentationItem* InRepresentation)
+				: Representation(InRepresentation)
+			{
+			}
+
+			void Execute()
+			{
+			}
 
-	TSharedPtr<UE::CADKernel::FModel> FTechSoftUtilities::TechSoftToCADKernel(A3DRiRepresentationItem* Representation)
+			TSharedPtr<FModel> Model;
+			A3DRiRepresentationItem* Representation;
+		};
+	}
+
+	TSharedPtr<FModel> FTechSoftUtilities::TechSoftToCADKernel(A3DRiRepresentationItem* Representation)
 	{
 		if (!FTechSoftLibrary::IsInitialized())
 		{
 			return {};
 		}
 
-		return {};
+		TechSoftUtilities::FRepresentationToModel Converter(Representation);
+
+		Converter.Execute();
+
+		return Converter.Model;
 	}
 }
 #else
-namespace UE::CADKernel
-{
-	TSharedPtr<UE::CADKernel::FModel> FTechSoftUtilities::TechSoftToCADKernel(A3DRiRepresentationItem* Representation)
-	{
-		return {};
-	}
-}
-#endif
 #endif
diff --git a/Engine/Source/Runtime/Datasmith/CADKernel/Engine/Private/TechSoft/TechSoftUtilities.cpp b/Engine/Source/Runtime/Datasmith/CADKernel/Engine/Private/TechSoft/TechSoftUtilities.cpp
--- a/Engine/Source/Runtime/Datasmith/CADKernel/Engine/Private/TechSoft/TechSoftUtilities.cpp
+++ b/Engine/Source/Runtime/Datasmith/CADKernel/Engine/Private/TechSoft/TechSoftUtilities.cpp
@@ -512,6 +512,11 @@ namespace UE::CADKernel
 	{
 		return nullptr;
 	}
+
+	TSharedPtr<UE::CADKernel::FModel> FTechSoftUtilities::TechSoftToCADKernel(A3DRiRepresentationItem* Representation)
+	{
+		return {};
+	}
 }
 #endif
 #endif
diff --git a/Engine/Source/Runtime/Datasmith/CADKernel/Engine/Public/CADKernelEngine.h b/Engine/Source/Runtime/Datasmith/CADKernel/Engine/Public/CADKernelEngine.h
--- a/Engine/Source/Runtime/Datasmith/CADKernel/Engine/Public/CADKernelEngine.h
+++ b/Engine/Source/Runtime/Datasmith/CADKernel/Engine/Public/CADKernelEngine.h
@@ -57,9 +57,9 @@ namespace UE::CADKernel
 	class FTechSoftLibrary
 	{
 	public:
-		static const TCHAR* GetVersion();
-		static bool Initialize();
-		static bool IsInitialized() { return bIsInitialized; }
+		static CADKERNELENGINE_API const TCHAR* GetVersion();
+		static bool CADKERNELENGINE_API Initialize();
+		static bool CADKERNELENGINE_API IsInitialized() { return bIsInitialized; }
 
 	private:
 		static bool bIsInitialized;
