FMeshDescriptionWrapper::OrientMesh() crash fix

diff --git a/Engine/Source/Runtime/Datasmith/CADKernel/Engine/Private/MeshDescriptionWrapper.cpp b/Engine/Source/Runtime/Datasmith/CADKernel/Engine/Private/MeshDescriptionWrapper.cpp
--- a/Engine/Source/Runtime/Datasmith/CADKernel/Engine/Private/MeshDescriptionWrapper.cpp
+++ b/Engine/Source/Runtime/Datasmith/CADKernel/Engine/Private/MeshDescriptionWrapper.cpp
@@ -661,7 +661,7 @@ namespace UE::CADKernel::MeshUtilities
 							float MaxComponent = 0;
 							for (const FVertexInstanceID VertexInstanceID : CoincidentVertexInstanceIdSet)
 							{
-								FVector Normal = (FVector)Normals[VertexInstanceID];
+								FVector Normal = (FVector)VertexInstanceNormals[VertexInstanceID];
 								if (FMath::Abs(MaxComponent) < FMath::Abs(Normal[VertexIndex]))
 								{
 									MaxComponent = Normal[VertexIndex];
