[In Progress] Interchange CAD translator plugin
- Finalized generation of model and representation for wire parser
- Implemented all methods of FTechSoftLibrary
- Added assets references by entities generated by Wire parser
- Replaced generation of material entities with that of material instances in wire parser

#rnx

#virtualized

diff --git a/Engine/Plugins/Experimental/Enterprise/DatasmithInterchange/Source/DatasmithInterchange/Private/InterchangeDatasmithTranslator.cpp b/Engine/Plugins/Experimental/Enterprise/DatasmithInterchange/Source/DatasmithInterchange/Private/InterchangeDatasmithTranslator.cpp
--- a/Engine/Plugins/Experimental/Enterprise/DatasmithInterchange/Source/DatasmithInterchange/Private/InterchangeDatasmithTranslator.cpp
+++ b/Engine/Plugins/Experimental/Enterprise/DatasmithInterchange/Source/DatasmithInterchange/Private/InterchangeDatasmithTranslator.cpp
@@ -56,8 +56,79 @@
 
 #define LOCTEXT_NAMESPACE "DatasmithInterchange"
 
+static TSet<FString> ExcludedFormats{
+	TEXT("catpart"),
+	TEXT("catproduct"),
+	TEXT("catshape"),
+	TEXT("cgr"),
+	TEXT("3dxml"),
+	TEXT("3drep"),
+	TEXT("model"),
+	TEXT("session"),
+	TEXT("exp"),
+	TEXT("dlv"),
+	TEXT("asm.*s"),
+	TEXT("creo.*"),
+	TEXT("creo"),
+	TEXT("neu.*"),
+	TEXT("neu"),
+	TEXT("prt.*"),
+	TEXT("xas"),
+	TEXT("xpr"),
+	TEXT("iam"),
+	TEXT("ipt"),
+	TEXT("iges"),
+	TEXT("igs"),
+	TEXT("jt"),
+	TEXT("sat"),
+	TEXT("sab"),
+	TEXT("sldasm"),
+	TEXT("sldprt"),
+	TEXT("step"),
+	TEXT("stp"),
+	TEXT("stpz"),
+	TEXT("stpx"),
+	TEXT("stpxz"),
+	TEXT("xml"),
+	TEXT("x_t"),
+	TEXT("x_b"),
+	TEXT("xmt"),
+	TEXT("xmt_txt"),
+	TEXT("asm"),
+	TEXT("prt"),
+	TEXT("par"),
+	TEXT("psm"),
+	TEXT("dwg"),
+	TEXT("dxf"),
+	TEXT("ifc"),
+	TEXT("ifczip"),
+	TEXT("hsf"),
+	TEXT("prc"),
+	TEXT("3mf"),
+	TEXT("3ds"),
+	TEXT("dae"),
+	TEXT("dwf"),
+	TEXT("dwfx"),
+	TEXT("nwd"),
+	TEXT("mf1"),
+	TEXT("arc"),
+	TEXT("unv"),
+	TEXT("pkg"),
+	TEXT("dgn"),  // available with Hoops Exchange 2023
+	TEXT("stl"),
+	TEXT("u3d"),
+	TEXT("vda"),
+	TEXT("vrml"),
+	TEXT("wrl"),
+	TEXT("wire"),
+	TEXT("3dm"),
+};
+
 TArray<FString> UInterchangeDatasmithTranslator::GetSupportedFormats() const
 {
+	IConsoleVariable* CVar = IConsoleManager::Get().FindConsoleVariable(TEXT("Interchange.FeatureFlags.Import.CAD"), false);
+	bool bInterchangeCADEnabled = CVar ? CVar->GetBool() : false;
+
 	const TArray<FString> DatasmithFormats = FDatasmithTranslatorManager::Get().GetSupportedFormats();
 	TArray<FString> Formats;
 	Formats.Reserve(DatasmithFormats.Num() - 1);
@@ -69,6 +140,11 @@ TArray<FString> UInterchangeDatasmithTranslator::GetSupportedFormats() const
 			continue;
 		}
 
+		if (bInterchangeCADEnabled && ExcludedFormats.Contains(Format.ToLower()))
+		{
+			continue;
+		}
+
 		Formats.Add(Format);
 	}
 
@@ -79,6 +155,9 @@ bool UInterchangeDatasmithTranslator::CanImportSourceData(const UInterchangeSour
 {
 	using namespace UE::DatasmithImporter;
 
+	IConsoleVariable* CVar = IConsoleManager::Get().FindConsoleVariable(TEXT("Interchange.FeatureFlags.Import.CAD"), false);
+	bool bInterchangeCADEnabled = CVar ? CVar->GetBool() : false;
+
 	const FString FilePath = InSourceData->GetFilename();
 	const FString FileExtension = FPaths::GetExtension(FilePath);
 	if (FileExtension.Equals(TEXT("gltf"), ESearchCase::IgnoreCase) || FileExtension.Equals(TEXT("glb"), ESearchCase::IgnoreCase) || FileExtension.Equals(TEXT("fbx"), ESearchCase::IgnoreCase))
@@ -87,6 +166,11 @@ bool UInterchangeDatasmithTranslator::CanImportSourceData(const UInterchangeSour
 		return false;
 	}
 
+	if (bInterchangeCADEnabled && ExcludedFormats.Contains(FileExtension.ToLower()))
+	{
+		return false;
+	}
+
 	const FSourceUri FileNameUri = FSourceUri::FromFilePath(FilePath);
 	TSharedPtr<FExternalSource> ExternalSource = IExternalSourceModule::GetOrCreateExternalSource(FileNameUri);
 
diff --git a/Engine/Source/Runtime/Datasmith/CADKernel/Engine/Private/TechSoft/TechSoftLibrary.cpp b/Engine/Source/Runtime/Datasmith/CADKernel/Engine/Private/TechSoft/TechSoftLibrary.cpp
--- a/Engine/Source/Runtime/Datasmith/CADKernel/Engine/Private/TechSoft/TechSoftLibrary.cpp
+++ b/Engine/Source/Runtime/Datasmith/CADKernel/Engine/Private/TechSoft/TechSoftLibrary.cpp
@@ -6,9 +6,18 @@
 #include "CADKernelEngine.h"
 #if PLATFORM_DESKTOP
 
+#include "TechSoftUniqueObjectImpl.h"
+
 #include "HAL/PlatformProcess.h"
 #include "Misc/Paths.h"
 
+// #cad_import: This is ugly :-(
+#define RECTIFY_STRUTC_SIZE(Struct, StructName) \
+{ \
+	TechSoft::TUniqueObject<StructName> OtherStruct; \
+	Struct.m_usStructSize = OtherStruct->m_usStructSize; \
+}
+
 namespace UE::CADKernel
 {
 	bool FTechSoftLibrary::bIsInitialized = false;
@@ -50,6 +59,308 @@ namespace UE::CADKernel
 		static const FString Version = TEXT("TechSoft ") + FString::FromInt(A3D_DLL_MAJORVERSION) + TEXT(".") + FString::FromInt(A3D_DLL_MINORVERSION) + TEXT(".") + FString::FromInt(A3D_DLL_UPDATEVERSION);
 		return bIsInitialized ? *Version : TEXT("TechSoft uninitialized");
 	}
+
+	A3DRiRepresentationItem* FTechSoftLibrary::CreateRIBRep(const TArray<A3DTopoShell*>& TopoShells)
+	{
+		if (!FTechSoftLibrary::Initialize() || TopoShells.IsEmpty())
+		{
+			return nullptr;
+		}
+
+		A3DTopoConnex* TopoConnexPtr = nullptr;
+		{
+			TechSoft::TUniqueObject<A3DTopoConnexData> TopoConnexData;
+			TopoConnexData->m_ppShells = (A3DTopoShell**)TopoShells.GetData();
+			TopoConnexData->m_uiShellSize = TopoShells.Num();
+			if (A3DTopoConnexCreate(TopoConnexData.GetPtr(), &TopoConnexPtr) != A3DStatus::A3D_SUCCESS)
+			{
+				return nullptr;
+			}
+		}
+
+		A3DTopoBrepData* TopoBRepDataPtr = nullptr;
+		{
+			TechSoft::TUniqueObject<A3DTopoBrepDataData> TopoBRepData;
+			TopoBRepData->m_uiConnexSize = 1;
+			TopoBRepData->m_ppConnexes = &TopoConnexPtr;
+			if (A3DTopoBrepDataCreate(TopoBRepData.GetPtr(), &TopoBRepDataPtr) != A3DStatus::A3D_SUCCESS)
+			{
+				return nullptr;
+			}
+		}
+
+		TechSoft::TUniqueObject<A3DRiBrepModelData> RiBRepModelData;
+		RiBRepModelData->m_pBrepData = TopoBRepDataPtr;
+		RiBRepModelData->m_bSolid = false;
+		A3DRiBrepModel* RiBrepModelPtr = nullptr;
+		if (A3DRiBrepModelCreate(RiBRepModelData.GetPtr(), &RiBrepModelPtr) != A3DStatus::A3D_SUCCESS)
+		{
+			return nullptr;
+		}
+
+		return RiBrepModelPtr;
+	}
+
+	A3DTopoFace* FTechSoftLibrary::CreateTopoFaceWithNaturalLoop(A3DSurfBase* CarrierSurface)
+	{
+		if (!FTechSoftLibrary::Initialize())
+		{
+			return nullptr;
+		}
+
+		TechSoft::TUniqueObject<A3DTopoFaceData> Face;
+		Face->m_pSurface = CarrierSurface;
+		Face->m_bHasTrimDomain = false;
+		Face->m_ppLoops = nullptr;
+		Face->m_uiLoopSize = 0;
+		Face->m_uiOuterLoopIndex = 0;
+		Face->m_dTolerance = 0.01; //mm
+
+		return CreateTopoFace(*Face);
+	}
+
+	A3DCrvNurbs* FTechSoftLibrary::CreateTrimNurbsCurve(A3DCrvNurbs* CurveNurbsPtr, double UMin, double UMax, bool bIs2D)
+	{
+		if (!FTechSoftLibrary::Initialize() || CurveNurbsPtr == nullptr)
+		{
+			return nullptr;
+		}
+
+		TechSoft::TUniqueObject<A3DCrvTransformData> TransformCurveData;
+
+		TransformCurveData->m_bIs2D = bIs2D;
+		TransformCurveData->m_sParam.m_sInterval.m_dMin = UMin;
+		TransformCurveData->m_sParam.m_sInterval.m_dMax = UMax;
+		TransformCurveData->m_sParam.m_dCoeffA = 1.;
+		TransformCurveData->m_sParam.m_dCoeffB = 0.;
+		TransformCurveData->m_pBasisCrv = CurveNurbsPtr;
+		TransformCurveData->m_pTransfo = nullptr;
+
+		TransformCurveData->m_sTrsf.m_sXVector.m_dX = 1.;
+		TransformCurveData->m_sTrsf.m_sYVector.m_dY = 1.;
+		TransformCurveData->m_sTrsf.m_sScale.m_dX = 1.;
+		TransformCurveData->m_sTrsf.m_sScale.m_dY = 1.;
+		TransformCurveData->m_sTrsf.m_sScale.m_dZ = 1.;
+
+		A3DCrvTransform* CurveTransformPtr = nullptr;
+		if (A3DCrvTransformCreate(TransformCurveData.GetPtr(), &CurveTransformPtr) != A3DStatus::A3D_SUCCESS)
+		{
+			return nullptr;
+		}
+
+		TechSoft::TUniqueObject<A3DCrvNurbsData> NurbsCurveData;
+		if (A3DCrvBaseGetAsNurbs(CurveTransformPtr, 0.01 /*mm*/, /*bUseSameParameterization*/ true, NurbsCurveData.GetPtr()) != A3DStatus::A3D_SUCCESS)
+		{
+			return nullptr;
+		}
+
+		return CreateCurveNurbs(*NurbsCurveData);
+	}
+
+	A3DTopoShell* FTechSoftLibrary::CreateTopoShell(A3DTopoShellData& TopoShellData)
+	{
+		if (!FTechSoftLibrary::Initialize())
+		{
+			return nullptr;
+		}
+
+		RECTIFY_STRUTC_SIZE(TopoShellData, A3DTopoShellData)
+
+		A3DTopoShell* TopoShellPtr = nullptr;
+		if (A3DTopoShellCreate(&TopoShellData, &TopoShellPtr) != A3DStatus::A3D_SUCCESS)
+		{
+			return nullptr;
+		}
+		return TopoShellPtr;
+	}
+
+	A3DSurfNurbs* FTechSoftLibrary::CreateSurfaceNurbs(A3DSurfNurbsData& SurfaceNurbsData)
+	{
+		if (!FTechSoftLibrary::Initialize())
+		{
+			return nullptr;
+		}
+
+		RECTIFY_STRUTC_SIZE(SurfaceNurbsData, A3DSurfNurbsData)
+
+		A3DSurfNurbs* SurfaceNurbsPtr = nullptr;
+		if (A3DSurfNurbsCreate(&SurfaceNurbsData, &SurfaceNurbsPtr) != A3DStatus::A3D_SUCCESS)
+		{
+			return nullptr;
+		}
+		return SurfaceNurbsPtr;
+	}
+
+	A3DStatus FTechSoftLibrary::SetEntityGraphicsColor(A3DEntity* Entity, FColor Color)
+	{
+		if (!FTechSoftLibrary::Initialize() || Entity == nullptr)
+		{
+			return A3DStatus::A3D_NOT_AVAILABLE;
+		}
+
+		TechSoft::TUniqueObject<A3DRootBaseWithGraphicsData> BaseWithGraphicsData(Entity);
+
+		//Create a style color
+		TechSoft::TUniqueObjectFromIndex<A3DGraphRgbColorData> RgbColor;
+
+		RgbColor->m_dRed = Color.R / 255.;
+		RgbColor->m_dGreen = Color.G / 255.;
+		RgbColor->m_dBlue = Color.B / 255.;
+		
+		A3DUns32 ColorIndex = A3D_DEFAULT_COLOR_INDEX;
+		if (A3DGlobalInsertGraphRgbColor(RgbColor.GetPtr(), &ColorIndex) != A3DStatus::A3D_SUCCESS)
+		{
+			ColorIndex = A3D_DEFAULT_COLOR_INDEX;
+		}
+
+		A3DUns32 StyleIndex = 0;
+		TechSoft::TUniqueObjectFromIndex<A3DGraphStyleData> StyleData;
+		StyleData->m_bMaterial = false;
+		StyleData->m_bVPicture = false;
+		StyleData->m_dWidth = 0.1; // default
+		A3DUns8 Alpha = Color.A;
+		if (Alpha < 255)
+		{
+			StyleData->m_bIsTransparencyDefined = true;
+			StyleData->m_ucTransparency = 255 - Alpha;
+		}
+		else
+		{
+			StyleData->m_bIsTransparencyDefined = false;
+			StyleData->m_ucTransparency = 0;
+		}
+
+		StyleData->m_bSpecialCulling = false;
+		StyleData->m_bBackCulling = false;
+		StyleData->m_uiRgbColorIndex = ColorIndex;
+
+		StyleIndex = A3D_DEFAULT_STYLE_INDEX;
+		if (A3DGlobalInsertGraphStyle(StyleData.GetPtr(), &StyleIndex) != A3DStatus::A3D_SUCCESS)
+		{
+			StyleIndex = A3D_DEFAULT_STYLE_INDEX;
+		}
+
+		TechSoft::TUniqueObject<A3DGraphicsData> GraphicsData;
+
+		GraphicsData->m_uiStyleIndex = StyleIndex;
+		GraphicsData->m_usBehaviour = kA3DGraphicsShow;
+		GraphicsData->m_usBehaviour |= kA3DGraphicsSonHeritColor;
+
+		BaseWithGraphicsData->m_pGraphics = nullptr;
+		if (A3DGraphicsCreate(GraphicsData.GetPtr(), &BaseWithGraphicsData->m_pGraphics) != A3DStatus::A3D_SUCCESS || BaseWithGraphicsData->m_pGraphics == nullptr)
+		{
+			return A3DStatus::A3D_ERROR;
+		}
+
+		return A3DRootBaseWithGraphicsSet(Entity, BaseWithGraphicsData.GetPtr());
+	}
+
+	A3DTopoFace* FTechSoftLibrary::CreateTopoFace(A3DTopoFaceData& TopoFaceData)
+	{
+		if (!FTechSoftLibrary::Initialize())
+		{
+			return nullptr;
+		}
+
+		RECTIFY_STRUTC_SIZE(TopoFaceData, A3DTopoFaceData)
+
+		A3DTopoFace* TopoFacePtr = nullptr;
+		if (A3DTopoFaceCreate(&TopoFaceData, &TopoFacePtr) != A3DStatus::A3D_SUCCESS)
+		{
+			return nullptr;
+		}
+		return TopoFacePtr;
+	}
+
+	A3DTopoLoop* FTechSoftLibrary::CreateTopoLoop(A3DTopoLoopData& TopoLoopData)
+	{
+		if (!FTechSoftLibrary::Initialize())
+		{
+			return nullptr;
+		}
+
+		RECTIFY_STRUTC_SIZE(TopoLoopData, A3DTopoLoopData)
+
+		A3DTopoLoop* TopoLoopPtr = nullptr;
+		if (A3DTopoLoopCreate(&TopoLoopData, &TopoLoopPtr) != A3DStatus::A3D_SUCCESS)
+		{
+			return nullptr;
+		}
+		return TopoLoopPtr;
+	}
+
+	A3DTopoEdge* FTechSoftLibrary::CreateTopoEdge()
+	{
+		if (!FTechSoftLibrary::Initialize())
+		{
+			return nullptr;
+		}
+
+		TechSoft::TUniqueObject<A3DTopoEdgeData> EdgeData;
+		return CreateTopoEdge(*EdgeData);
+	}
+
+	A3DTopoEdge* FTechSoftLibrary::CreateTopoEdge(A3DTopoEdgeData& TopoEdgeData)
+	{
+		if (!FTechSoftLibrary::Initialize())
+		{
+			return nullptr;
+		}
+
+		RECTIFY_STRUTC_SIZE(TopoEdgeData, A3DTopoEdgeData)
+
+		A3DTopoEdge* TopoEdgePtr = nullptr;
+		if (A3DTopoEdgeCreate(&TopoEdgeData, &TopoEdgePtr) != A3DStatus::A3D_SUCCESS)
+		{
+			return nullptr;
+		}
+		return TopoEdgePtr;
+	}
+
+	A3DTopoCoEdge* FTechSoftLibrary::CreateTopoCoEdge(A3DTopoCoEdgeData& TopoCoEdgeData)
+	{
+		if (!FTechSoftLibrary::Initialize())
+		{
+			return nullptr;
+		}
+
+		RECTIFY_STRUTC_SIZE(TopoCoEdgeData, A3DTopoCoEdgeData)
+
+		A3DTopoCoEdge* TopoCoEdgePtr = nullptr;
+		if (A3DTopoCoEdgeCreate(&TopoCoEdgeData, &TopoCoEdgePtr) != A3DStatus::A3D_SUCCESS)
+		{
+			return nullptr;
+		}
+		return TopoCoEdgePtr;
+	}
+
+	A3DCrvNurbs* FTechSoftLibrary::CreateCurveNurbs(A3DCrvNurbsData& CurveNurbsData)
+	{
+		if (!FTechSoftLibrary::Initialize())
+		{
+			return nullptr;
+		}
+
+		RECTIFY_STRUTC_SIZE(CurveNurbsData, A3DCrvNurbsData)
+
+		A3DCrvNurbs* CurveNurbsPtr = nullptr;
+		if (A3DCrvNurbsCreate(&CurveNurbsData, &CurveNurbsPtr) != A3DStatus::A3D_SUCCESS)
+		{
+			return nullptr;
+		}
+
+		return CurveNurbsPtr;
+	}
+
+	A3DStatus FTechSoftLibrary::LinkCoEdges(A3DTopoCoEdge* CoEdgePtr, A3DTopoCoEdge* NeighbourCoEdgePtr)
+	{
+		if (!FTechSoftLibrary::Initialize())
+		{
+			return A3DStatus::A3D_NOT_AVAILABLE;
+		}
+
+		return A3DTopoCoEdgeSetNeighbour(CoEdgePtr, NeighbourCoEdgePtr);
+	}
 #else
 	bool FTechSoftLibrary::Initialize()
 	{
diff --git a/Engine/Source/Runtime/Datasmith/CADKernel/Engine/Private/TechSoft/TechSoftUtilities.cpp b/Engine/Source/Runtime/Datasmith/CADKernel/Engine/Private/TechSoft/TechSoftUtilities.cpp
--- a/Engine/Source/Runtime/Datasmith/CADKernel/Engine/Private/TechSoft/TechSoftUtilities.cpp
+++ b/Engine/Source/Runtime/Datasmith/CADKernel/Engine/Private/TechSoft/TechSoftUtilities.cpp
@@ -18,170 +18,121 @@
 
 namespace UE::CADKernel
 {
-	bool FTechSoftUtilities::Save(A3DRiRepresentationItem* Representation, const FString& FilePath)
+	bool FTechSoftUtilities::Save(const TArray<A3DRiRepresentationItem*>& Representations, const FString& FilePath, const FString& AttributesStr)
 	{
-		if (!FTechSoftLibrary::IsInitialized())
+		if (!FTechSoftLibrary::Initialize())
 		{
 			return false;
 		}
 
-		return false;
-	}
-
-	bool FTechSoftUtilities::Tessellate(A3DRiRepresentationItem* Representation, const FTessellationContext& Context, FMeshDescription& Mesh, bool bEmptyMesh)
-	{
-		using namespace UE::CADKernel::MeshUtilities;
-
-		if (!FTechSoftLibrary::IsInitialized())
+		// Create PartDefinition
+		TechSoft::TUniqueObject<A3DAsmPartDefinitionData> PartDefinitionData;
+		PartDefinitionData->m_uiRepItemsSize = Representations.Num();
+		PartDefinitionData->m_ppRepItems = (A3DRiRepresentationItem**)Representations.GetData();
+		A3DAsmPartDefinition* PartDefinition = nullptr;
+		if (A3DAsmPartDefinitionCreate(PartDefinitionData.GetPtr(), &PartDefinition) != A3DStatus::A3D_SUCCESS || PartDefinition == nullptr)
 		{
 			return false;
 		}
 
-		TSharedPtr<FMeshWrapperAbstract> MeshWrapper = FMeshWrapperAbstract::MakeWrapper(Context, Mesh);
-		return TechSoftUtilities::Tessellate(Representation, Context, *MeshWrapper, bEmptyMesh);
-	}
-
-	bool FTechSoftUtilities::Tessellate(A3DRiRepresentationItem* Representation, const FTessellationContext& Context, UE::Geometry::FDynamicMesh3& Mesh, bool bEmptyMesh)
-	{
-		using namespace UE::CADKernel::MeshUtilities;
-
-		if (!FTechSoftLibrary::IsInitialized())
+		TechSoft::TUniqueObject<A3DAsmProductOccurrenceData> ProductOccurrenceData;
+		ProductOccurrenceData->m_pPart = PartDefinition;
+		A3DAsmProductOccurrence* ProductOccurrence = nullptr;
+		if (A3DAsmProductOccurrenceCreate(ProductOccurrenceData.GetPtr(), &ProductOccurrence) != A3DStatus::A3D_SUCCESS)
 		{
 			return false;
 		}
 
-		TSharedPtr<FMeshWrapperAbstract> MeshWrapper = FMeshWrapperAbstract::MakeWrapper(Context, Mesh);
-		return TechSoftUtilities::Tessellate(Representation, Context, *MeshWrapper, bEmptyMesh);
-	}
+		// Add MaterialTable as attribute to ProductOccurrence
+		std::string StringAnsi(TCHAR_TO_UTF8(*AttributesStr));
 
-	A3DRiRepresentationItem* FTechSoftUtilities::CreateRIBRep(A3DTopoShell* TopoShellPtr)
-	{
-		if (!FTechSoftLibrary::IsInitialized())
-		{
-			return nullptr;
-		}
-
-		return nullptr;
-	}
+		TechSoft::TUniqueObject<A3DMiscSingleAttributeData> SingleAttributeData;
+		SingleAttributeData->m_eType = kA3DModellerAttributeTypeString;
+		SingleAttributeData->m_pcTitle = (char*)"Attributes";
+		SingleAttributeData->m_pcData = (char*)StringAnsi.c_str();
 
-	A3DTopoEdge* FTechSoftUtilities::CreateTopoEdge()
-	{
-		if (!FTechSoftLibrary::IsInitialized())
+		TechSoft::TUniqueObject<A3DMiscAttributeData> AttributesData;
+		AttributesData->m_pcTitle = SingleAttributeData->m_pcTitle;
+		AttributesData->m_asSingleAttributesData = SingleAttributeData.GetPtr();
+		AttributesData->m_uiSize = 1;
+		A3DMiscAttribute* Attributes = nullptr;
+		if (A3DMiscAttributeCreate(AttributesData.GetPtr(), &Attributes) != A3DStatus::A3D_SUCCESS)
 		{
-			return nullptr;
-		}
-
-		return nullptr;
-	}
-
-	A3DTopoFace* FTechSoftUtilities::CreateTopoFaceWithNaturalLoop(A3DSurfBase* CarrierSurface)
-	{
-		if (!FTechSoftLibrary::IsInitialized())
-		{
-			return nullptr;
-		}
-
-		return nullptr;
-	}
-
-	A3DCrvNurbs* FTechSoftUtilities::CreateTrimNurbsCurve(A3DCrvNurbs* CurveNurbsPtr, double UMin, double UMax, bool bIs2D)
-	{
-		if (!FTechSoftLibrary::IsInitialized())
-		{
-			return nullptr;
+			return false;
 		}
 
-		return nullptr;
-	}
-
-	A3DTopoShell* FTechSoftUtilities::CreateTopoShell(const A3DTopoShellData& TopoShellData)
-	{
-		if (!FTechSoftLibrary::IsInitialized())
+		TechSoft::TUniqueObject<A3DRootBaseData> RootBaseData;
+		RootBaseData->m_pcName = SingleAttributeData->m_pcTitle;
+		RootBaseData->m_ppAttributes = &Attributes;
+		RootBaseData->m_uiSize = 1;
+		if (A3DRootBaseSet(ProductOccurrence, RootBaseData.GetPtr()) != A3DStatus::A3D_SUCCESS)
 		{
-			return nullptr;
+			return false;
 		}
 
-		return nullptr;
-	}
-
-	A3DSurfNurbs* FTechSoftUtilities::CreateSurfaceNurbs(const A3DSurfNurbsData& SurfaceNurbsData)
-	{
-		if (!FTechSoftLibrary::IsInitialized())
+		// Create ModelFile
+		TechSoft::TUniqueObject<A3DAsmModelFileData> ModelFileData;
+		ModelFileData->m_uiPOccurrencesSize = 1;
+		ModelFileData->m_dUnit = 1.0;
+		ModelFileData->m_ppPOccurrences = &ProductOccurrence;
+		A3DAsmModelFile* ModelFile = nullptr;
+		if (A3DAsmModelFileCreate(ModelFileData.GetPtr(), &ModelFile) != A3DStatus::A3D_SUCCESS)
 		{
-			return nullptr;
+			return false;
 		}
 
-		return nullptr;
-	}
-
-	A3DStatus FTechSoftUtilities::SetEntityGraphicsColor(A3DEntity* InEntity, FColor Color)
-	{
-		if (!FTechSoftLibrary::IsInitialized())
-		{
-			return A3DStatus::A3D_NOT_AVAILABLE;
-		}
+		// Save ModelFile to Pcr file
+		TechSoft::TUniqueObject<A3DRWParamsExportPrcData> ParamsExportData;
+		ParamsExportData->m_bCompressBrep = false;
+		ParamsExportData->m_bCompressTessellation = false;
 
-		return A3DStatus::A3D_NOT_IMPLEMENTED;
-	}
+#if PLATFORM_WINDOWS
+		A3DUTF8Char HsfFileName[MAX_PATH];
+		FCStringAnsi::Strncpy(HsfFileName, TCHAR_TO_UTF8(*FilePath), MAX_PATH);
+#elif PLATFORM_LINUX
+		A3DUTF8Char HsfFileName[PATH_MAX];
+		FCStringAnsi::Strncpy(HsfFileName, TCHAR_TO_UTF8(*FilePath), PATH_MAX);
+#else
+#error Platform not supported
+#endif // PLATFORM_WINDOWS
 
-	A3DTopoFace* FTechSoftUtilities::CreateTopoFace(const A3DTopoFaceData& TopoFaceData)
-	{
-		if (!FTechSoftLibrary::IsInitialized())
+		if (A3DAsmModelFileExportToPrcFile(ModelFile, ParamsExportData.GetPtr(), HsfFileName, nullptr) != A3DStatus::A3D_SUCCESS)
 		{
-			return nullptr;
+			return false;
 		}
 
-		return nullptr;
-	}
+		A3DAsmModelFileUnloadParts(ModelFile, 1, &ProductOccurrence);
 
-	A3DTopoLoop* FTechSoftUtilities::CreateTopoLoop(const A3DTopoLoopData& TopoLoopData)
-	{
-		if (!FTechSoftLibrary::IsInitialized())
-		{
-			return nullptr;
-		}
+		// #ueent_techsoft: Deleting the model seems to delete the entire content. To be double-checked
+		//A3DEntityDelete(Attributes);
 
-		return nullptr;
+		return true;
 	}
 
-	A3DTopoEdge* FTechSoftUtilities::CreateTopoEdge(const A3DTopoEdgeData& TopoEdgeData)
+	bool FTechSoftUtilities::Tessellate(A3DRiRepresentationItem* Representation, const FTessellationContext& Context, FMeshDescription& Mesh, bool bEmptyMesh)
 	{
-		if (!FTechSoftLibrary::IsInitialized())
-		{
-			return nullptr;
-		}
-
-		return nullptr;
-	}
+		using namespace UE::CADKernel::MeshUtilities;
 
-	A3DTopoCoEdge* FTechSoftUtilities::CreateTopoCoEdge(const A3DTopoCoEdgeData& TopoCoEdgeData)
-	{
-		if (!FTechSoftLibrary::IsInitialized())
+		if (!FTechSoftLibrary::Initialize())
 		{
-			return nullptr;
+			return false;
 		}
 
-		return nullptr;
+		TSharedPtr<FMeshWrapperAbstract> MeshWrapper = FMeshWrapperAbstract::MakeWrapper(Context, Mesh);
+		return TechSoftUtilities::Tessellate(Representation, Context, *MeshWrapper, bEmptyMesh);
 	}
 
-	A3DCrvNurbs* FTechSoftUtilities::CreateCurveNurbs(const A3DCrvNurbsData& CurveNurbsData)
+	bool FTechSoftUtilities::Tessellate(A3DRiRepresentationItem* Representation, const FTessellationContext& Context, UE::Geometry::FDynamicMesh3& Mesh, bool bEmptyMesh)
 	{
-		if (!FTechSoftLibrary::IsInitialized())
-		{
-			return nullptr;
-		}
-
-		return nullptr;
-	}
+		using namespace UE::CADKernel::MeshUtilities;
 
-	A3DStatus FTechSoftUtilities::LinkCoEdges(A3DTopoCoEdge* CoEdgePtr, A3DTopoCoEdge* NeighbourCoEdgePtr)
-	{
-		if (!FTechSoftLibrary::IsInitialized())
+		if (!FTechSoftLibrary::Initialize())
 		{
-			return A3DStatus::A3D_NOT_AVAILABLE;
+			return false;
 		}
 
-		return A3DStatus::A3D_NOT_IMPLEMENTED;
+		TSharedPtr<FMeshWrapperAbstract> MeshWrapper = FMeshWrapperAbstract::MakeWrapper(Context, Mesh);
+		return TechSoftUtilities::Tessellate(Representation, Context, *MeshWrapper, bEmptyMesh);
 	}
 
 	A3DRiRepresentationItem* FTechSoftUtilities::GetRepresentation(const TArray<uint8>& TechSoftRawData)
@@ -225,7 +176,7 @@ namespace UE::CADKernel
 
 	bool TechSoftUtilities::Tessellate(A3DRiRepresentationItem* Representation, const FTessellationContext& Context, MeshUtilities::FMeshWrapperAbstract& MeshWrapper, bool bEmptyMesh)
 	{
-		if (!FTechSoftLibrary::IsInitialized())
+		if (!FTechSoftLibrary::Initialize())
 		{
 			return false;
 		}
@@ -293,7 +244,7 @@ namespace UE::CADKernel
 
 	bool TechSoftUtilities::TessellateRepresentation(A3DRiRepresentationItem* Representation, const FCADKernelTessellationSettings& Settings)
 	{
-		if (!FTechSoftLibrary::IsInitialized())
+		if (!FTechSoftLibrary::Initialize())
 		{
 			return false;
 		}
@@ -335,7 +286,7 @@ namespace UE::CADKernel
 
 	bool TechSoftUtilities::SewBReps(const TArray<A3DRiBrepModel*>& BRepsIn, double Tolerance, TArray<A3DRiBrepModel*>& BRepsOut)
 	{
-		if (!FTechSoftLibrary::IsInitialized())
+		if (!FTechSoftLibrary::Initialize())
 		{
 			return false;
 		}
@@ -493,7 +444,7 @@ namespace UE::CADKernel
 #else
 namespace UE::CADKernel
 {
-	bool FTechSoftUtilities::Save(A3DRiRepresentationItem* Representation, const FString& FilePath)
+	bool FTechSoftUtilities::Save(TArray<A3DRiRepresentationItem*> Representations, const FString& FilePath)
 	{
 		return false;
 	}
diff --git a/Engine/Source/Runtime/Datasmith/CADKernel/Engine/Public/CADKernelEngine.h b/Engine/Source/Runtime/Datasmith/CADKernel/Engine/Public/CADKernelEngine.h
--- a/Engine/Source/Runtime/Datasmith/CADKernel/Engine/Public/CADKernelEngine.h
+++ b/Engine/Source/Runtime/Datasmith/CADKernel/Engine/Public/CADKernelEngine.h
@@ -61,6 +61,25 @@ namespace UE::CADKernel
 		static bool CADKERNELENGINE_API Initialize();
 		static bool CADKERNELENGINE_API IsInitialized() { return bIsInitialized; }
 
+#ifdef WITH_HOOPS
+		static CADKERNELENGINE_API A3DRiRepresentationItem* CreateRIBRep(const TArray<A3DTopoShell*>& TopoShells);
+
+		static CADKERNELENGINE_API A3DTopoEdge* CreateTopoEdge();
+		static CADKERNELENGINE_API A3DTopoEdge* CreateTopoEdge(A3DTopoEdgeData& TopoEdgeData);
+		static CADKERNELENGINE_API A3DTopoFace* CreateTopoFaceWithNaturalLoop(A3DSurfBase* CarrierSurface);
+		static CADKERNELENGINE_API A3DTopoShell* CreateTopoShell(A3DTopoShellData& TopoShellData);
+		static CADKERNELENGINE_API A3DTopoFace* CreateTopoFace(A3DTopoFaceData& TopoFaceData);
+		static CADKERNELENGINE_API A3DTopoLoop* CreateTopoLoop(A3DTopoLoopData& TopoLoopData);
+		static CADKERNELENGINE_API A3DTopoCoEdge* CreateTopoCoEdge(A3DTopoCoEdgeData& TopoCoEdgeData);
+		static CADKERNELENGINE_API A3DStatus LinkCoEdges(A3DTopoCoEdge* CoEdgePtr, A3DTopoCoEdge* NeighbourCoEdgePtr);
+
+		static CADKERNELENGINE_API A3DCrvNurbs* CreateTrimNurbsCurve(A3DCrvNurbs* CurveNurbsPtr, double UMin, double UMax, bool bIs2D);
+		static CADKERNELENGINE_API A3DSurfNurbs* CreateSurfaceNurbs(A3DSurfNurbsData& SurfaceNurbsData);
+		static CADKERNELENGINE_API A3DCrvNurbs* CreateCurveNurbs(A3DCrvNurbsData& CurveNurbsData);
+
+		static CADKERNELENGINE_API A3DStatus SetEntityGraphicsColor(A3DEntity* InEntity, FColor Color);
+#endif
+
 	private:
 		static bool bIsInitialized;
 	};
@@ -73,7 +92,11 @@ namespace UE::CADKernel
 		 ** @param bSewModel: Optional argument to force a stitching of the model
 		 ** @param StitchingTolerance: Optional stitching tolerance in mm
 		 */
-		static CADKERNELENGINE_API bool Save(A3DRiRepresentationItem* Representation, const FString& FilePath);
+		static CADKERNELENGINE_API bool Save(const TArray<A3DRiRepresentationItem*>& Representations, const FString& FilePath, const FString& Attributes = FString());
+		static CADKERNELENGINE_API bool Save(A3DRiRepresentationItem* Representation, const FString& FilePath, const FString& Attributes = FString())
+		{
+			return Save(TArray<A3DRiRepresentationItem*>( &Representation, 1 ), FilePath, Attributes);
+		}
 		static CADKERNELENGINE_API bool Tessellate(A3DRiRepresentationItem* Representation, const FTessellationContext& Context, FMeshDescription& MeshOut, bool bEmptyMesh = false);
 		static CADKERNELENGINE_API bool Tessellate(A3DRiRepresentationItem* Representation, const FTessellationContext& Context, UE::Geometry::FDynamicMesh3& MeshOut, bool bEmptyMesh = false);
 
@@ -81,25 +104,6 @@ namespace UE::CADKernel
 		/* Limitation: Non of the implicit geometry is supported */
 		static CADKERNELENGINE_API A3DRiRepresentationItem* CADKernelToTechSoft(TSharedPtr<UE::CADKernel::FModel>& Model);
 		static CADKERNELENGINE_API A3DRiRepresentationItem* GetRepresentation(const TArray<uint8>& TechSoftRawData);
-
-#ifdef WITH_HOOPS
-		static CADKERNELENGINE_API A3DRiRepresentationItem* CreateRIBRep(A3DTopoShell* TopoShellPtr);
-
-		static CADKERNELENGINE_API A3DTopoEdge* CreateTopoEdge();
-		static CADKERNELENGINE_API A3DTopoFace* CreateTopoFaceWithNaturalLoop(A3DSurfBase* CarrierSurface);
-		static CADKERNELENGINE_API A3DTopoShell* CreateTopoShell(const A3DTopoShellData& TopoShellData);
-		static CADKERNELENGINE_API A3DTopoFace* CreateTopoFace(const A3DTopoFaceData& TopoFaceData);
-		static CADKERNELENGINE_API A3DTopoLoop* CreateTopoLoop(const A3DTopoLoopData& TopoLoopData);
-		static CADKERNELENGINE_API A3DTopoEdge* CreateTopoEdge(const A3DTopoEdgeData& TopoEdgeData);
-		static CADKERNELENGINE_API A3DTopoCoEdge* CreateTopoCoEdge(const A3DTopoCoEdgeData& TopoCoEdgeData);
-		static CADKERNELENGINE_API A3DStatus LinkCoEdges(A3DTopoCoEdge* CoEdgePtr, A3DTopoCoEdge* NeighbourCoEdgePtr);
-
-		static CADKERNELENGINE_API A3DCrvNurbs* CreateTrimNurbsCurve(A3DCrvNurbs* CurveNurbsPtr, double UMin, double UMax, bool bIs2D);
-		static CADKERNELENGINE_API A3DSurfNurbs* CreateSurfaceNurbs(const A3DSurfNurbsData& SurfaceNurbsData);
-		static CADKERNELENGINE_API A3DCrvNurbs* CreateCurveNurbs(const A3DCrvNurbsData& CurveNurbsData);
-
-		static CADKERNELENGINE_API A3DStatus SetEntityGraphicsColor(A3DEntity* InEntity, FColor Color);
-#endif
 	};
 
 	namespace MathUtils
