Fixed for issues: 718909, 718900, 718896
With my apologies

#rnx

diff --git a/Engine/Source/Runtime/Datasmith/CADKernel/Engine/Private/TechSoft/TechSoftUtilities.cpp b/Engine/Source/Runtime/Datasmith/CADKernel/Engine/Private/TechSoft/TechSoftUtilities.cpp
--- a/Engine/Source/Runtime/Datasmith/CADKernel/Engine/Private/TechSoft/TechSoftUtilities.cpp
+++ b/Engine/Source/Runtime/Datasmith/CADKernel/Engine/Private/TechSoft/TechSoftUtilities.cpp
@@ -444,7 +444,7 @@ namespace UE::CADKernel
 #else
 namespace UE::CADKernel
 {
-	bool FTechSoftUtilities::Save(TArray<A3DRiRepresentationItem*> Representations, const FString& FilePath)
+	bool FTechSoftUtilities::Save(const TArray<A3DRiRepresentationItem*>& Representations, const FString& FilePath, const FString& Attributes)
 	{
 		return false;
 	}
