[In Progress] Interchange CAD Translator plugin
- Fixed compilation issues based on non-availability of different third party libraries the code depends on
- Fixed issues with extraction of meshes from wire files: Just use faces' colorid to identify mesh's sections.
- miscellaneous code cleanup: OpenModelUtils::TransferAlMeshToMeshDescription,
- Started implementation of CAD Import settings with Interchange. This could be used by legacy too...
- Added temporary code in the translator to validate that the parsing is correct for Wire and OpenNurbs.
- Added missing code in FInterchangeWireParserImpl::GetParametricSurfaceData: the setting of the different parameters were not done.
- Fixed a bug in FDatasmithPackedMeshes::Serialize and GetDatasmithMeshFromFile. It looks like this new class was never actually used :-(

diff --git a/Engine/Source/Runtime/Datasmith/CADKernel/Engine/Private/CADKernelMesh.cpp b/Engine/Source/Runtime/Datasmith/CADKernel/Engine/Private/CADKernelMesh.cpp
--- a/Engine/Source/Runtime/Datasmith/CADKernel/Engine/Private/CADKernelMesh.cpp
+++ b/Engine/Source/Runtime/Datasmith/CADKernel/Engine/Private/CADKernelMesh.cpp
@@ -34,14 +34,6 @@ namespace UE::CADKernel::Private
 			, MeshWrapper(InMeshWrapper)
 		{}
 
-		uint32 GetMaterialID(const FTopologicalFace& Face)
-		{
-			const uint32 Prefix = Face.GetMaterialId() != 0 ? PointerHash(TEXT("Material")) : PointerHash(TEXT("Color"));
-			const int32 MaterialID = Face.GetMaterialId() != 0 ? Face.GetMaterialId() : Face.GetColorId();
-
-			return HashCombine(Prefix, MaterialID);
-		}
-
 		void AddTrianglesFromFaceMesh(const FFaceMesh& FaceMesh)
 		{
 			const FTopologicalFace& Face = static_cast<const FTopologicalFace&>(FaceMesh.GetGeometricEntity());
@@ -52,7 +44,7 @@ namespace UE::CADKernel::Private
 				return;
 			}
 
-			const uint32 MaterialID = GetMaterialID(Face);
+			const uint32 MaterialID = Face.GetColorId();
 
 			const TArray<int32>& TriangleVertexIndices = FaceMesh.TrianglesVerticesIndex;
 
diff --git a/Engine/Source/Runtime/Datasmith/CADKernel/Engine/Private/MeshDescriptionWrapper.h b/Engine/Source/Runtime/Datasmith/CADKernel/Engine/Private/MeshDescriptionWrapper.h
--- a/Engine/Source/Runtime/Datasmith/CADKernel/Engine/Private/MeshDescriptionWrapper.h
+++ b/Engine/Source/Runtime/Datasmith/CADKernel/Engine/Private/MeshDescriptionWrapper.h
@@ -7,38 +7,9 @@
 
 #include "MeshDescription.h"
 #include "NaniteDefinitions.h"
-#include "StaticMeshAttributes.h"
 
 namespace UE::CADKernel::MeshUtilities
 {
-	class FCADKernelStaticMeshAttributes : public FStaticMeshAttributes
-	{
-	public:
-
-		explicit FCADKernelStaticMeshAttributes(FMeshDescription& InMeshDescription) : FStaticMeshAttributes(InMeshDescription) {}
-
-		virtual void Register(bool bKeepExistingAttribute = false) override;
-
-		bool IsValid() const
-		{
-			return GetVertexInstanceNormals().IsValid() &&
-				GetVertexInstanceTangents().IsValid() &&
-				GetVertexInstanceBinormalSigns().IsValid() &&
-				GetVertexInstanceColors().IsValid() &&
-				GetVertexInstanceUVs().IsValid() &&
-				GetPolygonGroupMaterialSlotNames().IsValid() &&
-				GetPolygonGroups().IsValid();
-		}
-
-		/**
-		 * Enable per-triangle integer attribute named PolyTriGroups (see ExtendedMeshAttribute::PolyTriGroups)
-		 * This integer defines the identifier of the PolyTriGroup containing the triangle.
-		 * In case of mesh coming from a CAD file, a PolyTriGroup is associated to a CAD topological face
-		 */
-		TPolygonAttributesRef<int32> GetPolygonGroups();
-		TPolygonAttributesConstRef<int32> GetPolygonGroups() const;
-	};
-
 	class FMeshDescriptionWrapper : public FMeshWrapperAbstract
 	{
 	public:
diff --git a/Engine/Source/Runtime/Datasmith/CADKernel/Engine/Private/MeshUtilities.cpp b/Engine/Source/Runtime/Datasmith/CADKernel/Engine/Private/MeshUtilities.cpp
--- a/Engine/Source/Runtime/Datasmith/CADKernel/Engine/Private/MeshUtilities.cpp
+++ b/Engine/Source/Runtime/Datasmith/CADKernel/Engine/Private/MeshUtilities.cpp
@@ -12,7 +12,7 @@
 
 namespace UE::CADKernel
 {
-	void ApplyExtractionContext(const FMeshExtractionContext& Context, FMeshDescription& MeshInOut)
+	void FCADKernelUtilities::ApplyExtractionContext(const FMeshExtractionContext& Context, FMeshDescription& MeshInOut)
 	{
 		using namespace MeshUtilities;
 
@@ -20,7 +20,7 @@ namespace UE::CADKernel
 		MeshWrapper->Complete();
 	}
 
-	void ApplyExtractionContext(const FMeshExtractionContext& Context, UE::Geometry::FDynamicMesh3& MeshInOut)
+	void FCADKernelUtilities::ApplyExtractionContext(const FMeshExtractionContext& Context, UE::Geometry::FDynamicMesh3& MeshInOut)
 	{
 		using namespace MeshUtilities;
 
@@ -28,6 +28,13 @@ namespace UE::CADKernel
 		MeshWrapper->Complete();
 	}
 
+	void FCADKernelUtilities::RegisterAttributes(FMeshDescription& MeshInOut, bool bKeepExistingAttribute)
+	{
+		MeshUtilities::FCADKernelStaticMeshAttributes Attributes(MeshInOut);
+		Attributes.Register(bKeepExistingAttribute);
+	}
+
+
 	namespace MeshUtilities
 	{
 		bool FMeshOperations::OrientMesh(FMeshDescription& MeshDescription)
diff --git a/Engine/Source/Runtime/Datasmith/CADKernel/Engine/Private/MeshUtilities.h b/Engine/Source/Runtime/Datasmith/CADKernel/Engine/Private/MeshUtilities.h
--- a/Engine/Source/Runtime/Datasmith/CADKernel/Engine/Private/MeshUtilities.h
+++ b/Engine/Source/Runtime/Datasmith/CADKernel/Engine/Private/MeshUtilities.h
@@ -7,6 +7,7 @@
 #include "ParametricSurfaceData.h"
 
 #include "IntVectorTypes.h"
+#include "StaticMeshAttributes.h"
 #include "Util/DynamicVector.h"
 
 namespace UE::CADKernel
@@ -26,6 +27,34 @@ using FArray3i = UE::Geometry::FVector3i;
 
 namespace UE::CADKernel::MeshUtilities
 {
+	class FCADKernelStaticMeshAttributes : public FStaticMeshAttributes
+	{
+	public:
+
+		explicit FCADKernelStaticMeshAttributes(FMeshDescription& InMeshDescription) : FStaticMeshAttributes(InMeshDescription) {}
+
+		virtual void Register(bool bKeepExistingAttribute = false) override;
+
+		bool IsValid() const
+		{
+			return GetVertexInstanceNormals().IsValid() &&
+				GetVertexInstanceTangents().IsValid() &&
+				GetVertexInstanceBinormalSigns().IsValid() &&
+				GetVertexInstanceColors().IsValid() &&
+				GetVertexInstanceUVs().IsValid() &&
+				GetPolygonGroupMaterialSlotNames().IsValid() &&
+				GetPolygonGroups().IsValid();
+		}
+
+		/**
+		 * Enable per-triangle integer attribute named PolyTriGroups (see ExtendedMeshAttribute::PolyTriGroups)
+		 * This integer defines the identifier of the PolyTriGroup containing the triangle.
+		 * In case of mesh coming from a CAD file, a PolyTriGroup is associated to a CAD topological face
+		 */
+		TPolygonAttributesRef<int32> GetPolygonGroups();
+		TPolygonAttributesConstRef<int32> GetPolygonGroups() const;
+	};
+
 	class FMeshWrapperAbstract
 	{
 	public:
diff --git a/Engine/Source/Runtime/Datasmith/CADKernel/Engine/Public/CADKernelEngine.h b/Engine/Source/Runtime/Datasmith/CADKernel/Engine/Public/CADKernelEngine.h
--- a/Engine/Source/Runtime/Datasmith/CADKernel/Engine/Public/CADKernelEngine.h
+++ b/Engine/Source/Runtime/Datasmith/CADKernel/Engine/Public/CADKernelEngine.h
@@ -39,6 +39,7 @@ namespace UE::CADKernel
 	struct FTessellationContext : public FMeshExtractionContext
 	{
 		FCADKernelTessellationSettings TessellationSettings;
+		FTessellationContext() {};
 		FTessellationContext(const FCADKernelModelParameters& ModelParams, const FCADKernelMeshParameters& MeshParams, const FCADKernelRetessellationSettings& Settings);
 	};
 
@@ -52,6 +53,8 @@ namespace UE::CADKernel
 		/* Note: Meshes are expected to be in their native coordinates system and unit */
 		static void CADKERNELENGINE_API ApplyExtractionContext(const FMeshExtractionContext& Context, FMeshDescription& MeshInOut);
 		static void CADKERNELENGINE_API ApplyExtractionContext(const FMeshExtractionContext& Context, UE::Geometry::FDynamicMesh3& MeshInOut);
+		// Set the proper attributes for FMeshDescription created from CAD data.
+		static void CADKERNELENGINE_API RegisterAttributes(FMeshDescription& MeshInOut, bool bKeepExistingAttribute = false);
 	};
 
 	class FTechSoftLibrary
diff --git a/Engine/Source/Runtime/Datasmith/DatasmithCore/Private/DatasmithMeshSerialization.cpp b/Engine/Source/Runtime/Datasmith/DatasmithCore/Private/DatasmithMeshSerialization.cpp
--- a/Engine/Source/Runtime/Datasmith/DatasmithCore/Private/DatasmithMeshSerialization.cpp
+++ b/Engine/Source/Runtime/Datasmith/DatasmithCore/Private/DatasmithMeshSerialization.cpp
@@ -168,6 +168,9 @@ FMD5Hash FDatasmithPackedMeshes::Serialize(FArchive& Ar, bool bCompressed)
 	// | CustomVersions  | v0 | internal structures use them, we have to pass them along
 	// | Mesh buffer     | v0 | may be compressed, the actual payload
 
+	int32 AntiLegacy = 0;
+	Ar << AntiLegacy;
+
 	const TCHAR* StaticGuard = TEXT("FDatasmithPackedMeshes");
 	FString Guard = Ar.IsLoading() ? TEXT("") : StaticGuard;
 	Ar << Guard;
@@ -357,6 +360,7 @@ FDatasmithPackedMeshes GetDatasmithMeshFromFile(const FString& MeshPath)
 	}
 	else
 	{
+		Archive->Seek(0);
 		Result.Serialize(*Archive);
 
 		if (Archive->IsError())
