Backed out changes related to Datasmith mesh serialization

#rnx

diff --git a/Engine/Source/Developer/Datasmith/DatasmithExporter/Private/DatasmithMeshExporter.cpp b/Engine/Source/Developer/Datasmith/DatasmithExporter/Private/DatasmithMeshExporter.cpp
--- a/Engine/Source/Developer/Datasmith/DatasmithExporter/Private/DatasmithMeshExporter.cpp
+++ b/Engine/Source/Developer/Datasmith/DatasmithExporter/Private/DatasmithMeshExporter.cpp
@@ -144,6 +144,8 @@ bool FDatasmithMeshExporterImpl::WriteMeshFile(const FDatasmithMeshExporterOptio
 		return false;
 	}
 
+	int32 LegacyMeshCount = 0;
+	*Archive << LegacyMeshCount; // the legacy importer expect a mesh count on the first bytes. Just in case a new file would end up parsed by the legacy code...
 	OutHash = Pack.Serialize(*Archive);
 
 	return !Archive->IsError();
diff --git a/Engine/Source/Runtime/Datasmith/DatasmithCore/Private/DatasmithMeshSerialization.cpp b/Engine/Source/Runtime/Datasmith/DatasmithCore/Private/DatasmithMeshSerialization.cpp
--- a/Engine/Source/Runtime/Datasmith/DatasmithCore/Private/DatasmithMeshSerialization.cpp
+++ b/Engine/Source/Runtime/Datasmith/DatasmithCore/Private/DatasmithMeshSerialization.cpp
@@ -168,9 +168,6 @@ FMD5Hash FDatasmithPackedMeshes::Serialize(FArchive& Ar, bool bCompressed)
 	// | CustomVersions  | v0 | internal structures use them, we have to pass them along
 	// | Mesh buffer     | v0 | may be compressed, the actual payload
 
-	int32 AntiLegacy = 0;
-	Ar << AntiLegacy;
-
 	const TCHAR* StaticGuard = TEXT("FDatasmithPackedMeshes");
 	FString Guard = Ar.IsLoading() ? TEXT("") : StaticGuard;
 	Ar << Guard;
@@ -360,7 +357,6 @@ FDatasmithPackedMeshes GetDatasmithMeshFromFile(const FString& MeshPath)
 	}
 	else
 	{
-		Archive->Seek(0);
 		Result.Serialize(*Archive);
 
 		if (Archive->IsError())
