Proper fix for template

[RN] none
#rnx

diff --git a/Engine/Source/Runtime/Datasmith/CADKernel/Base/Private/Core/Session.cpp b/Engine/Source/Runtime/Datasmith/CADKernel/Base/Private/Core/Session.cpp
--- a/Engine/Source/Runtime/Datasmith/CADKernel/Base/Private/Core/Session.cpp
+++ b/Engine/Source/Runtime/Datasmith/CADKernel/Base/Private/Core/Session.cpp
@@ -11,16 +11,34 @@ namespace UE::CADKernel
 FSession FSession::Session(0.01);
 #endif
 
-void FSession::SaveDatabase(const TCHAR* FileName)
+bool FSession::SaveDatabase(const TCHAR* FileName)
 {
 	TSharedPtr<FCADKernelArchive> Archive = FCADKernelArchive::CreateArchiveWriter(*this, FileName);
 	if (!Archive.IsValid())
 	{
 		FMessage::Printf(Log, TEXT("The archive file %s is corrupted\n"), FileName);
+		return false;
 	}
 
-	Database.Serialize(*Archive.Get());
+	Database.Serialize(*Archive);
 	Archive->Close();
+
+	return true;
+}
+
+bool FSession::SaveDatabase(TArray<uint8>& Bytes)
+{
+	TSharedPtr<FCADKernelArchive> Archive = MakeShared<FCADKernelArchive>(*this, Bytes);
+	if (!Archive.IsValid())
+	{
+		FMessage::Printf(Log, TEXT("The archive is corrupted\n"));
+		return false;
+	}
+
+	Database.Serialize(*Archive);
+	Archive->Close();
+
+	return true;
 }
 
 TSharedPtr<FModel> FSession::GetModelAsShared()
@@ -64,27 +82,56 @@ void FSession::SaveDatabase(const TCHAR* FileName, const TArray<TSharedPtr<FEnti
 	Archive->Close();
 }
 
-void FSession::LoadDatabase(const TCHAR* FilePath)
+bool FSession::LoadDatabase(const TCHAR* FilePath)
 {
 	TGuardValue<double> GeometricToleranceGuard(GeometricTolerance, 0.01);
+
 	TSharedPtr<FCADKernelArchive> Archive = FCADKernelArchive::CreateArchiveReader(*this, FilePath);
 	if (!Archive.IsValid())
 	{
 		FMessage::Printf(Log, TEXT("The archive file %s is corrupted\n"), FilePath);
-		return;
+		return false;
 	}
 
 	FModel& Model = GetModel();
 
-	Database.Deserialize(*Archive.Get());
+	Database.Deserialize(*Archive);
 	FModel* ArchiveModel = Archive->ArchiveModel;
 
-	if(ArchiveModel != nullptr)
+	if (ArchiveModel != nullptr)
 	{
 		ArchiveModel->Empty();
 		Database.RemoveEntity(*ArchiveModel);
 	}
 	Archive->Close();
+
+	return true;
+}
+
+bool FSession::LoadDatabase(const TArray<uint8>& Bytes)
+{
+	TGuardValue<double> GeometricToleranceGuard(GeometricTolerance, 0.01);
+
+	TSharedPtr<FCADKernelArchive> Archive = MakeShared<FCADKernelArchive>(*this, Bytes);
+	if (!Archive.IsValid())
+	{
+		FMessage::Printf(Log, TEXT("The archive buffer is corrupted\n"));
+		return false;
+	}
+
+	FModel& Model = GetModel();
+
+	Database.Deserialize(*Archive);
+	FModel* ArchiveModel = Archive->ArchiveModel;
+
+	if (ArchiveModel != nullptr)
+	{
+		ArchiveModel->Empty();
+		Database.RemoveEntity(*ArchiveModel);
+	}
+	Archive->Close();
+
+	return true;
 }
 
 void FSession::AddDatabase(const TArray<uint8>& InRawData)
diff --git a/Engine/Source/Runtime/Datasmith/CADKernel/Base/Public/Core/CADKernelArchive.h b/Engine/Source/Runtime/Datasmith/CADKernel/Base/Public/Core/CADKernelArchive.h
--- a/Engine/Source/Runtime/Datasmith/CADKernel/Base/Public/Core/CADKernelArchive.h
+++ b/Engine/Source/Runtime/Datasmith/CADKernel/Base/Public/Core/CADKernelArchive.h
@@ -7,6 +7,7 @@
 #include "HAL/FileManager.h"
 #include "Serialization/Archive.h"
 #include "Serialization/MemoryReader.h"
+#include "Serialization/MemoryWriter.h"
 
 namespace UE::CADKernel
 {
@@ -18,6 +19,8 @@ class FCADKernelArchive
 {
 public:
 	TUniquePtr<FMemoryReader> MemoryReader;
+	TUniquePtr<FMemoryArchive> MemoryArchive;
+	TSharedPtr<FArchive> ArchivePtr;
 	FArchive& Archive;
 	FSession& Session;
 	FModel* ArchiveModel = nullptr;
@@ -32,8 +35,22 @@ public:
 	 * Archive reader constructor
 	 */
 	FCADKernelArchive(FSession& InSession, const TArray<uint8>& InRawData)
-		: MemoryReader(MakeUnique<FMemoryReader>(InRawData, true))
-		, Archive(*MemoryReader)
+		: MemoryArchive(MakeUnique<FMemoryReader>(InRawData, true))
+		, Archive(*MemoryArchive)
+		, Session(InSession)
+	{
+	}
+
+	FCADKernelArchive(FSession& InSession, TArray<uint8>& InRawData)
+		: MemoryArchive(MakeUnique<FMemoryWriter>(InRawData, true))
+		, Archive(*MemoryArchive)
+		, Session(InSession)
+	{
+	}
+
+	FCADKernelArchive(FSession& InSession, FArchive* InArchive)
+		: ArchivePtr(MakeShareable<FArchive>(InArchive))
+		, Archive(*InArchive)
 		, Session(InSession)
 	{
 	}
@@ -46,18 +63,40 @@ public:
 			FMessage::Printf(Log, TEXT("The archive file %s is corrupted\n"), FileName);
 			return TSharedPtr<FCADKernelArchive>();
 		}
-		return MakeShared<FCADKernelArchive>(InSession, *Ar);
+		return MakeShared<FCADKernelArchive>(InSession, Ar);
 	}
 
 	static TSharedPtr<FCADKernelArchive> CreateArchiveReader(FSession& InSession, const TCHAR* FileName)
 	{
 		FArchive* Ar = IFileManager::Get().CreateFileReader(FileName);
-		if (Ar == nullptr)
+		if (!Ar)
 		{
 			FMessage::Printf(Log, TEXT("The archive file %s is corrupted\n"), FileName);
 			return TSharedPtr<FCADKernelArchive>();
 		}
-		return MakeShared<FCADKernelArchive>(InSession, *Ar);
+		return MakeShared<FCADKernelArchive>(InSession, Ar);
+	}
+
+	static TSharedPtr<FCADKernelArchive> CreateArchiveWriter(FSession& InSession, TArray<uint8>& InRawData)
+	{
+		FArchive* Ar = new FMemoryWriter(InRawData, true);
+		if (!Ar)
+		{
+			FMessage::Printf(Log, TEXT("The archive is corrupted\n"));
+			return TSharedPtr<FCADKernelArchive>();
+		}
+		return MakeShared<FCADKernelArchive>(InSession, Ar);
+	}
+
+	static TSharedPtr<FCADKernelArchive> CreateArchiveReader(FSession& InSession, const TArray<uint8>& InRawData)
+	{
+		FArchive* Ar = new FMemoryReader(InRawData, true);
+		if (!Ar)
+		{
+			FMessage::Printf(Log, TEXT("The archive buffer is corrupted\n"));
+			return TSharedPtr<FCADKernelArchive>();
+		}
+		return MakeShared<FCADKernelArchive>(InSession, Ar);
 	}
 
 	template<typename EntityType>
diff --git a/Engine/Source/Runtime/Datasmith/CADKernel/Base/Public/Core/Session.h b/Engine/Source/Runtime/Datasmith/CADKernel/Base/Public/Core/Session.h
--- a/Engine/Source/Runtime/Datasmith/CADKernel/Base/Public/Core/Session.h
+++ b/Engine/Source/Runtime/Datasmith/CADKernel/Base/Public/Core/Session.h
@@ -66,7 +66,7 @@ public:
 	 * Mandatory: all entity have to have a defined ID
 	 * Use SpawnEntityIdent if needed
 	 */
-	void SaveDatabase(const TCHAR* FilePath);
+	bool SaveDatabase(const TCHAR* FilePath);
 
 	/**
 	 * Save a selection and all the dependencies as a FAchive in a file
@@ -110,7 +110,10 @@ public:
 	 * Load and add a database in the current session database
 	 * Entity ID is set for all loaded entities
 	 */
-	void LoadDatabase(const TCHAR* FilePath);
+	bool LoadDatabase(const TCHAR* FilePath);
+
+	bool SaveDatabase(TArray<uint8>& Bytes);
+	bool LoadDatabase(const TArray<uint8>& Bytes);
 
 	/**
 	 * Add a database defined by a RawData in the current session database
diff --git a/Engine/Source/Runtime/Datasmith/CADKernel/Engine/Private/CADKernelEngine.cpp b/Engine/Source/Runtime/Datasmith/CADKernel/Engine/Private/CADKernelEngine.cpp
--- a/Engine/Source/Runtime/Datasmith/CADKernel/Engine/Private/CADKernelEngine.cpp
+++ b/Engine/Source/Runtime/Datasmith/CADKernel/Engine/Private/CADKernelEngine.cpp
@@ -6,6 +6,7 @@
 #include "CADKernelEnginePrivate.h"
 #include "MeshUtilities.h"
 
+#include "Core/Session.h"
 #include "Topo/Body.h"
 #include "Topo/Model.h"
 #include "Topo/Shell.h"
@@ -38,6 +39,22 @@ namespace UE::CADKernel
 
 	bool FCADKernelUtilities::Save(TSharedPtr<UE::CADKernel::FModel>& Model, const FString& FilePath)
 	{
+		TSharedPtr<FSession> Session = FEntity::MakeShared<FSession>(0.01);
+		Session->GetModel().Copy(*Model);
+
+		return Session->SaveDatabase(*FilePath);
+	}
+
+	bool FCADKernelUtilities::Load(TSharedPtr<UE::CADKernel::FModel>& Model, const FString& FilePath)
+	{
+		TSharedPtr<FSession> Session = FEntity::MakeShared<FSession>(0.01);
+
+		if (Session->LoadDatabase(*FilePath))
+		{
+			Model->Copy(Session->GetModel());
+			return true;
+		}
+
 		return false;
 	}
 
diff --git a/Engine/Source/Runtime/Datasmith/CADKernel/Engine/Private/TechSoft/TechSoftLibrary.cpp b/Engine/Source/Runtime/Datasmith/CADKernel/Engine/Private/TechSoft/TechSoftLibrary.cpp
--- a/Engine/Source/Runtime/Datasmith/CADKernel/Engine/Private/TechSoft/TechSoftLibrary.cpp
+++ b/Engine/Source/Runtime/Datasmith/CADKernel/Engine/Private/TechSoft/TechSoftLibrary.cpp
@@ -11,13 +11,6 @@
 #include "HAL/PlatformProcess.h"
 #include "Misc/Paths.h"
 
-// #cad_import: This is ugly :-(
-#define RECTIFY_STRUTC_SIZE(Struct, StructName) \
-{ \
-	TechSoft::TUniqueObject<StructName> OtherStruct; \
-	Struct.m_usStructSize = OtherStruct->m_usStructSize; \
-}
-
 namespace UE::CADKernel
 {
 	bool FTechSoftLibrary::bIsInitialized = false;
@@ -164,8 +157,6 @@ namespace UE::CADKernel
 			return nullptr;
 		}
 
-		RECTIFY_STRUTC_SIZE(TopoShellData, A3DTopoShellData)
-
 		A3DTopoShell* TopoShellPtr = nullptr;
 		if (A3DTopoShellCreate(&TopoShellData, &TopoShellPtr) != A3DStatus::A3D_SUCCESS)
 		{
@@ -181,8 +172,6 @@ namespace UE::CADKernel
 			return nullptr;
 		}
 
-		RECTIFY_STRUTC_SIZE(SurfaceNurbsData, A3DSurfNurbsData)
-
 		A3DSurfNurbs* SurfaceNurbsPtr = nullptr;
 		if (A3DSurfNurbsCreate(&SurfaceNurbsData, &SurfaceNurbsPtr) != A3DStatus::A3D_SUCCESS)
 		{
@@ -262,8 +251,6 @@ namespace UE::CADKernel
 			return nullptr;
 		}
 
-		RECTIFY_STRUTC_SIZE(TopoFaceData, A3DTopoFaceData)
-
 		A3DTopoFace* TopoFacePtr = nullptr;
 		if (A3DTopoFaceCreate(&TopoFaceData, &TopoFacePtr) != A3DStatus::A3D_SUCCESS)
 		{
@@ -279,8 +266,6 @@ namespace UE::CADKernel
 			return nullptr;
 		}
 
-		RECTIFY_STRUTC_SIZE(TopoLoopData, A3DTopoLoopData)
-
 		A3DTopoLoop* TopoLoopPtr = nullptr;
 		if (A3DTopoLoopCreate(&TopoLoopData, &TopoLoopPtr) != A3DStatus::A3D_SUCCESS)
 		{
@@ -307,8 +292,6 @@ namespace UE::CADKernel
 			return nullptr;
 		}
 
-		RECTIFY_STRUTC_SIZE(TopoEdgeData, A3DTopoEdgeData)
-
 		A3DTopoEdge* TopoEdgePtr = nullptr;
 		if (A3DTopoEdgeCreate(&TopoEdgeData, &TopoEdgePtr) != A3DStatus::A3D_SUCCESS)
 		{
@@ -324,8 +307,6 @@ namespace UE::CADKernel
 			return nullptr;
 		}
 
-		RECTIFY_STRUTC_SIZE(TopoCoEdgeData, A3DTopoCoEdgeData)
-
 		A3DTopoCoEdge* TopoCoEdgePtr = nullptr;
 		if (A3DTopoCoEdgeCreate(&TopoCoEdgeData, &TopoCoEdgePtr) != A3DStatus::A3D_SUCCESS)
 		{
@@ -341,8 +322,6 @@ namespace UE::CADKernel
 			return nullptr;
 		}
 
-		RECTIFY_STRUTC_SIZE(CurveNurbsData, A3DCrvNurbsData)
-
 		A3DCrvNurbs* CurveNurbsPtr = nullptr;
 		if (A3DCrvNurbsCreate(&CurveNurbsData, &CurveNurbsPtr) != A3DStatus::A3D_SUCCESS)
 		{
diff --git a/Engine/Source/Runtime/Datasmith/CADKernel/Engine/Public/CADKernelEngine.h b/Engine/Source/Runtime/Datasmith/CADKernel/Engine/Public/CADKernelEngine.h
--- a/Engine/Source/Runtime/Datasmith/CADKernel/Engine/Public/CADKernelEngine.h
+++ b/Engine/Source/Runtime/Datasmith/CADKernel/Engine/Public/CADKernelEngine.h
@@ -47,6 +47,7 @@ namespace UE::CADKernel
 	{
 	public:
 		static bool CADKERNELENGINE_API Save(TSharedPtr<UE::CADKernel::FModel>& Model, const FString& FilePath);
+		static bool CADKERNELENGINE_API Load(TSharedPtr<UE::CADKernel::FModel>& Model, const FString& FilePath);
 		static bool CADKERNELENGINE_API Tessellate(UE::CADKernel::FModel& Model, const FTessellationContext& Context, FMeshDescription& MeshOut, bool bEmptyMesh = false);
 		static bool CADKERNELENGINE_API Tessellate(UE::CADKernel::FModel& Model, const FTessellationContext& Context, UE::Geometry::FDynamicMesh3& MeshOut, bool bEmptyMesh = false);
 		
diff --git a/Engine/Source/Runtime/Datasmith/CADKernel/Engine/Public/TechSoftUniqueObject.h b/Engine/Source/Runtime/Datasmith/CADKernel/Engine/Public/TechSoftUniqueObject.h
--- a/Engine/Source/Runtime/Datasmith/CADKernel/Engine/Public/TechSoftUniqueObject.h
+++ b/Engine/Source/Runtime/Datasmith/CADKernel/Engine/Public/TechSoftUniqueObject.h
@@ -183,9 +183,7 @@ namespace UE::CADKernel::TechSoft
 			return (IndexerType)INDEX_NONE;
 		}
 
-		CADKERNELENGINE_API void InitializeData()
-		{
-		}
+		CADKERNELENGINE_API void InitializeData();
 
 		CADKERNELENGINE_API A3DStatus GetData(IndexerType AsmModelFilePtr)
 		{
