Fixed crash when extracting triangles from PolyBrepModel entities

#rnx

diff --git a/Engine/Source/Runtime/Datasmith/CADKernel/Engine/Private/MeshDescriptionWrapper.cpp b/Engine/Source/Runtime/Datasmith/CADKernel/Engine/Private/MeshDescriptionWrapper.cpp
--- a/Engine/Source/Runtime/Datasmith/CADKernel/Engine/Private/MeshDescriptionWrapper.cpp
+++ b/Engine/Source/Runtime/Datasmith/CADKernel/Engine/Private/MeshDescriptionWrapper.cpp
@@ -10,6 +10,7 @@
 
 #include "Containers/Queue.h"
 #include "MeshDescriptionBuilder.h"
+#include "StaticMeshOperations.h"
 
 namespace UE::CADKernel::MeshUtilities
 {
@@ -77,6 +78,8 @@ namespace UE::CADKernel::MeshUtilities
 			VertexInstanceBinormalSigns[VertexInstanceID] = 0.0f;
 		}
 
+		FStaticMeshOperations::DetermineEdgeHardnessesFromVertexInstanceNormals(Mesh);
+
 		bIsFinalized = true;
 	}
 
@@ -316,7 +319,7 @@ namespace UE::CADKernel::MeshUtilities
 	bool FMeshDescriptionWrapper::StartFaceTriangles(int32 TriangleCount, const TArray<FVector3f>& InNormals, const TArray<FVector2f>& InTexCoords)
 	{
 		const int32 ArraySize = InNormals.Num();
-		ensure(ArraySize == InTexCoords.Num());
+		ensure(InTexCoords.IsEmpty() || (ArraySize == InTexCoords.Num()));
 
 		Normals = InNormals;
 		for (FVector3f& Normal : Normals)
@@ -334,7 +337,7 @@ namespace UE::CADKernel::MeshUtilities
 	bool FMeshDescriptionWrapper::StartFaceTriangles(const TArrayView<FVector>& InNormals, const TArrayView<FVector2d>& InTexCoords)
 	{
 		const int32 ArraySize = InNormals.Num();
-		ensure(ArraySize == InTexCoords.Num());
+		ensure(InTexCoords.IsEmpty() || (ArraySize == InTexCoords.Num()));
 
 		Normals.Reserve(ArraySize);
 		TexCoords.Reserve(ArraySize);
@@ -342,7 +345,10 @@ namespace UE::CADKernel::MeshUtilities
 		for (int32 Index = 0; Index < ArraySize; ++Index)
 		{
 			Normals.Emplace(FVector3f(InNormals[Index].X, InNormals[Index].Y, InNormals[Index].Z).GetSafeNormal());
-			TexCoords.Emplace(InTexCoords[Index].X, InTexCoords[Index].Y);
+			if (!InTexCoords.IsEmpty())
+			{
+				TexCoords.Emplace(InTexCoords[Index].X, InTexCoords[Index].Y);
+			}
 		}
 
 		ConvertVectorArray(Context.ModelParams.ModelCoordSys, Normals);
@@ -374,9 +380,12 @@ namespace UE::CADKernel::MeshUtilities
 			}
 
 			// Fill up arrays associated with vertex instances
-			VertexInstanceUVs.Set(VertexInstanceIDs[0], 0/*UVChannel*/, FVector2f(TexCoords[FaceTriangle.TexCoords[Orientation[0]]] * ScaleUV));
-			VertexInstanceUVs.Set(VertexInstanceIDs[1], 0/*UVChannel*/, FVector2f(TexCoords[FaceTriangle.TexCoords[Orientation[1]]] * ScaleUV));
-			VertexInstanceUVs.Set(VertexInstanceIDs[2], 0/*UVChannel*/, FVector2f(TexCoords[FaceTriangle.TexCoords[Orientation[2]]] * ScaleUV));
+			if (!TexCoords.IsEmpty())
+			{
+				VertexInstanceUVs.Set(VertexInstanceIDs[0], 0/*UVChannel*/, FVector2f(TexCoords[FaceTriangle.TexCoords[Orientation[0]]] * ScaleUV));
+				VertexInstanceUVs.Set(VertexInstanceIDs[1], 0/*UVChannel*/, FVector2f(TexCoords[FaceTriangle.TexCoords[Orientation[1]]] * ScaleUV));
+				VertexInstanceUVs.Set(VertexInstanceIDs[2], 0/*UVChannel*/, FVector2f(TexCoords[FaceTriangle.TexCoords[Orientation[2]]] * ScaleUV));
+			}
 
 			VertexInstanceNormals[VertexInstanceIDs[0]] = Normals[FaceTriangle.Normals[Orientation[0]]];
 			VertexInstanceNormals[VertexInstanceIDs[1]] = Normals[FaceTriangle.Normals[Orientation[1]]];
diff --git a/Engine/Source/Runtime/Datasmith/CADKernel/Engine/Private/TechSoft/TechSoftMesh.cpp b/Engine/Source/Runtime/Datasmith/CADKernel/Engine/Private/TechSoft/TechSoftMesh.cpp
--- a/Engine/Source/Runtime/Datasmith/CADKernel/Engine/Private/TechSoft/TechSoftMesh.cpp
+++ b/Engine/Source/Runtime/Datasmith/CADKernel/Engine/Private/TechSoft/TechSoftMesh.cpp
@@ -59,6 +59,21 @@ namespace UE::CADKernel::TechSoftUtilities
 		void AddFaceTriangleStripeWithUniqueNormal(const uint32 TriCount, uint32& StartIndexInOut);
 		void AddFaceTriangleStripeWithUniqueNormalAndTexture(const uint32 TriCount, const uint32 TexOffset, uint32& StartIndexInOut);
 
+		void AddTriangle(const FArray3i& VertexIndices, const FArray3i& NormalIndices)
+		{
+			const int32 NextIndex = FaceNormals.Num();
+			const FArray3i OtherIndices{ NextIndex + 0, NextIndex + 1, NextIndex + 2 };
+
+			const FVector3f Normals[3]{
+				{FVector3f(TessellationNormals[NormalIndices[0] / 3])},
+				{FVector3f(TessellationNormals[NormalIndices[1] / 3])},
+				{FVector3f(TessellationNormals[NormalIndices[2] / 3])},
+			};
+			FaceNormals.Append(Normals, 3);
+
+			FaceTriangles.Emplace(GroupID, MaterialID, VertexIndices, OtherIndices, OtherIndices);
+		}
+
 		void AddTriangle(const FArray3i& VertexIndices, const FArray3i& NormalIndices, const FArray3i& TexCoordIndicies)
 		{
 			const int32 NextIndex = FaceNormals.Num();
@@ -292,7 +307,6 @@ namespace UE::CADKernel::TechSoftUtilities
 	{
 		FArray3i VertexIndices;
 		FArray3i NormalIndices;
-		FArray3i TexCoordIndices;
 
 		NormalIndices[0] = TriangulatedIndexes[StartIndexInOut++];
 		VertexIndices[0] = TriangulatedIndexes[StartIndexInOut++] / 3;
@@ -304,7 +318,7 @@ namespace UE::CADKernel::TechSoftUtilities
 			NormalIndices[2] = TriangulatedIndexes[StartIndexInOut++];
 			VertexIndices[2] = TriangulatedIndexes[StartIndexInOut++] / 3;
 
-			AddTriangle(VertexIndices, NormalIndices, TexCoordIndices);
+			AddTriangle(VertexIndices, NormalIndices);
 
 			Index++;
 			if (Index == TriCount)
@@ -318,7 +332,7 @@ namespace UE::CADKernel::TechSoftUtilities
 			NormalIndices[0] = TriangulatedIndexes[StartIndexInOut++];
 			VertexIndices[0] = TriangulatedIndexes[StartIndexInOut++] / 3;
 
-			AddTriangle(VertexIndices, NormalIndices, TexCoordIndices);
+			AddTriangle(VertexIndices, NormalIndices);
 
 			Swap(VertexIndices[0], VertexIndices[1]);
 			Swap(NormalIndices[0], NormalIndices[1]);
@@ -722,11 +736,6 @@ namespace UE::CADKernel::TechSoftUtilities
 			}
 		}
 
-		if (BRepFaces.IsEmpty())
-		{
-			return false;
-		}
-
 		TechSoft::TUniqueObject<A3DTess3DData> TessellationData(TessellationBase);
 
 		if (!TessellationData.IsValid() || TessellationData->m_uiFaceTessSize == 0)
@@ -734,12 +743,6 @@ namespace UE::CADKernel::TechSoftUtilities
 			return false;
 		}
 
-		if (TessellationData->m_uiFaceTessSize != BRepFaces.Num())
-		{
-			BRepFaces.Empty();
-			return false;
-		}
-
 		for (A3DUns32 Index = 0; Index < TessellationData->m_uiFaceTessSize; ++Index)
 		{
 			TriangleCount += GetTriangleCount(TessellationData->m_psFaceTessData[Index]);
@@ -801,8 +804,6 @@ namespace UE::CADKernel::TechSoftUtilities
 		TArrayView<FVector> TriangleNormals((FVector*)TessellationData->m_pdNormals, TessellationData->m_uiNormalSize / 3);
 		TArrayView<FVector2d> TriangleTexCoords((FVector2d*)TessellationData->m_pdTextureCoords, TessellationData->m_uiTextureCoordSize / 2);
 
-		ensure(TriangleNormals.Num() == TriangleTexCoords.Num());
-
 		FFaceTriangleCollector TriangleCollector(TessellationData->m_puiTriangulatedIndexes, TriangleNormals, TriangleTexCoords);
 		FFaceTriangleCollector::FCollectionResult CollectionResult;
 
@@ -815,7 +816,7 @@ namespace UE::CADKernel::TechSoftUtilities
 				continue;
 			}
 
-			const A3DTopoFace* TopoFace = BRepFaces[Index];
+			const A3DTopoFace* TopoFace = BRepFaces.IsValidIndex(Index) ? BRepFaces[Index] : nullptr;
 
 			TechSoft::TUniqueObject<A3DRootBaseData> RootBaseData(TopoFace);
 			int32 GroupID = RootBaseData.IsValid() ? (int32)RootBaseData->m_uiPersistentId : Index;
@@ -863,11 +864,7 @@ namespace UE::CADKernel::TechSoftUtilities
 
 		int32 TriangleCount;
 		TArray<A3DTopoFace*> BRepFaces;
-		if (!CollectBRepFaces(BRepFaces, TriangleCount))
-		{
-			return false;
-		}
-
+		CollectBRepFaces(BRepFaces, TriangleCount);
 		ParseTriangles(BRepFaces);
 
 		return true;
diff --git a/Engine/Source/Runtime/Datasmith/CADKernel/Engine/Private/TechSoft/TechSoftUtilities.cpp b/Engine/Source/Runtime/Datasmith/CADKernel/Engine/Private/TechSoft/TechSoftUtilities.cpp
--- a/Engine/Source/Runtime/Datasmith/CADKernel/Engine/Private/TechSoft/TechSoftUtilities.cpp
+++ b/Engine/Source/Runtime/Datasmith/CADKernel/Engine/Private/TechSoft/TechSoftUtilities.cpp
@@ -313,7 +313,7 @@ namespace UE::CADKernel
 		TechSoft::TUniqueObject<A3DTopoFaceData> TopoFaceData(TopoFace);
 		if (!TopoFaceData.IsValid())
 		{
-			return FVector2d::ZeroVector;
+			return FVector2d::UnitVector;
 		}
 
 		TechSoft::TUniqueObject<A3DDomainData> Domain;
