[In Progress] Interchange CAD Import
- Fixed issue with CADKernelEngine's mesh converter not taking in account scale factor applied on surface tessellated.

#rnx

diff --git a/Engine/Source/Runtime/Datasmith/CADKernel/Engine/Private/MeshDescriptionWrapper.cpp b/Engine/Source/Runtime/Datasmith/CADKernel/Engine/Private/MeshDescriptionWrapper.cpp
--- a/Engine/Source/Runtime/Datasmith/CADKernel/Engine/Private/MeshDescriptionWrapper.cpp
+++ b/Engine/Source/Runtime/Datasmith/CADKernel/Engine/Private/MeshDescriptionWrapper.cpp
@@ -55,11 +55,12 @@ namespace UE::CADKernel::MeshUtilities
 		TArrayView<FVector3f> Positions = VertexPositions.GetRawArray();
 		ConvertVectorArray(Context.ModelParams.ModelCoordSys, Positions);
 
-		if (!FMath::IsNearlyEqual(Context.ModelParams.ModelUnitToCentimeter, 1.f))
+		const float ScaleFactor = Context.ModelParams.ModelUnitToCentimeter * Context.MeshParams.ScaleFactor;
+		if (!FMath::IsNearlyEqual(ScaleFactor, 1.f))
 		{
 			for (FVector3f& Position : Positions)
 			{
-				Position *= Context.ModelParams.ModelUnitToCentimeter;
+				Position *= ScaleFactor;
 			}
 		}
 
