[In Progress] Interchange CAD Import
- Removed use of WITH_DATASMITH in OpenNurbss and Wire's code

#rnx

diff --git a/Engine/Source/Runtime/Datasmith/CADKernel/Engine/Private/DynamicMeshWrapper.cpp b/Engine/Source/Runtime/Datasmith/CADKernel/Engine/Private/DynamicMeshWrapper.cpp
--- a/Engine/Source/Runtime/Datasmith/CADKernel/Engine/Private/DynamicMeshWrapper.cpp
+++ b/Engine/Source/Runtime/Datasmith/CADKernel/Engine/Private/DynamicMeshWrapper.cpp
@@ -82,7 +82,7 @@ namespace UE::CADKernel::MeshUtilities
 			}
 		}
 
-		ConvertVectorArray(Context.ModelParams.ModelCoordSys, Normals);
+		UE::CADKernel::MathUtils::ConvertVectorArray(Context.ModelParams.ModelCoordSys, Normals);
 
 		AddTriangles();
 	}
diff --git a/Engine/Source/Runtime/Datasmith/CADKernel/Engine/Private/MeshDescriptionWrapper.cpp b/Engine/Source/Runtime/Datasmith/CADKernel/Engine/Private/MeshDescriptionWrapper.cpp
--- a/Engine/Source/Runtime/Datasmith/CADKernel/Engine/Private/MeshDescriptionWrapper.cpp
+++ b/Engine/Source/Runtime/Datasmith/CADKernel/Engine/Private/MeshDescriptionWrapper.cpp
@@ -53,7 +53,7 @@ namespace UE::CADKernel::MeshUtilities
 		}
 
 		TArrayView<FVector3f> Positions = VertexPositions.GetRawArray();
-		ConvertVectorArray(Context.ModelParams.ModelCoordSys, Positions);
+		UE::CADKernel::MathUtils::ConvertVectorArray(Context.ModelParams.ModelCoordSys, Positions);
 
 		const float ScaleFactor = Context.ModelParams.ModelUnitToCentimeter * Context.MeshParams.ScaleFactor;
 		if (!FMath::IsNearlyEqual(ScaleFactor, 1.f))
@@ -328,7 +328,7 @@ namespace UE::CADKernel::MeshUtilities
 			Normal = Normal.GetSafeNormal();
 		}
 
-		ConvertVectorArray(Context.ModelParams.ModelCoordSys, Normals);
+		UE::CADKernel::MathUtils::ConvertVectorArray(Context.ModelParams.ModelCoordSys, Normals);
 
 		TexCoords = InTexCoords;
 
@@ -352,7 +352,7 @@ namespace UE::CADKernel::MeshUtilities
 			}
 		}
 
-		ConvertVectorArray(Context.ModelParams.ModelCoordSys, Normals);
+		UE::CADKernel::MathUtils::ConvertVectorArray(Context.ModelParams.ModelCoordSys, Normals);
 
 		return true;
 	}
diff --git a/Engine/Source/Runtime/Datasmith/CADKernel/Engine/Private/MeshUtilities.h b/Engine/Source/Runtime/Datasmith/CADKernel/Engine/Private/MeshUtilities.h
--- a/Engine/Source/Runtime/Datasmith/CADKernel/Engine/Private/MeshUtilities.h
+++ b/Engine/Source/Runtime/Datasmith/CADKernel/Engine/Private/MeshUtilities.h
@@ -180,52 +180,6 @@ namespace UE::CADKernel::MeshUtilities
 	void GetExistingFaceGroups(FMeshDescription& Mesh, TSet<int32>& FaceGroupsOut);
 	void GetExistingFaceGroups(UE::Geometry::FDynamicMesh3& Mesh, TSet<int32>& FaceGroupsOut);
 
-	template<typename VecType>
-	void ConvertVectorArray(uint8 ModelCoordSys, TArrayView<VecType>& Array)
-	{
-		switch (ECADKernelModelCoordSystem(ModelCoordSys))
-		{
-		case ECADKernelModelCoordSystem::YUp_LeftHanded:
-			for (VecType& Vector : Array)
-			{
-				Vector.Set(Vector[2], Vector[0], Vector[1]);
-			}
-			break;
-
-		case ECADKernelModelCoordSystem::YUp_RightHanded:
-			for (VecType& Vector : Array)
-			{
-				Vector.Set(-Vector[2], Vector[0], Vector[1]);
-			}
-			break;
-
-		case ECADKernelModelCoordSystem::ZUp_RightHanded:
-			for (VecType& Vector : Array)
-			{
-				Vector.Set(-Vector[0], Vector[1], Vector[2]);
-			}
-			break;
-
-		case ECADKernelModelCoordSystem::ZUp_RightHanded_FBXLegacy:
-			for (VecType& Vector : Array)
-			{
-				Vector.Set(Vector[0], -Vector[1], Vector[2]);
-			}
-			break;
-
-		case ECADKernelModelCoordSystem::ZUp_LeftHanded:
-		default:
-			break;
-		}
-	}
-
-	template<typename VecType>
-	void ConvertVectorArray(uint8 ModelCoordSys, TArray<VecType>& Array)
-	{
-		TArrayView<VecType> ArrayView(Array);
-		ConvertVectorArray(ModelCoordSys, ArrayView);
-	}
-
 	template<typename VecType>
 	void ConvertVectorArray(uint8 ModelCoordSys, UE::Geometry::TDynamicVector<VecType>& Array)
 	{
@@ -265,29 +219,6 @@ namespace UE::CADKernel::MeshUtilities
 		}
 	}
 
-	template<typename VecType>
-	VecType ConvertVector(uint8 ModelCoordSys, const VecType& V)
-	{
-		switch (ECADKernelModelCoordSystem(ModelCoordSys))
-		{
-		case ECADKernelModelCoordSystem::YUp_LeftHanded:
-			return VecType(V[2], V[0], V[1]);
-
-		case ECADKernelModelCoordSystem::YUp_RightHanded:
-			return VecType(-V[2], V[0], V[1]);
-
-		case ECADKernelModelCoordSystem::ZUp_RightHanded:
-			return VecType(-V[0], V[1], V[2]);
-
-		case ECADKernelModelCoordSystem::ZUp_RightHanded_FBXLegacy:
-			return VecType(V[0], -V[1], V[2]);
-
-		case ECADKernelModelCoordSystem::ZUp_LeftHanded:
-		default:
-			return VecType(V[0], V[1], V[2]);
-		}
-	}
-
 	template<typename Type>
 	UE::Math::TMatrix<Type> GetSymmetricMatrix(const UE::Math::TVector<Type>& Origin, const UE::Math::TVector<Type>& Normal)
 	{
diff --git a/Engine/Source/Runtime/Datasmith/CADKernel/Engine/Public/CADKernelEngine.h b/Engine/Source/Runtime/Datasmith/CADKernel/Engine/Public/CADKernelEngine.h
--- a/Engine/Source/Runtime/Datasmith/CADKernel/Engine/Public/CADKernelEngine.h
+++ b/Engine/Source/Runtime/Datasmith/CADKernel/Engine/Public/CADKernelEngine.h
@@ -147,6 +147,94 @@ namespace UE::CADKernel
 				}
 			}
 		}
+
+		template<typename VecType>
+		void ConvertVectorArray(ECADKernelModelCoordSystem ModelCoordSys, TArrayView<VecType>& Array)
+		{
+			switch (ModelCoordSys)
+			{
+			case ECADKernelModelCoordSystem::YUp_LeftHanded:
+				for (VecType& Vector : Array)
+				{
+					Vector.Set(Vector[2], Vector[0], Vector[1]);
+				}
+				break;
+
+			case ECADKernelModelCoordSystem::YUp_RightHanded:
+				for (VecType& Vector : Array)
+				{
+					Vector.Set(-Vector[2], Vector[0], Vector[1]);
+				}
+				break;
+
+			case ECADKernelModelCoordSystem::ZUp_RightHanded:
+				for (VecType& Vector : Array)
+				{
+					Vector.Set(-Vector[0], Vector[1], Vector[2]);
+				}
+				break;
+
+			case ECADKernelModelCoordSystem::ZUp_RightHanded_FBXLegacy:
+				for (VecType& Vector : Array)
+				{
+					Vector.Set(Vector[0], -Vector[1], Vector[2]);
+				}
+				break;
+
+			case ECADKernelModelCoordSystem::ZUp_LeftHanded:
+			default:
+				break;
+			}
+		}
+
+		template<typename VecType>
+		void ConvertVectorArray(uint8 ModelCoordSys, TArrayView<VecType>& Array)
+		{
+			ConvertVectorArray(static_cast<ECADKernelModelCoordSystem>(ModelCoordSys), Array);
+		}
+
+		template<typename VecType>
+		void ConvertVectorArray(ECADKernelModelCoordSystem ModelCoordSys, TArray<VecType>& Array)
+		{
+			TArrayView<VecType> ArrayView(Array);
+			ConvertVectorArray(ModelCoordSys, ArrayView);
+		}
+
+		template<typename VecType>
+		void ConvertVectorArray(uint8 ModelCoordSys, TArray<VecType>& Array)
+		{
+			TArrayView<VecType> ArrayView(Array);
+			ConvertVectorArray(static_cast<ECADKernelModelCoordSystem>(ModelCoordSys), ArrayView);
+		}
+
+		template<typename VecType>
+		VecType ConvertVector(ECADKernelModelCoordSystem ModelCoordSys, const VecType& V)
+		{
+			switch (ModelCoordSys)
+			{
+			case ECADKernelModelCoordSystem::YUp_LeftHanded:
+				return VecType(V[2], V[0], V[1]);
+
+			case ECADKernelModelCoordSystem::YUp_RightHanded:
+				return VecType(-V[2], V[0], V[1]);
+
+			case ECADKernelModelCoordSystem::ZUp_RightHanded:
+				return VecType(-V[0], V[1], V[2]);
+
+			case ECADKernelModelCoordSystem::ZUp_RightHanded_FBXLegacy:
+				return VecType(V[0], -V[1], V[2]);
+
+			case ECADKernelModelCoordSystem::ZUp_LeftHanded:
+			default:
+				return VecType(V[0], V[1], V[2]);
+			}
+		}
+
+		template<typename VecType>
+		VecType ConvertVector(uint8 ModelCoordSys, const VecType& V)
+		{
+			return ConvertVector(static_cast<ECADKernelModelCoordSystem>(ModelCoordSys), V);
+		}
 	}
 }
 #endif
