[CADKernel] Mac build fix

https://horde.devtools.epicgames.com/job/67e66db7b61ffeabf8102c9e?step=b2d0&issue=846257

diff --git a/Engine/Source/Runtime/Datasmith/CADKernel/Engine/Private/TechSoft/TechSoftUtilities.cpp b/Engine/Source/Runtime/Datasmith/CADKernel/Engine/Private/TechSoft/TechSoftUtilities.cpp
--- a/Engine/Source/Runtime/Datasmith/CADKernel/Engine/Private/TechSoft/TechSoftUtilities.cpp
+++ b/Engine/Source/Runtime/Datasmith/CADKernel/Engine/Private/TechSoft/TechSoftUtilities.cpp
@@ -464,7 +464,7 @@ namespace UE::CADKernel
 		return nullptr;
 	}
 
-	TSharedPtr<UE::CADKernel::FModel> FTechSoftUtilities::TechSoftToCADKernel(A3DRiRepresentationItem* Representation)
+	TSharedPtr<UE::CADKernel::FModel> FTechSoftUtilities::TechSoftToCADKernel(A3DRiRepresentationItem* Representation, double Unit, double InGeometricTolerance)
 	{
 		return {};
 	}
