[CADKernel] More functions for TechSoftParser

- this seems enought for now to remove direct dependency from Parser to TechSoft library

diff --git a/Engine/Source/Runtime/Datasmith/CADKernel/Engine/Private/TechSoft/TechSoftLibrary.cpp b/Engine/Source/Runtime/Datasmith/CADKernel/Engine/Private/TechSoft/TechSoftLibrary.cpp
--- a/Engine/Source/Runtime/Datasmith/CADKernel/Engine/Private/TechSoft/TechSoftLibrary.cpp
+++ b/Engine/Source/Runtime/Datasmith/CADKernel/Engine/Private/TechSoft/TechSoftLibrary.cpp
@@ -667,11 +667,11 @@ namespace UE::CADKernel
 
 						if (GraphicsData->m_usBehaviour & kA3DGraphicsFatherHeritColor)
 						{
-							// todo:
+							Result.MaterialInheritance = FGraphicsProperties::EInheritance::Father;
 						}
 						else if (GraphicsData->m_usBehaviour & kA3DGraphicsSonHeritColor)
 						{
-							// todo:
+							Result.MaterialInheritance = FGraphicsProperties::EInheritance::Child;
 						}
 
 						// may be A3D_DEFAULT_STYLE_INDEX)
@@ -785,6 +785,64 @@ namespace UE::CADKernel
 		}
 	}
 
+	bool FTechSoftLibrary::IsMaterialTexture(const uint32 MaterialIndex)
+	{
+		A3DBool bIsTexture = false;
+		return A3DGlobalIsMaterialTexture(MaterialIndex, &bIsTexture) == A3DStatus::A3D_SUCCESS ? bool(bIsTexture) : false;
+	}
+
+	bool FTechSoftLibrary::GetEntityType(A3DRiRepresentationItem* RepresentationItem, A3DEEntityType& Type)
+	{
+		return A3DEntityGetType(RepresentationItem, &Type) == A3D_SUCCESS;
+	}
+
+	A3DAsmProductOccurrence* FTechSoftLibrary::FindConfiguration(const A3DAsmProductOccurrence* ConfigurationSetOccurrencePtr, TFunctionRef<bool(A3DAsmProductOccurrenceData& ConfigurationData)> Callback)
+	{
+		TechSoft::TUniqueObject<A3DAsmProductOccurrenceData> ConfigurationSetData(ConfigurationSetOccurrencePtr);
+		if (ConfigurationSetData.IsValid())
+		{
+			for (uint32 Index = 0; Index < ConfigurationSetData->m_uiPOccurrencesSize; ++Index)
+			{
+				A3DAsmProductOccurrence* ConfigurationOccurrence = ConfigurationSetData->m_ppPOccurrences[Index];
+				TechSoft::TUniqueObject<A3DAsmProductOccurrenceData> ConfigurationData(ConfigurationOccurrence);
+				if (!ConfigurationData.IsValid())
+				{
+					continue;
+				}
+
+				if (ConfigurationData->m_uiProductFlags & A3D_PRODUCT_FLAG_CONFIG )
+				{
+					if (Callback(*ConfigurationData))
+					{
+						return ConfigurationOccurrence;
+					}
+				}
+			}
+		}
+		return nullptr;
+	};
+
+	bool FTechSoftLibrary::IsConfigurationSet(UE::CADKernel::ECADFormat Format, const A3DAsmProductOccurrence* Occurrence)
+	{
+		switch (Format)
+		{
+		case ECADFormat::CATIAV4:
+		case ECADFormat::N_X:
+		case ECADFormat::SOLIDWORKS:
+		{
+			TechSoft::TUniqueObject<A3DAsmProductOccurrenceData> OccurrenceData(Occurrence);
+			if (!OccurrenceData.IsValid())
+			{
+				return false;
+			}
+
+			return OccurrenceData->m_uiProductFlags & A3D_PRODUCT_FLAG_CONTAINER;
+		}
+
+		default : 
+			return false;
+		}
+	}
 
 
 #else
diff --git a/Engine/Source/Runtime/Datasmith/CADKernel/Engine/Public/CADKernelEngine.h b/Engine/Source/Runtime/Datasmith/CADKernel/Engine/Public/CADKernelEngine.h
--- a/Engine/Source/Runtime/Datasmith/CADKernel/Engine/Public/CADKernelEngine.h
+++ b/Engine/Source/Runtime/Datasmith/CADKernel/Engine/Public/CADKernelEngine.h
@@ -65,6 +65,32 @@ namespace UE::CADKernel
 		static void CADKERNELENGINE_API RegisterAttributes(FMeshDescription& MeshInOut, bool bKeepExistingAttribute = false);
 	};
 
+	enum class ECADFormat
+	{
+		ACIS,
+		AUTOCAD,
+		CATIA,
+		CATIA_CGR,
+		CATIA_3DXML,
+		CATIAV4,
+		CREO,
+		DWG,
+		DGN,
+		TECHSOFT,
+		IFC,
+		IGES,
+		INVENTOR,
+		JT,
+		N_X,  
+		MICROSTATION,
+		PARASOLID,
+		SOLID_EDGE,
+		SOLIDWORKS,
+		STEP,
+		OTHER
+	};
+
+
 	class FTechSoftLibrary
 	{
 	public:
@@ -112,12 +138,29 @@ namespace UE::CADKernel
 
 			int32 ColorUid = 0;
 			int32 StyleIndex = 0;
+
+			enum class EInheritance
+			{
+				Unset,
+				Father,
+				Child
+			};
+
+			EInheritance MaterialInheritance = EInheritance::Unset;
+
 		};
 		static CADKERNELENGINE_API void ExtractGraphicsProperties(const A3DEntity* Entity, FGraphicsProperties& Result);
 
 		static CADKERNELENGINE_API bool ParseRootBaseData(const A3DEntity* Entity, TMap<FString, FString>& MetaData, FString& UniqueID, FString& Label);
 		static CADKERNELENGINE_API void TraverseAttribute(const A3DMiscAttributeData& AttributeData, TMap<FString, FString>& OutMetaData);
 
+
+		static CADKERNELENGINE_API bool IsMaterialTexture(const uint32 MaterialIndex);
+		static CADKERNELENGINE_API bool GetEntityType(A3DRiRepresentationItem* RepresentationItem, A3DEEntityType& Type);
+
+		static CADKERNELENGINE_API bool IsConfigurationSet(ECADFormat Format, const A3DAsmProductOccurrence* Occurrence);
+		static CADKERNELENGINE_API A3DAsmProductOccurrence* FindConfiguration(const A3DAsmProductOccurrence* ConfigurationSetOccurrencePtr, TFunctionRef<bool(A3DAsmProductOccurrenceData& ConfigurationData)> Callback);
+
 #endif
 
 	private:
