[CADKernel] FFaceTriangleCollector fixed crashing converting geometry without texcoords

diff --git a/Engine/Source/Runtime/Datasmith/CADKernel/Engine/Private/TechSoft/TechSoftMesh.cpp b/Engine/Source/Runtime/Datasmith/CADKernel/Engine/Private/TechSoft/TechSoftMesh.cpp
--- a/Engine/Source/Runtime/Datasmith/CADKernel/Engine/Private/TechSoft/TechSoftMesh.cpp
+++ b/Engine/Source/Runtime/Datasmith/CADKernel/Engine/Private/TechSoft/TechSoftMesh.cpp
@@ -101,7 +101,6 @@ namespace UE::CADKernel::TechSoftUtilities
 	{
 		FArray3i VertexIndices;
 		FArray3i NormalIndices;
-		FArray3i TexCoordIndices(0, 0, 0);
 
 		// Get Triangles
 		for (uint32 Index = 0; Index < TriCount; ++Index)
@@ -113,7 +112,7 @@ namespace UE::CADKernel::TechSoftUtilities
 			NormalIndices[2] = TriangulatedIndexes[StartIndexInOut++];
 			VertexIndices[2] = TriangulatedIndexes[StartIndexInOut++] / 3;
 
-			AddTriangle(VertexIndices, NormalIndices, TexCoordIndices);
+			AddTriangle(VertexIndices, NormalIndices);
 		}
 	}
 
@@ -121,7 +120,6 @@ namespace UE::CADKernel::TechSoftUtilities
 	{
 		FArray3i VertexIndices;
 		FArray3i NormalIndices;
-		FArray3i TexCoordIndices(0, 0, 0);
 
 		for (uint32 Index = 0; Index < TriCount; ++Index)
 		{
@@ -133,7 +131,7 @@ namespace UE::CADKernel::TechSoftUtilities
 			VertexIndices[1] = TriangulatedIndexes[StartIndexInOut++] / 3;
 			VertexIndices[2] = TriangulatedIndexes[StartIndexInOut++] / 3;
 
-			AddTriangle(VertexIndices, NormalIndices, TexCoordIndices);
+			AddTriangle(VertexIndices, NormalIndices);
 		}
 	}
 
@@ -197,7 +195,6 @@ namespace UE::CADKernel::TechSoftUtilities
 	{
 		FArray3i VertexIndices;
 		FArray3i NormalIndices;
-		FArray3i TexCoordIndices(0, 0, 0);
 
 		NormalIndices[0] = TriangulatedIndexes[StartIndexInOut++];
 		VertexIndices[0] = TriangulatedIndexes[StartIndexInOut++] / 3;
@@ -209,7 +206,7 @@ namespace UE::CADKernel::TechSoftUtilities
 			NormalIndices[2] = TriangulatedIndexes[StartIndexInOut++];
 			VertexIndices[2] = TriangulatedIndexes[StartIndexInOut++] / 3;
 
-			AddTriangle(VertexIndices, NormalIndices, TexCoordIndices);
+			AddTriangle(VertexIndices, NormalIndices);
 
 			NormalIndices[1] = NormalIndices[2];
 			VertexIndices[1] = VertexIndices[2];
@@ -220,7 +217,6 @@ namespace UE::CADKernel::TechSoftUtilities
 	{
 		FArray3i VertexIndices;
 		FArray3i NormalIndices;
-		FArray3i TexCoordIndices(0, 0, 0);
 
 
 		NormalIndices[0] = TriangulatedIndexes[StartIndexInOut++];
@@ -235,7 +231,7 @@ namespace UE::CADKernel::TechSoftUtilities
 		{
 			VertexIndices[2] = TriangulatedIndexes[StartIndexInOut++] / 3;
 
-			AddTriangle(VertexIndices, NormalIndices, TexCoordIndices);
+			AddTriangle(VertexIndices, NormalIndices);
 
 			VertexIndices[1] = VertexIndices[2];
 		}
@@ -388,7 +384,6 @@ namespace UE::CADKernel::TechSoftUtilities
 	{
 		FArray3i VertexIndices;
 		FArray3i NormalIndices;
-		FArray3i TexCoordIndices;
 
 		NormalIndices[0] = TriangulatedIndexes[StartIndexInOut++];
 		NormalIndices[1] = NormalIndices[0];
@@ -401,7 +396,7 @@ namespace UE::CADKernel::TechSoftUtilities
 		{
 			VertexIndices[2] = TriangulatedIndexes[StartIndexInOut++] / 3;
 
-			AddTriangle(VertexIndices, NormalIndices, TexCoordIndices);
+			AddTriangle(VertexIndices, NormalIndices);
 
 			Index++;
 			if (Index == TriCount)
@@ -413,7 +408,7 @@ namespace UE::CADKernel::TechSoftUtilities
 
 			VertexIndices[0] = TriangulatedIndexes[StartIndexInOut++] / 3;
 
-			AddTriangle(VertexIndices, NormalIndices, TexCoordIndices);
+			AddTriangle(VertexIndices, NormalIndices);
 
 			Swap(VertexIndices[0], VertexIndices[1]);
 		}
