[Core] Added GetMinMax to FMath to swap the two provided min max values if min is greater than max

diff --git a/Engine/Source/Runtime/Core/Public/GenericPlatform/GenericPlatformMath.h b/Engine/Source/Runtime/Core/Public/GenericPlatform/GenericPlatformMath.h
--- a/Engine/Source/Runtime/Core/Public/GenericPlatform/GenericPlatformMath.h
+++ b/Engine/Source/Runtime/Core/Public/GenericPlatform/GenericPlatformMath.h
@@ -1132,6 +1132,36 @@ struct FGenericPlatformMath
 		return CurMax;
 	}
 
+	/**
+	 * Sort input A & B values into OutMin & OutMax
+	 */
+	template< class T >
+	static constexpr FORCEINLINE void GetMinMax(const T& A, const T& B, T& OutMin, T& OutMax)
+	{
+		if (A < B)
+		{
+			OutMin = A;
+			OutMax = B;
+		}
+		else
+		{
+			OutMin = B;
+			OutMax = A;
+		}
+	}
+
+	/**
+	 * Sort input values to be Min and Max
+	 */
+	template< class T >
+	static constexpr FORCEINLINE void GetMinMax(T& Min, T& Max)
+	{
+		if (Max < Min)
+		{
+			Swap(Min, Max);
+		}
+	}
+
 	[[nodiscard]] static constexpr FORCEINLINE int32 CountBits(uint64 Bits)
 	{
 		// https://en.wikipedia.org/wiki/Hamming_weight
diff --git a/Engine/Source/Runtime/Datasmith/CADKernel/Base/Private/Mesh/Meshers/CycleTriangulator.cpp b/Engine/Source/Runtime/Datasmith/CADKernel/Base/Private/Mesh/Meshers/CycleTriangulator.cpp
--- a/Engine/Source/Runtime/Datasmith/CADKernel/Base/Private/Mesh/Meshers/CycleTriangulator.cpp
+++ b/Engine/Source/Runtime/Datasmith/CADKernel/Base/Private/Mesh/Meshers/CycleTriangulator.cpp
@@ -503,7 +503,7 @@ bool FCycleTriangulator::FindTheBestCandidateNode()
 		// check if the angle of the candidate triangle are not degenerated
 		double MinSlope;
 		double MaxSlope;
-		GetMinMax(SlopeAtStartNode, SlopeAtEndNode, MinSlope, MaxSlope);
+		FMath::GetMinMax(SlopeAtStartNode, SlopeAtEndNode, MinSlope, MaxSlope);
 		constexpr double PiMinus5 = Slope::PiSlope - Slope::FiveDegree;
 		return MinSlope < Slope::FiveDegree&& MaxSlope > PiMinus5;
 	};
diff --git a/Engine/Source/Runtime/Datasmith/CADKernel/Base/Private/Mesh/Meshers/ParametricMesher.cpp b/Engine/Source/Runtime/Datasmith/CADKernel/Base/Private/Mesh/Meshers/ParametricMesher.cpp
--- a/Engine/Source/Runtime/Datasmith/CADKernel/Base/Private/Mesh/Meshers/ParametricMesher.cpp
+++ b/Engine/Source/Runtime/Datasmith/CADKernel/Base/Private/Mesh/Meshers/ParametricMesher.cpp
@@ -351,7 +351,7 @@ void FParametricMesher::IsolateQuadFace(TArray<ParametricMesherTool::FCostToFace
 		case EQuadType::Quadrangular:
 			double LocalMinCurvature;
 			double LocalMaxCurvature;
-			GetMinMax(Face.GetCurvature(EIso::IsoU).Max, Face.GetCurvature(EIso::IsoV).Max, LocalMinCurvature, LocalMaxCurvature);
+			FMath::GetMinMax(Face.GetCurvature(EIso::IsoU).Max, Face.GetCurvature(EIso::IsoV).Max, LocalMinCurvature, LocalMaxCurvature);
 			if (LocalMaxCurvature > MeshConst::MinCurvature)
 			{
 				QuadSurfaces.Emplace(LocalMaxCurvature, &Face);
@@ -488,7 +488,7 @@ void FParametricMesher::LinkQuadSurfaceForMesh(TArray<ParametricMesherTool::FCos
 				SideIndex = Face->GetSideIndex(LocalEdgeIndex);
 				double OtherSideLength = Face->GetSideProperties()[SideIndex].Length3D;
 
-				GetMinMax(OtherSideLength, SideLength);
+				FMath::GetMinMax(OtherSideLength, SideLength);
 				if (SideLength - OtherSideLength > GeometricTolerance)
 				{
 					Face = nullptr;
diff --git a/Engine/Source/Runtime/Datasmith/CADKernel/Base/Private/Mesh/Structure/Grid.cpp b/Engine/Source/Runtime/Datasmith/CADKernel/Base/Private/Mesh/Structure/Grid.cpp
--- a/Engine/Source/Runtime/Datasmith/CADKernel/Base/Private/Mesh/Structure/Grid.cpp
+++ b/Engine/Source/Runtime/Datasmith/CADKernel/Base/Private/Mesh/Structure/Grid.cpp
@@ -1326,8 +1326,8 @@ void FGrid::FindInnerFacePoints()
 			double VMin = FirstSegmentPoint->Y;
 			double Umax = SecondSegmentPoint->X;
 			double Vmax = SecondSegmentPoint->Y;
-			GetMinMax(UMin, Umax);
-			GetMinMax(VMin, Vmax);
+			FMath::GetMinMax(UMin, Umax);
+			FMath::GetMinMax(VMin, Vmax);
 
 			// AB^AP = ABu*APv - ABv*APu
 			// AB^AP = ABu*(Pv-Av) - ABv*(Pu-Au)
@@ -1416,8 +1416,8 @@ void FGrid::FindInnerFacePoints()
 			double VMin = FirstSegmentPoint->Y;
 			double Umax = SecondSegmentPoint->X;
 			double Vmax = SecondSegmentPoint->Y;
-			GetMinMax(UMin, Umax);
-			GetMinMax(VMin, Vmax);
+			FMath::GetMinMax(UMin, Umax);
+			FMath::GetMinMax(VMin, Vmax);
 
 			// AB^AP = ABu*APv - ABv*APu
 			// AB^AP = ABu*(Pv-Av) - ABv*(Pu-Au)
diff --git a/Engine/Source/Runtime/Datasmith/CADKernel/Base/Private/Mesh/Structure/ModelMeshAnalyzer.cpp b/Engine/Source/Runtime/Datasmith/CADKernel/Base/Private/Mesh/Structure/ModelMeshAnalyzer.cpp
--- a/Engine/Source/Runtime/Datasmith/CADKernel/Base/Private/Mesh/Structure/ModelMeshAnalyzer.cpp
+++ b/Engine/Source/Runtime/Datasmith/CADKernel/Base/Private/Mesh/Structure/ModelMeshAnalyzer.cpp
@@ -37,7 +37,7 @@ void FModelMeshAnalyzer::BuildMesh()
 	{
 		uint32 EdgeVertexIndex0;
 		uint32 EdgeVertexIndex1;
-		GetMinMax(VertexIndex0, VertexIndex1, EdgeVertexIndex0, EdgeVertexIndex1);
+		FMath::GetMinMax(VertexIndex0, VertexIndex1, EdgeVertexIndex0, EdgeVertexIndex1);
 		uint32 EdgeHash = HashCombine(EdgeVertexIndex0, EdgeVertexIndex1);
 		Analyzer::FEdge** EdgePtr = EdgeMap.Find(EdgeHash);
 		if (EdgePtr != nullptr)
diff --git a/Engine/Source/Runtime/Datasmith/CADKernel/Base/Private/Mesh/Structure/ThinZone2DFinder.cpp b/Engine/Source/Runtime/Datasmith/CADKernel/Base/Private/Mesh/Structure/ThinZone2DFinder.cpp
--- a/Engine/Source/Runtime/Datasmith/CADKernel/Base/Private/Mesh/Structure/ThinZone2DFinder.cpp
+++ b/Engine/Source/Runtime/Datasmith/CADKernel/Base/Private/Mesh/Structure/ThinZone2DFinder.cpp
@@ -1616,7 +1616,7 @@ void FThinZoneSide::AddToEdge()
 	{
 		double UMin = EdgeSegment.GetCoordinate(ELimit::Start);
 		double UMax = EdgeSegment.GetCoordinate(ELimit::End);
-		GetMinMax(UMin, UMax);
+		FMath::GetMinMax(UMin, UMax);
 
 		if (Edge != EdgeSegment.GetEdge())
 		{
@@ -1843,7 +1843,7 @@ void FThinZoneSide::GetExistingMeshNodes(const FTopologicalFace& Face, FModelMes
 		{
 			SegmentUMin = EdgeSegment.GetCoordinate(ELimit::Start);
 			SegmentUMax = EdgeSegment.GetCoordinate(ELimit::End);
-			GetMinMax(SegmentUMin, SegmentUMax);
+			FMath::GetMinMax(SegmentUMin, SegmentUMax);
 
 			if (Edge != EdgeSegment.GetEdge())
 			{
diff --git a/Engine/Source/Runtime/Datasmith/CADKernel/Base/Public/Math/Boundary.h b/Engine/Source/Runtime/Datasmith/CADKernel/Base/Public/Math/Boundary.h
--- a/Engine/Source/Runtime/Datasmith/CADKernel/Base/Public/Math/Boundary.h
+++ b/Engine/Source/Runtime/Datasmith/CADKernel/Base/Public/Math/Boundary.h
@@ -82,17 +82,17 @@ struct CADKERNEL_API FLinearBoundary
 
 	void SetMin(const double Coordinates)
 	{
-		GetMinMax(Coordinates, Max, Min, Max);
+		FMath::GetMinMax(Coordinates, Max, Min, Max);
 	}
 
 	void SetMax(const double Coordinates)
 	{
-		GetMinMax(Min, Coordinates, Min, Max);
+		FMath::GetMinMax(Min, Coordinates, Min, Max);
 	}
 
 	void Set(const double InUMin = 0., const double InUMax = 1.)
 	{
-		GetMinMax(InUMin, InUMax, Min, Max);
+		FMath::GetMinMax(InUMin, InUMax, Min, Max);
 	}
 
 	/**
@@ -175,7 +175,7 @@ struct CADKERNEL_API FLinearBoundary
 
 	void ExtendTo(double MinCoordinate, double MaxCoordinate)
 	{
-		GetMinMax(MinCoordinate, MaxCoordinate);
+		FMath::GetMinMax(MinCoordinate, MaxCoordinate);
 		Min = FMath::Min(Min, MinCoordinate);
 		Max = FMath::Max(Max, MaxCoordinate);
 	}
diff --git a/Engine/Source/Runtime/Datasmith/CADKernel/Base/Public/Math/MathConst.h b/Engine/Source/Runtime/Datasmith/CADKernel/Base/Public/Math/MathConst.h
--- a/Engine/Source/Runtime/Datasmith/CADKernel/Base/Public/Math/MathConst.h
+++ b/Engine/Source/Runtime/Datasmith/CADKernel/Base/Public/Math/MathConst.h
@@ -5,32 +5,22 @@
 
 namespace UE::CADKernel
 {
+UE_DEPRECATED(5.6, "Use FMath::GetMinMax instead.");
 /**
  * Sort input A & B values into OutMin & OutMax
  */
 template<typename ValueType> const void GetMinMax(const ValueType& ValueA, const ValueType& ValueB, ValueType& OutMin, ValueType& OutMax)
 {
-	if (ValueA < ValueB)
-	{
-		OutMin = ValueA;
-		OutMax = ValueB;
-	}
-	else
-	{
-		OutMin = ValueB;
-		OutMax = ValueA;
-	}
+	FMath::GetMinMax(ValueA, ValueB, OutMin, OutMax);
 }
 
+UE_DEPRECATED(5.6, "Use FMath::GetMinMax instead.");
 /**
  * Sort input values to be Min and Max
  */
 template<typename ValueType> const void GetMinMax(ValueType& Min, ValueType& Max)
 {
-	if (Max < Min)
-	{
-		Swap(Min, Max);
-	}
+	FMath::GetMinMax(Min, Max);
 }
 
 /** Checks if value is within a range, exclusive on MinValue and MaxValue) */
