[Datasmith]
* Added missing dllexport of FDatasmithSceneReceiver::FDatasmithElementPointers::AsSharedPtr

diff --git a/Engine/Source/Runtime/Datasmith/DatasmithCore/Public/DirectLink/DatasmithSceneReceiver.h b/Engine/Source/Runtime/Datasmith/DatasmithCore/Public/DirectLink/DatasmithSceneReceiver.h
--- a/Engine/Source/Runtime/Datasmith/DatasmithCore/Public/DirectLink/DatasmithSceneReceiver.h
+++ b/Engine/Source/Runtime/Datasmith/DatasmithCore/Public/DirectLink/DatasmithSceneReceiver.h
@@ -56,7 +56,7 @@ private:
 	struct FDatasmithElementPointers : public DirectLink::IReferenceResolutionProvider
 	{
 		// IReferenceResolutionProvider API
-		virtual TSharedPtr<DirectLink::ISceneGraphNode> AsSharedPtr(DirectLink::FSceneGraphId NodeId) override;
+		UE_API virtual TSharedPtr<DirectLink::ISceneGraphNode> AsSharedPtr(DirectLink::FSceneGraphId NodeId) override;
 		void Reset() { ElementsSharedPtrs.Reset(); }
 		void Remove(DirectLink::FSceneGraphId NodeId) { ElementsSharedPtrs.Remove(NodeId); }
 		void Add(DirectLink::FSceneGraphId Id, TSharedPtr<IDatasmithElement> Element) { ElementsSharedPtrs.Add(Id, Element); }
