Fixed compilation error when compiling DatasmithCADImporter plugin without TechSoft SDK

#rnx

diff --git a/Engine/Source/Runtime/Datasmith/CADKernel/Engine/Private/TechSoft/TechSoftLibrary.cpp b/Engine/Source/Runtime/Datasmith/CADKernel/Engine/Private/TechSoft/TechSoftLibrary.cpp
--- a/Engine/Source/Runtime/Datasmith/CADKernel/Engine/Private/TechSoft/TechSoftLibrary.cpp
+++ b/Engine/Source/Runtime/Datasmith/CADKernel/Engine/Private/TechSoft/TechSoftLibrary.cpp
@@ -6,7 +6,9 @@
 #include "CADKernelEngine.h"
 #if PLATFORM_DESKTOP
 
+#ifdef WITH_HOOPS
 #include "hoops_license.h"
+#endif
 #include "TechSoftUniqueObjectImpl.h"
 
 #include "CADKernelEngineLog.h"
