Fixed unreachable code

#rnx

diff --git a/Engine/Source/Runtime/Datasmith/CADKernel/Base/Private/Core/Database.cpp b/Engine/Source/Runtime/Datasmith/CADKernel/Base/Private/Core/Database.cpp
--- a/Engine/Source/Runtime/Datasmith/CADKernel/Base/Private/Core/Database.cpp
+++ b/Engine/Source/Runtime/Datasmith/CADKernel/Base/Private/Core/Database.cpp
@@ -598,7 +598,6 @@ void FDatabase::ExpandSelection(TSharedPtr<FEntity> Entity, const TSet<EEntity>&
 
 	case EEntity::TopologicalEdge:
 	{
-		break;
 		TSharedPtr<FTopologicalEdge>Edge = StaticCastSharedPtr<FTopologicalEdge>(Entity);
 		ExpandSelection(Edge->GetCurve(), Filter, Selection);
 		break;
