[CADKernel] TechSoftParser support - styles/materials

- moved from Parser

diff --git a/Engine/Source/Runtime/Datasmith/CADKernel/Engine/Private/TechSoft/TechSoftLibrary.cpp b/Engine/Source/Runtime/Datasmith/CADKernel/Engine/Private/TechSoft/TechSoftLibrary.cpp
--- a/Engine/Source/Runtime/Datasmith/CADKernel/Engine/Private/TechSoft/TechSoftLibrary.cpp
+++ b/Engine/Source/Runtime/Datasmith/CADKernel/Engine/Private/TechSoft/TechSoftLibrary.cpp
@@ -684,6 +684,104 @@ namespace UE::CADKernel
 		}
 	}
 
+	bool FTechSoftLibrary::ExtractMaterial(uint32 TechSoftMaterialIndex, const A3DGraphStyleData& GraphStyleData, FTechSoftLibrary::FMaterial& ParsedMaterial)
+	{
+		UE::CADKernel::TechSoft::TUniqueObjectFromIndex<A3DGraphMaterialData> MaterialData(TechSoftMaterialIndex);
+		if(MaterialData.IsValid())
+		{
+			ParsedMaterial.Diffuse = GetColorAt(MaterialData->m_uiDiffuse);
+			ParsedMaterial.Ambient = GetColorAt(MaterialData->m_uiAmbient);
+			ParsedMaterial.Specular = GetColorAt(MaterialData->m_uiSpecular);
+			ParsedMaterial.Shininess = MaterialData->m_dShininess;
+			ParsedMaterial.Transparency = GraphStyleData.m_bIsTransparencyDefined
+				                              ? 1. - GraphStyleData.m_ucTransparency/255.f
+				                              : 0;
+			// todo: find how to convert Emissive color into ? reflexion coef...
+			// Material.Emissive = GetColor(MaterialData->m_uiEmissive);
+			// Material.Reflexion;
+			return true;
+		}
+		return false;
+	}
+
+	FColor FTechSoftLibrary::GetColorAt(uint32 ColorIndex)
+	{
+		using namespace UE::CADKernel;
+		TechSoft::TUniqueObjectFromIndex<A3DGraphRgbColorData> ColorData(ColorIndex);
+		if (ColorData.IsValid())
+		{
+			return FColor((uint8)(ColorData->m_dRed * 255)
+				, (uint8)(ColorData->m_dGreen * 255)
+				, (uint8)(ColorData->m_dBlue * 255));
+		}
+		return FColor(200, 200, 200);
+	}
+
+
+	TArray<uint32> FTechSoftLibrary::GetStylesUsedOnFaces(A3DRiBrepModel* RepresentationItem)
+	{
+		UE::CADKernel::TechSoft::TUniqueObject<A3DRiRepresentationItemData> RepresentationData(RepresentationItem);
+
+		using namespace UE::CADKernel;
+		TArray<uint32> MaterialsUsedOnFaces;
+		TechSoft::TUniqueObject<A3DTess3DData> Tess3DData(RepresentationData->m_pTessBase);
+		if (Tess3DData.IsValid() && Tess3DData->m_uiFaceTessSize != 0)
+		{
+			for (uint32 FaceTessIndex = 0; FaceTessIndex < Tess3DData->m_uiFaceTessSize; ++FaceTessIndex)
+			{
+				const A3DTessFaceData& TessFaceData = Tess3DData->m_psFaceTessData[FaceTessIndex];
+
+				if (TessFaceData.m_uiStyleIndexesSize != 0)
+				{
+					A3DUns32 StyleIndex = TessFaceData.m_puiStyleIndexes[0];
+					MaterialsUsedOnFaces.Add(StyleIndex);
+				}
+			}
+		}
+
+		TechSoft::TUniqueObject<A3DRiBrepModelData> BRepModelData(RepresentationItem);
+
+		if (BRepModelData.IsValid())
+		{
+			A3DTopoBrepData* BRepData = BRepModelData->m_pBrepData;
+
+			TechSoft::TUniqueObject<A3DTopoBrepDataData> TopoBrepData(BRepData);
+			if (TopoBrepData.IsValid())
+			{
+				for (A3DUns32 Index = 0; Index < TopoBrepData->m_uiConnexSize; ++Index)
+				{
+					TechSoft::TUniqueObject<A3DTopoConnexData> TopoConnexData(TopoBrepData->m_ppConnexes[Index]);
+					if (TopoConnexData.IsValid())
+					{
+						for (A3DUns32 Sndex = 0; Sndex < TopoConnexData->m_uiShellSize; ++Sndex)
+						{
+							TechSoft::TUniqueObject<A3DTopoShellData> ShellData(TopoConnexData->m_ppShells[Sndex]);
+							if (ShellData.IsValid())
+							{
+								for (A3DUns32 Fndex = 0; Fndex < ShellData->m_uiFaceSize; ++Fndex)
+								{
+									A3DTopoFace* TopoFace = ShellData->m_ppFaces[Fndex];
+									FTechSoftLibrary::FGraphicsProperties Result;
+									FTechSoftLibrary::ExtractGraphicsProperties(TopoFace, Result);
+
+									if (!Result.bIsRemoved && Result.bShow)
+									{
+										A3DUns32 StyleIndex = Result.StyleIndex;
+
+										MaterialsUsedOnFaces.Add(StyleIndex);
+									}
+								}
+							}
+						}
+					}
+				}
+			}
+		}
+
+		return MoveTemp(MaterialsUsedOnFaces);
+	}
+
+
 	bool FTechSoftLibrary::ParseRootBaseData(const A3DEntity* Entity, TMap<FString, FString>& MetaData, FString& UniqueID, FString& Label)
 	{
 		using namespace UE::CADKernel;
diff --git a/Engine/Source/Runtime/Datasmith/CADKernel/Engine/Public/CADKernelEngine.h b/Engine/Source/Runtime/Datasmith/CADKernel/Engine/Public/CADKernelEngine.h
--- a/Engine/Source/Runtime/Datasmith/CADKernel/Engine/Public/CADKernelEngine.h
+++ b/Engine/Source/Runtime/Datasmith/CADKernel/Engine/Public/CADKernelEngine.h
@@ -151,6 +151,19 @@ namespace UE::CADKernel
 		};
 		static CADKERNELENGINE_API void ExtractGraphicsProperties(const A3DEntity* Entity, FGraphicsProperties& Result);
 
+		struct FMaterial
+		{
+			FColor Diffuse;
+			FColor Ambient;
+			FColor Specular;
+			float Shininess;
+			float Transparency = 0;
+		};
+
+		static CADKERNELENGINE_API bool ExtractMaterial(uint32 TechSoftMaterialIndex, const A3DGraphStyleData& GraphStyleData, FMaterial& ParsedMaterial);
+		static CADKERNELENGINE_API FColor GetColorAt(uint32 ColorIndex);
+		static CADKERNELENGINE_API TArray<uint32> GetStylesUsedOnFaces(A3DRiBrepModel* RepresentationItem);
+
 		static CADKERNELENGINE_API bool ParseRootBaseData(const A3DEntity* Entity, TMap<FString, FString>& MetaData, FString& UniqueID, FString& Label);
 		static CADKERNELENGINE_API void TraverseAttribute(const A3DMiscAttributeData& AttributeData, TMap<FString, FString>& OutMetaData);
 
