[In Progress] InterchangeCADImporter
- Now UParametricSurfaceData derived from UAssetUserData so it can be embedded into a UStaticMesh object
- Added CADKernelEditor module to remove Engine dependency required as UParametricSurfaceData derives from  UAssetUserData
- Added code to transition out of UDatasmithParametricSurfaceData class so legacy DatasmithCAD imported static meshes can be tessellated
- Updated wire parser to support OpenModel 2025.0
- Enabled TechSoft parser
- Fixed empty FMeshDescription for OpenNurbs parser.

#rnx

diff --git a/Engine/Source/Runtime/Datasmith/CADKernel/Editor/CADKernelEditor.Build.cs b/Engine/Source/Runtime/Datasmith/CADKernel/Editor/CADKernelEditor.Build.cs
new file mode 100644
--- /dev/null
+++ b/Engine/Source/Runtime/Datasmith/CADKernel/Editor/CADKernelEditor.Build.cs
@@ -0,0 +1,48 @@
+// Copyright Epic Games, Inc. All Rights Reserved.
+using System.IO;
+
+namespace UnrealBuildTool.Rules
+{
+	public class CADKernelEditor : ModuleRules
+	{
+		public CADKernelEditor(ReadOnlyTargetRules Target)
+			: base(Target)
+		{
+			bEnableExceptions = true;
+			bUseUnity = false; // TechSoftLibrary.cpp cannot be compiled with the other cpp files.
+			OptimizeCode = CodeOptimization.Never;
+
+			PublicDependencyModuleNames.AddRange(
+				new string[]
+				{
+					"Core",
+					"CoreUObject",
+				}
+			);
+
+			PrivateDependencyModuleNames.AddRange(
+				new string[]
+				{
+					"CADKernel",
+					"CADKernelEngine",
+					"Engine",
+				}
+			);
+
+			if (Target.bBuildEditor)
+			{
+				PrivateDependencyModuleNames.AddRange(
+					new string[]
+					{
+						"ContentBrowser",
+						"EditorFramework",
+						"Slate",
+						"SlateCore",
+						"StaticMeshEditor",
+						"ToolMenus",
+					}
+				);
+			}
+		}
+	}
+}
\ No newline at end of file
diff --git a/Engine/Source/Runtime/Datasmith/CADKernel/Editor/Private/CADKernelRetessellateAction.cpp b/Engine/Source/Runtime/Datasmith/CADKernel/Editor/Private/CADKernelRetessellateAction.cpp
new file mode 100644
--- /dev/null
+++ b/Engine/Source/Runtime/Datasmith/CADKernel/Editor/Private/CADKernelRetessellateAction.cpp
@@ -0,0 +1,228 @@
+// Copyright Epic Games, Inc. All Rights Reserved.
+
+#include "CADKernelRetessellateAction.h"
+
+#include "LegacyParametricSurfaceData.h"
+
+#include "CADKernelEngine.h"
+#include "ParametricSurfaceData.h"
+
+#include "Algo/AnyOf.h"
+#include "Async/ParallelFor.h"
+#include "Components/StaticMeshComponent.h"
+#include "Engine/StaticMesh.h"
+#include "Engine/StaticMeshActor.h"
+#include "Logging/LogMacros.h"
+
+#if WITH_EDITOR
+#include "ContentBrowserMenuContexts.h"
+#include "IStaticMeshEditor.h"
+#include "Toolkits/IToolkit.h"
+#include "Toolkits/ToolkitManager.h"
+#include "ToolMenus.h"
+#endif
+
+DEFINE_LOG_CATEGORY_STATIC(LogCADKernelRetessellation, Log, All)
+
+#define LOCTEXT_NAMESPACE "CADKernelRetessellateAction"
+
+static const FName CADKernelTessellationTag(TEXT("CADKernel"));
+
+static bool GCADKernelRetessellation = false;
+FAutoConsoleVariableRef GCADKernelDebugUseEngine(
+	TEXT("CADKernel.Debug.Retessellation"),
+	GCADKernelRetessellation,
+	TEXT(""),
+	ECVF_Default);
+
+namespace UE::CADKernel
+{
+	class FRetessellationUtils
+	{
+	public:
+		static bool ConfirmAndUpdateTessellation(FCADKernelRetessellationSettings& Settings)
+		{
+			return false;
+		}
+
+		static UParametricSurfaceData* GetParametricSurfaceData(UStaticMesh* StaticMesh)
+		{
+			UParametricSurfaceData* ParametricSurfaceData = nullptr;
+
+			if (IInterface_AssetUserData* AssetUserData = Cast<IInterface_AssetUserData>(StaticMesh))
+			{
+				ParametricSurfaceData = AssetUserData->GetAssetUserData<UParametricSurfaceData>();
+			}
+
+			return ParametricSurfaceData;
+		}
+
+		static bool Retessellate(UStaticMesh& StaticMesh, const FLegacyParametricSurfaceData* LegacyData = nullptr)
+		{
+			IInterface_AssetUserData* AssetUserData = Cast<IInterface_AssetUserData>(&StaticMesh);
+			if (!ensure(AssetUserData))
+			{
+				return false;
+			}
+
+			UParametricSurfaceData* ParametricSurfaceData = LegacyData ? LegacyData->ToCADKernel() : GetParametricSurfaceData(&StaticMesh);
+			if (!ensure(ParametricSurfaceData))
+			{
+				return false;
+			}
+
+			if (GCADKernelRetessellation)
+			{
+				// If StaticMesh has been generated using the Datasmith importer, add ParametricSurfaceData to StaticMesh's AssetUserData
+				if (LegacyData)
+				{
+					ParametricSurfaceData->Rename(nullptr, &StaticMesh, REN_DontCreateRedirectors | REN_NonTransactional);
+					AssetUserData->AddAssetUserData(ParametricSurfaceData);
+				}
+			}
+
+			return GCADKernelRetessellation;
+		}
+
+#if WITH_EDITOR
+		static void FinalizeChanges(UStaticMesh* StaticMesh)
+		{
+			StaticMesh->PostEditChange();
+			StaticMesh->MarkPackageDirty();
+
+			// Refresh associated editor 
+			TSharedPtr<IToolkit> EditingToolkit = FToolkitManager::Get().FindEditorForAsset(StaticMesh);
+			if (IStaticMeshEditor* StaticMeshEditorInUse = StaticCastSharedPtr<IStaticMeshEditor>(EditingToolkit).Get())
+			{
+				StaticMeshEditorInUse->RefreshTool();
+			}
+		}
+#endif
+	};
+}
+
+bool FCADKernelRetessellateAction::CanRetessellate(UStaticMesh* StaticMesh)
+{
+	using namespace UE::CADKernel;
+	return FRetessellationUtils::GetParametricSurfaceData(StaticMesh) != nullptr;
+}
+
+bool FCADKernelRetessellateAction::RetessellateArray(const TArray<UStaticMesh*>& StaticMeshes)
+{
+	using namespace UE::CADKernel;
+
+	return false;
+}
+
+bool FCADKernelRetessellateAction::RetessellateArray(const TArray<FAssetData>& Assets)
+{
+	TArray<UStaticMesh*> StaticMeshes;
+	StaticMeshes.Reserve(Assets.Num());
+
+	for (const FAssetData& AssetData : Assets)
+	{
+		if (UStaticMesh* StaticMesh = Cast<UStaticMesh>(AssetData.GetAsset()))
+		{
+			StaticMeshes.Add(StaticMesh);
+		}
+	}
+
+	return RetessellateArray(StaticMeshes);
+}
+
+bool FCADKernelRetessellateAction::RetessellateArray(const TArray<AStaticMeshActor*>& StaticMeshActors)
+{
+	TSet<UStaticMesh*> StaticMeshes;
+	StaticMeshes.Reserve(StaticMeshActors.Num());
+
+	for (AStaticMeshActor* StaticMeshActor : StaticMeshActors)
+	{
+		if (StaticMeshActor)
+		{
+			if (UStaticMeshComponent* MeshComponent = StaticMeshActor->GetStaticMeshComponent())
+			{
+				if (UStaticMesh* StaticMesh = MeshComponent->GetStaticMesh())
+				{
+					StaticMeshes.Add(StaticMesh);
+				}
+			}
+		}
+	}
+
+	return RetessellateArray(StaticMeshes.Array());
+}
+
+bool FCADKernelRetessellateAction::RetessellateLegacy(UStaticMesh& StaticMesh, FArchive& Ar)
+{
+	using namespace UE::CADKernel;
+
+	ensure(Ar.IsLoading());
+
+	FLegacyParametricSurfaceData LegacyParametricSurfaceData;
+	LegacyParametricSurfaceData.Serialize(Ar);
+
+	return FRetessellationUtils::Retessellate(StaticMesh, &LegacyParametricSurfaceData);
+}
+
+#if WITH_EDITOR
+namespace MenuExtension_CADKernel_StaticMesh
+{
+	static const TAttribute<FText> Label = LOCTEXT("CADKernelRetessellateAction_Label", "Retessellate (New)");
+	static const TAttribute<FText> ToolTip = LOCTEXT("CADKernelRetessellateAction_Tooltip", "This will execute the action for your type.");
+	static const FSlateIcon Icon = FSlateIcon();
+
+	static void ExecuteRetessellation(const FToolMenuContext& MenuContext)
+	{
+		const UContentBrowserAssetContextMenuContext* CBContext = UContentBrowserAssetContextMenuContext::FindContextWithAssets(MenuContext);
+		FCADKernelRetessellateAction::RetessellateArray(CBContext->LoadSelectedObjects<UStaticMesh>());
+	}
+
+	static bool DisplayMenuEntry(const FToolMenuContext& MenuContext)
+	{
+		const UContentBrowserAssetContextMenuContext* CBContext = UContentBrowserAssetContextMenuContext::FindContextWithAssets(MenuContext);
+
+		for (UStaticMesh* StaticMesh : CBContext->LoadSelectedObjects<UStaticMesh>())
+		{
+			if (FCADKernelRetessellateAction::CanRetessellate(StaticMesh))
+			{
+				return GCADKernelRetessellation;
+			}
+		}
+
+		return false;
+	}
+
+	static FDelayedAutoRegisterHelper DelayedAutoRegister(EDelayedRegisterRunPhase::EndOfEngineInit, []
+		{
+			UToolMenus::RegisterStartupCallback(FSimpleMulticastDelegate::FDelegate::CreateLambda([]()
+				{
+					FToolMenuOwnerScoped OwnerScoped(UE_MODULE_NAME);
+					UToolMenu* Menu = UE::ContentBrowser::ExtendToolMenu_AssetContextMenu(UStaticMesh::StaticClass());
+
+					FToolMenuSection& Section = Menu->FindOrAddSection("GetAssetActions");
+					Section.AddDynamicEntry(NAME_None, FNewToolMenuSectionDelegate::CreateLambda([](FToolMenuSection& InSection)
+						{
+							if (const UContentBrowserAssetContextMenuContext* CBContext = UContentBrowserAssetContextMenuContext::FindContextWithAssets(InSection))
+							{
+								bool bCanRetessellate = Algo::AnyOf(CBContext->SelectedAssets, [](const FAssetData& InAssetData)
+									{
+										int32 Value;
+										return InAssetData.GetTagValue(CADKernelTessellationTag, Value);
+									}
+								);
+
+								FToolUIAction UIAction;
+								UIAction.IsActionVisibleDelegate = FToolMenuIsActionButtonVisible::CreateStatic(&DisplayMenuEntry);
+								UIAction.ExecuteAction = FToolMenuExecuteAction::CreateStatic(&ExecuteRetessellation);
+								InSection.AddMenuEntry("CADKernelRetessellateAction_Entry", Label, ToolTip, Icon, UIAction);
+							}
+						})
+					);
+				})
+			);
+		}
+	);
+}
+#endif
+
+#undef LOCTEXT_NAMESPACE
diff --git a/Engine/Source/Runtime/Datasmith/CADKernel/Editor/Private/LegacyParametricSurfaceData.cpp b/Engine/Source/Runtime/Datasmith/CADKernel/Editor/Private/LegacyParametricSurfaceData.cpp
new file mode 100644
--- /dev/null
+++ b/Engine/Source/Runtime/Datasmith/CADKernel/Editor/Private/LegacyParametricSurfaceData.cpp
@@ -0,0 +1,40 @@
+// Copyright Epic Games, Inc. All Rights Reserved.
+
+#include "LegacyParametricSurfaceData.h"
+
+UParametricSurfaceData* FLegacyParametricSurfaceData::ToCADKernel() const
+{
+	UParametricSurfaceData* ParametricSurfaceData = NewObject<UParametricSurfaceData>();
+
+	ParametricSurfaceData->SetModelParameters(SceneParameters.ToCADKernel());
+	ParametricSurfaceData->SetMeshParameters(MeshParameters.ToCADKernel());
+	ParametricSurfaceData->SetLastTessellationSettings(LastTessellationOptions.ToCADKernel());
+
+	ParametricSurfaceData->SetRawData(RawData, !LastTessellationOptions.bUseCADKernel);
+
+	return ParametricSurfaceData;
+}
+
+void FLegacyParametricSurfaceData::Serialize(FArchive& Ar)
+{
+	ensure(Ar.IsLoading());
+
+	Ar << SceneParameters.ModelCoordSys;
+	Ar << SceneParameters.MetricUnit;
+	Ar << SceneParameters.ScaleFactor;
+
+	Ar << MeshParameters.bNeedSwapOrientation;
+	Ar << MeshParameters.bIsSymmetric;
+	Ar << MeshParameters.SymmetricOrigin;
+	Ar << MeshParameters.SymmetricNormal;
+
+	Ar << LastTessellationOptions.ChordTolerance;
+	Ar << LastTessellationOptions.MaxEdgeLength;
+	Ar << LastTessellationOptions.NormalTolerance;
+	Ar << LastTessellationOptions.StitchingTechnique;
+	Ar << LastTessellationOptions.bUseCADKernel;
+	Ar << LastTessellationOptions.GeometricTolerance;
+	Ar << LastTessellationOptions.StitchingTolerance;
+
+	Ar << RawData;
+}
diff --git a/Engine/Source/Runtime/Datasmith/CADKernel/Editor/Private/LegacyParametricSurfaceData.h b/Engine/Source/Runtime/Datasmith/CADKernel/Editor/Private/LegacyParametricSurfaceData.h
new file mode 100644
--- /dev/null
+++ b/Engine/Source/Runtime/Datasmith/CADKernel/Editor/Private/LegacyParametricSurfaceData.h
@@ -0,0 +1,86 @@
+// Copyright Epic Games, Inc. All Rights Reserved.
+
+#pragma once
+
+#include "CoreMinimal.h"
+
+#include "ParametricSurfaceData.h"
+
+namespace UE::CADKernel
+{
+	class FRetessellationUtils;
+}
+
+struct FLegacyParametricSceneParameters
+{
+	uint8 ModelCoordSys = (uint8)ECADKernelModelCoordSystem::ZUp_LeftHanded;
+	float MetricUnit = 0.01f;
+	float ScaleFactor = 1.0f;
+
+	FCADKernelModelParameters ToCADKernel() const
+	{
+		return {
+			ModelCoordSys,
+			MetricUnit,
+			ScaleFactor
+		};
+	}
+};
+
+struct FLegacyParametricMeshParameters
+{
+	bool bNeedSwapOrientation = false;
+	bool bIsSymmetric = false;
+	FVector SymmetricOrigin = FVector::ZeroVector;
+	FVector SymmetricNormal = FVector::ZeroVector;
+
+	FCADKernelMeshParameters ToCADKernel() const
+	{
+		return {
+			bNeedSwapOrientation,
+			bIsSymmetric,
+			SymmetricOrigin,
+			SymmetricNormal,
+		};
+	}
+};
+
+struct FLegacyTessellationOptions
+{
+	float ChordTolerance;
+	float MaxEdgeLength;
+	float NormalTolerance;
+	uint8 StitchingTechnique;
+	bool bUseCADKernel = true;
+	double GeometricTolerance = 0.001;
+	double StitchingTolerance = 0.001;
+
+	FCADKernelTessellationSettings ToCADKernel() const
+	{
+		return FCADKernelTessellationSettings(
+			ChordTolerance,
+			MaxEdgeLength,
+			NormalTolerance,
+			static_cast<ECADKernelStitchingTechnique>(StitchingTechnique),
+			GeometricTolerance,
+			StitchingTolerance,
+			bUseCADKernel
+		);
+	}
+};
+
+class FLegacyParametricSurfaceData
+{
+	UParametricSurfaceData* ToCADKernel() const;
+
+	void Serialize(FArchive& Ar);
+
+	FLegacyParametricSceneParameters SceneParameters;
+	FLegacyParametricMeshParameters MeshParameters;
+	FLegacyTessellationOptions LastTessellationOptions;
+	TArray<uint8> RawData;
+
+	friend UE::CADKernel::FRetessellationUtils;
+	friend class FCADKernelRetessellateAction;
+};
+
diff --git a/Engine/Source/Runtime/Datasmith/CADKernel/Editor/Private/ParametricSurfaceData.cpp b/Engine/Source/Runtime/Datasmith/CADKernel/Editor/Private/ParametricSurfaceData.cpp
new file mode 100644
--- /dev/null
+++ b/Engine/Source/Runtime/Datasmith/CADKernel/Editor/Private/ParametricSurfaceData.cpp
@@ -0,0 +1,172 @@
+// Copyright Epic Games, Inc. All Rights Reserved.
+
+#include "ParametricSurfaceData.h"
+
+#include "CADKernelEngine.h"
+
+#include "Core/Session.h"
+#include "Topo/Model.h"
+
+#include "Misc/FileHelper.h"
+#include "Misc/Paths.h"
+#include "UObject/EnterpriseObjectVersion.h"
+
+#define LOCTEXT_NAMESPACE "ParametricSurfaceData"
+
+namespace UE::CADKernel::Editor
+{
+	namespace Private
+	{
+#if PLATFORM_DESKTOP
+		using namespace UE::CADKernel;
+
+		template<typename MeshType>
+		bool Tessellate(UParametricSurfaceData& Data, MeshType& MeshOut)
+		{
+			using namespace UE::CADKernel;
+
+			FTessellationContext Context(Data.GetModelParameters(), Data.GetMeshParameters(), Data.GetLastTessellationSettings());
+
+			if (Context.TessellationSettings.bUseCADKernel)
+			{
+				TSharedPtr<FModel> Model = Data.GetModel();
+				if (Model.IsValid())
+				{
+					return FCADKernelUtilities::Tessellate(*Model, Context, MeshOut);
+				}
+			}
+			else if (FTechSoftLibrary::Initialize())
+			{
+				A3DRiRepresentationItem* Representation = Data.GetRepresentation();
+				return FTechSoftUtilities::Tessellate(Representation, Context, MeshOut);
+			}
+
+			return false;
+		}
+
+		template<typename MeshType>
+		bool Retessellate(UParametricSurfaceData& Data, const FCADKernelRetessellationSettings& Settings, MeshType& MeshOut)
+		{
+			using namespace UE::CADKernel;
+			using namespace UE::CADKernel::MeshUtilities;
+
+			FTessellationContext Context(Data.GetModelParameters(), Data.GetMeshParameters(), Settings);
+
+			if (Context.TessellationSettings.bUseCADKernel)
+			{
+				TSharedPtr<FModel> Model = Data.GetModel();
+				if (Model.IsValid())
+				{
+					const bool bSkipDeletedFaces = Settings.RetessellationRule == ECADKernelRetessellationRule::SkipDeletedFaces;
+					return FCADKernelUtilities::Tessellate(*Model, Context, MeshOut, bSkipDeletedFaces, true);
+				}
+			}
+			else if (FTechSoftLibrary::Initialize())
+			{
+				if (A3DRiRepresentationItem* Representation = Data.GetRepresentation())
+				{
+					const bool bSkipDeletedFaces = Settings.RetessellationRule == ECADKernelRetessellationRule::SkipDeletedFaces;
+					return FTechSoftUtilities::Tessellate(Representation, Context, MeshOut, bSkipDeletedFaces, true);
+				}
+			}
+
+			return false;
+		}
+#endif
+	}
+}
+
+void UParametricSurfaceData::Serialize(FArchive& Ar)
+{
+	Ar.UsingCustomVersion(FEnterpriseObjectVersion::GUID);
+
+	if (Ar.IsLoading() && Ar.CustomVer(FEnterpriseObjectVersion::GUID) < FEnterpriseObjectVersion::AddedParametricSurfaceData)
+	{
+		// #cad_import_todo : Put code to deserialize UDatasmithParametricSurfaceData
+		 return;
+	}
+
+	Super::Serialize(Ar);
+
+	Ar << CADKernelRawData;
+	Ar << TechSoftRawData;
+}
+
+TSharedPtr<UE::CADKernel::FModel> UParametricSurfaceData::GetModel()
+{
+	using namespace UE::CADKernel;
+
+	if (CADKernelRawData.IsEmpty())
+	{
+		return {};
+	}
+
+	TSharedRef<FSession> Session = MakeShared<FSession>(LastTessellationSettings.GetGeometricTolerance(true));
+	Session->AddDatabase(CADKernelRawData);
+
+	return Session->GetModelAsShared();
+}
+
+bool UParametricSurfaceData::SetFromFile(const TCHAR* FilePath, bool bForTechSoft)
+{
+	if (FPaths::FileExists(FilePath))
+	{
+		TArray<uint8> ByteArray;
+
+		if (FFileHelper::LoadFileToArray(ByteArray, FilePath))
+		{
+			TArray<uint8>& RawData = bForTechSoft ? TechSoftRawData : CADKernelRawData;
+			RawData = MoveTemp(ByteArray);
+			return true;
+		}
+	}
+
+	return false;
+}
+
+#if PLATFORM_DESKTOP
+bool UParametricSurfaceData::SetModel(TSharedPtr<UE::CADKernel::FModel>& Model, double UnitModelToCentimeter)
+{
+	ModelParameters.ModelUnitToCentimeter = UnitModelToCentimeter;
+	return UE::CADKernel::FCADKernelUtilities::ToRawData(Model, CADKernelRawData);
+}
+
+A3DRiRepresentationItem* UParametricSurfaceData::GetRepresentation()
+{
+	using namespace UE::CADKernel;
+
+	return TechSoftRawData.IsEmpty() ? nullptr : FTechSoftUtilities::GetRepresentation(TechSoftRawData);
+}
+
+bool UParametricSurfaceData::SetRepresentation(A3DRiRepresentationItem* Representation, int32 MaterialID, double UnitRepresentationToCentimeter)
+{
+#ifdef WITH_HOOPS
+	ModelParameters.ModelUnitToCentimeter = UnitRepresentationToCentimeter;
+	return UE::CADKernel::FTechSoftUtilities::ToRawData(Representation, MaterialID, TechSoftRawData);
+#else
+	return false;
+#endif
+}
+
+bool UParametricSurfaceData::Tessellate(UE::Geometry::FDynamicMesh3& MeshOut)
+{
+	return UE::CADKernel::Editor::Private::Tessellate(*this, MeshOut);
+}
+
+bool UParametricSurfaceData::Tessellate(FMeshDescription& MeshOut)
+{
+	return UE::CADKernel::Editor::Private::Tessellate(*this, MeshOut);
+}
+
+bool UParametricSurfaceData::Retessellate(const FCADKernelRetessellationSettings& Settings, UE::Geometry::FDynamicMesh3& MeshOut)
+{
+	return UE::CADKernel::Editor::Private::Retessellate(*this, Settings, MeshOut);
+}
+
+bool UParametricSurfaceData::Retessellate(const FCADKernelRetessellationSettings& Settings, FMeshDescription& MeshOut)
+{
+	return UE::CADKernel::Editor::Private::Retessellate(*this, Settings, MeshOut);
+}
+#endif
+
+#undef LOCTEXT_NAMESPACE // "ParametricSurfaceData"
diff --git a/Engine/Source/Runtime/Datasmith/CADKernel/Editor/Public/CADKernelRetessellateAction.h b/Engine/Source/Runtime/Datasmith/CADKernel/Editor/Public/CADKernelRetessellateAction.h
new file mode 100644
--- /dev/null
+++ b/Engine/Source/Runtime/Datasmith/CADKernel/Editor/Public/CADKernelRetessellateAction.h
@@ -0,0 +1,37 @@
+// Copyright Epic Games, Inc. All Rights Reserved.
+
+#pragma once
+
+#include "CoreMinimal.h"
+
+#include "AssetRegistry/AssetData.h"
+
+
+class UStaticMesh;
+class AStaticMeshActor;
+
+class FCADKernelRetessellateAction
+{
+public:
+	static CADKERNELEDITOR_API bool CanRetessellate(UStaticMesh* StaticMesh);
+
+	static CADKERNELEDITOR_API bool RetessellateArray(const TArray<UStaticMesh*>& StaticMeshes);
+	static CADKERNELEDITOR_API bool Retessellate(UStaticMesh* StaticMesh)
+	{
+		return RetessellateArray({ StaticMesh });
+	}
+
+	static CADKERNELEDITOR_API bool RetessellateArray(const TArray<FAssetData>& Assets);
+	static CADKERNELEDITOR_API bool Retessellate(const FAssetData& Asset)
+	{
+		return RetessellateArray({ Asset });
+	}
+
+	static CADKERNELEDITOR_API bool RetessellateArray(const TArray<AStaticMeshActor*>& StaticMeshActors);
+	static CADKERNELEDITOR_API bool Retessellate(AStaticMeshActor* StaticMeshActor)
+	{
+		return RetessellateArray({ StaticMeshActor });
+	}
+
+	static CADKERNELEDITOR_API bool RetessellateLegacy(UStaticMesh& StaticMesh, FArchive& Ar);
+};
diff --git a/Engine/Source/Runtime/Datasmith/CADKernel/Editor/Public/ParametricSurfaceData.h b/Engine/Source/Runtime/Datasmith/CADKernel/Editor/Public/ParametricSurfaceData.h
new file mode 100644
--- /dev/null
+++ b/Engine/Source/Runtime/Datasmith/CADKernel/Editor/Public/ParametricSurfaceData.h
@@ -0,0 +1,90 @@
+// Copyright Epic Games, Inc. All Rights Reserved.
+
+#pragma once
+
+#include "TechSoftIncludes.h"
+#include "CADKernelEngineDefinitions.h"
+
+#include "Engine/AssetUserData.h"
+
+#include "ParametricSurfaceData.generated.h"
+
+namespace UE::CADKernel
+{
+	class FModel;
+}
+
+namespace UE::Geometry
+{
+	class FDynamicMesh3;
+}
+
+class UStaticMesh;
+struct FMeshDescription;
+
+typedef void A3DRiRepresentationItem;
+typedef void A3DAsmModelFile;
+
+UCLASS(meta = (DisplayName = "CADKernel Parametric Surface Data"))
+class CADKERNELEDITOR_API UParametricSurfaceData : public UAssetUserData
+{
+	GENERATED_BODY()
+
+public:
+
+	virtual bool IsValid() { return CADKernelRawData.Num() > 0 || TechSoftRawData.Num() > 0; }
+
+	void SetModelParameters(const FCADKernelModelParameters& InModelParameters) { ModelParameters = InModelParameters; }
+	const FCADKernelModelParameters& GetModelParameters() const { return ModelParameters; }
+	FCADKernelModelParameters& GetModelParameters() { return ModelParameters; }
+
+	void SetMeshParameters(const FCADKernelMeshParameters& InMeshParameters) { MeshParameters = InMeshParameters; }
+	const FCADKernelMeshParameters& GetMeshParameters() const { return MeshParameters; }
+	FCADKernelMeshParameters& GetMeshParameters() { return MeshParameters; }
+
+	void SetLastTessellationSettings(const FCADKernelTessellationSettings& InTessellationSettings) { LastTessellationSettings = InTessellationSettings; }
+	const FCADKernelTessellationSettings& GetLastTessellationSettings() const { return LastTessellationSettings; }
+	
+	TSharedPtr<UE::CADKernel::FModel> GetModel();
+
+	virtual bool SetFromFile(const TCHAR* FilePath, bool bForTechSoft = false);
+
+#if PLATFORM_DESKTOP
+	bool SetModel(TSharedPtr<UE::CADKernel::FModel>& Model, double UnitModelToCentimeter = 1.);
+
+	A3DRiRepresentationItem* GetRepresentation();
+	bool SetRepresentation(A3DRiRepresentationItem* Representation, int32 MaterialID, double UnitRepresentationToCentimeter = 1.);
+
+	bool Tessellate(UE::Geometry::FDynamicMesh3& MeshOut);
+	bool Tessellate(FMeshDescription& MeshOut);
+
+	bool Retessellate(const FCADKernelRetessellationSettings& Settings, UE::Geometry::FDynamicMesh3& MeshOut);
+	bool Retessellate(const FCADKernelRetessellationSettings& Settings, FMeshDescription& MeshOut);
+#endif
+
+	/*
+	 ** The SetRawData is only for internal use to facilitate the transition
+	 ** out of UDatasmithParametricSurfaceData into the new UParametricSurfaceData class
+	 */
+	virtual void SetRawData(const TArray<uint8>& InRawData, bool bForTechSoft = false)
+	{
+		TArray<uint8>& RawData = bForTechSoft ? TechSoftRawData : CADKernelRawData;
+		RawData = InRawData;
+	}
+
+protected:
+	virtual void Serialize(FArchive& Ar) override;
+
+protected:
+	UPROPERTY()
+	FCADKernelModelParameters ModelParameters;
+
+	UPROPERTY()
+	FCADKernelMeshParameters MeshParameters;
+
+	UPROPERTY(EditAnywhere, Category = NURBS)
+	FCADKernelTessellationSettings LastTessellationSettings;
+
+	TArray<uint8> CADKernelRawData;
+	TArray<uint8> TechSoftRawData;
+};
diff --git a/Engine/Source/Runtime/Datasmith/CADKernel/Engine/CADKernelEngine.Build.cs b/Engine/Source/Runtime/Datasmith/CADKernel/Engine/CADKernelEngine.Build.cs
--- a/Engine/Source/Runtime/Datasmith/CADKernel/Engine/CADKernelEngine.Build.cs
+++ b/Engine/Source/Runtime/Datasmith/CADKernel/Engine/CADKernelEngine.Build.cs
@@ -8,8 +8,9 @@ namespace UnrealBuildTool.Rules
 		public CADKernelEngine(ReadOnlyTargetRules Target)
 			: base(Target)
 		{
+			bEnableExceptions = true;
 			bUseUnity = false; // TechSoftLibrary.cpp cannot be compiled with the other cpp files.
-			//OptimizeCode = CodeOptimization.Never;
+			OptimizeCode = CodeOptimization.Never;
 
 			PublicDependencyModuleNames.AddRange(
 				new string[]
diff --git a/Engine/Source/Runtime/Datasmith/CADKernel/Engine/Private/CADKernelEngine.cpp b/Engine/Source/Runtime/Datasmith/CADKernel/Engine/Private/CADKernelEngine.cpp
--- a/Engine/Source/Runtime/Datasmith/CADKernel/Engine/Private/CADKernelEngine.cpp
+++ b/Engine/Source/Runtime/Datasmith/CADKernel/Engine/Private/CADKernelEngine.cpp
@@ -18,6 +18,22 @@
 
 DEFINE_LOG_CATEGORY(LogCADKernelEngine);
 
+#ifdef WITH_HOOPS
+bool FCADKernelTessellationSettings::bWithHoops = true;
+#else
+bool FCADKernelTessellationSettings::bWithHoops = false;
+#endif
+
+namespace UE::CADKernel
+{
+	const double FUnitConverter::CentimeterToMillimeter = 10.;
+	const double FUnitConverter::MillimeterToCentimeter = 0.1;
+	const double FUnitConverter::CentimeterToMeter = 0.01;
+	const double FUnitConverter::MeterToCentimeter = 100.;
+	const double FUnitConverter::MillimeterToMeter = 0.001;
+	const double FUnitConverter::MeterToMillimeter = 1000.;
+}
+
 namespace UE::CADKernel
 {
 	using namespace UE::CADKernel::MeshUtilities;
@@ -58,20 +74,35 @@ namespace UE::CADKernel
 		return false;
 	}
 
-	bool FCADKernelUtilities::Tessellate(UE::CADKernel::FModel& Model, const FTessellationContext& Context, FMeshDescription& Mesh, bool bEmptyMesh)
+	bool FCADKernelUtilities::Tessellate(UE::CADKernel::FModel& Model, const FTessellationContext& Context, FMeshDescription& Mesh, bool bSkipDeletedFaces, bool bEmptyMesh)
 	{
 		using namespace UE::CADKernel::MeshUtilities;
 
+		if (bSkipDeletedFaces)
+		{
+			GetExistingFaceGroups(Mesh, Context.FaceGroupsToExtract);
+		}
+
 		TSharedPtr<FMeshWrapperAbstract> MeshWrapper = FMeshWrapperAbstract::MakeWrapper(Context, Mesh);
 		return Private::Tessellate(Model, Context, *MeshWrapper, bEmptyMesh);
 	}
 
-	bool FCADKernelUtilities::Tessellate(UE::CADKernel::FModel& Model, const FTessellationContext& Context, UE::Geometry::FDynamicMesh3& Mesh, bool bEmptyMesh)
+	bool FCADKernelUtilities::Tessellate(UE::CADKernel::FModel& Model, const FTessellationContext& Context, UE::Geometry::FDynamicMesh3& Mesh, bool bSkipDeletedFaces, bool bEmptyMesh)
 	{
 		using namespace UE::CADKernel::MeshUtilities;
 
+		if (bSkipDeletedFaces)
+		{
+			GetExistingFaceGroups(Mesh, Context.FaceGroupsToExtract);
+		}
+
 		TSharedPtr<FMeshWrapperAbstract> MeshWrapper = FMeshWrapperAbstract::MakeWrapper(Context, Mesh);
 		return Private::Tessellate(Model, Context, *MeshWrapper, bEmptyMesh);
 	}
+
+	bool FCADKernelUtilities::ToRawData(TSharedPtr<UE::CADKernel::FModel>& Model, TArray<uint8>& RawDataOut)
+	{
+		return false;
+	}
 }
 #endif
diff --git a/Engine/Source/Runtime/Datasmith/CADKernel/Engine/Private/CADKernelEnginePrivate.cpp b/Engine/Source/Runtime/Datasmith/CADKernel/Engine/Private/CADKernelEnginePrivate.cpp
--- a/Engine/Source/Runtime/Datasmith/CADKernel/Engine/Private/CADKernelEnginePrivate.cpp
+++ b/Engine/Source/Runtime/Datasmith/CADKernel/Engine/Private/CADKernelEnginePrivate.cpp
@@ -83,11 +83,6 @@ namespace UE::CADKernel::Private
 		return AddModelMesh(*ModelMesh, MeshWrapper);
 	}
 
-	bool ToRawData(TSharedPtr<UE::CADKernel::FModel>& Model, TArray<uint8>& RawDataOut)
-	{
-		return false;
-	}
-
 	bool GetFaceTrimmingCurves(const FModel& Model, const FTopologicalFace& Face, TArray<TArray<TArray<FVector>>>& CurvesOut)
 	{
 		using namespace UE::CADKernel;
diff --git a/Engine/Source/Runtime/Datasmith/CADKernel/Engine/Private/CADKernelEnginePrivate.h b/Engine/Source/Runtime/Datasmith/CADKernel/Engine/Private/CADKernelEnginePrivate.h
--- a/Engine/Source/Runtime/Datasmith/CADKernel/Engine/Private/CADKernelEnginePrivate.h
+++ b/Engine/Source/Runtime/Datasmith/CADKernel/Engine/Private/CADKernelEnginePrivate.h
@@ -25,67 +25,6 @@ namespace UE::CADKernel
 
 	namespace Private
 	{
-
-		template<typename MeshType>
-		bool Tessellate(UParametricSurfaceData& Data, MeshType& MeshOut)
-		{
-			using namespace UE::CADKernel;
-
-			bool bTessellationSuccessful = false;
-
-			FTessellationContext Context(Data.GetModelParameters(), Data.GetMeshParameters(), Data.GetLastTessellationSettings());
-
-			if (Context.TessellationSettings.bUseCADKernel)
-			{
-				TSharedPtr<FModel> Model = Data.GetModel();
-				if (Model.IsValid())
-				{
-					bTessellationSuccessful = FCADKernelUtilities::Tessellate(*Model, Context, MeshOut);
-				}
-			}
-			else if (FTechSoftLibrary::Initialize())
-			{
-				A3DRiRepresentationItem* Representation = Data.GetRepresentation();
-				bTessellationSuccessful = FTechSoftUtilities::Tessellate(Representation, Context, MeshOut);
-			}
-
-			return bTessellationSuccessful;
-		}
-
-		template<typename MeshType>
-		bool Retessellate(UParametricSurfaceData& Data, const FCADKernelRetessellationSettings& Settings, MeshType& MeshOut)
-		{
-			using namespace UE::CADKernel;
-			using namespace UE::CADKernel::MeshUtilities;
-
-			bool bTessellationSuccessful = false;
-
-			FTessellationContext Context(Data.GetModelParameters(), Data.GetMeshParameters(), Settings);
-
-			if (Settings.RetessellationRule == ECADKernelRetessellationRule::SkipDeletedFaces)
-			{
-				GetExistingFaceGroups(MeshOut, Context.FaceGroupsToExtract);
-			}
-
-			if (Context.TessellationSettings.bUseCADKernel)
-			{
-				TSharedPtr<FModel> Model = Data.GetModel();
-				if (Model.IsValid())
-				{
-					bTessellationSuccessful = FCADKernelUtilities::Tessellate(*Model, Context, MeshOut, true);
-				}
-			}
-			else if(FTechSoftLibrary::Initialize())
-			{
-				if (A3DRiRepresentationItem* Representation = Data.GetRepresentation())
-				{
-					bTessellationSuccessful = FTechSoftUtilities::Tessellate(Representation, Context, MeshOut, true);
-				}
-			}
-
-			return bTessellationSuccessful;
-		}
-
 		bool Tessellate(FModel& Model, const UE::CADKernel::FTessellationContext& Context, UE::CADKernel::MeshUtilities::FMeshWrapperAbstract& MeshWrapper, bool bEmptyMesh = false);
 
 		// For future use
@@ -94,7 +33,6 @@ namespace UE::CADKernel
 		bool GetFaceTrimming3DPolylines(const UE::CADKernel::FModel& Model, const UE::CADKernel::FTopologicalFace& Face, TArray<TArray<FVector>>& PolylinesOut);
 
 		bool AddModelMesh(const UE::CADKernel::FModelMesh& ModelMesh, UE::CADKernel::MeshUtilities::FMeshWrapperAbstract& MeshWrapper);
-		bool ToRawData(TSharedPtr<UE::CADKernel::FModel>& Model, TArray<uint8>& RawDataOut);
 	}
 }
 #endif
diff --git a/Engine/Source/Runtime/Datasmith/CADKernel/Engine/Private/MeshUtilities.h b/Engine/Source/Runtime/Datasmith/CADKernel/Engine/Private/MeshUtilities.h
--- a/Engine/Source/Runtime/Datasmith/CADKernel/Engine/Private/MeshUtilities.h
+++ b/Engine/Source/Runtime/Datasmith/CADKernel/Engine/Private/MeshUtilities.h
@@ -4,7 +4,7 @@
 
 #include "CADKernelEngine.h"
 #if PLATFORM_DESKTOP
-#include "ParametricSurfaceData.h"
+#include "CADKernelEngineDefinitions.h"
 
 #include "IntVectorTypes.h"
 #include "StaticMeshAttributes.h"
diff --git a/Engine/Source/Runtime/Datasmith/CADKernel/Engine/Private/ParametricSurfaceData.cpp b/Engine/Source/Runtime/Datasmith/CADKernel/Engine/Private/ParametricSurfaceData.cpp
deleted file mode 100644
diff --git a/Engine/Source/Runtime/Datasmith/CADKernel/Engine/Private/TechSoft/TechSoftUtilities.cpp b/Engine/Source/Runtime/Datasmith/CADKernel/Engine/Private/TechSoft/TechSoftUtilities.cpp
--- a/Engine/Source/Runtime/Datasmith/CADKernel/Engine/Private/TechSoft/TechSoftUtilities.cpp
+++ b/Engine/Source/Runtime/Datasmith/CADKernel/Engine/Private/TechSoft/TechSoftUtilities.cpp
@@ -109,7 +109,7 @@ namespace UE::CADKernel
 		return true;
 	}
 
-	bool FTechSoftUtilities::Tessellate(A3DRiRepresentationItem* Representation, const FTessellationContext& Context, FMeshDescription& Mesh, bool bEmptyMesh)
+	bool FTechSoftUtilities::Tessellate(A3DRiRepresentationItem* Representation, const FTessellationContext& Context, FMeshDescription& Mesh, bool bSkipDeletedFaces, bool bEmptyMesh)
 	{
 		using namespace UE::CADKernel::MeshUtilities;
 
@@ -118,11 +118,16 @@ namespace UE::CADKernel
 			return false;
 		}
 
+		if (bSkipDeletedFaces)
+		{
+			GetExistingFaceGroups(Mesh, Context.FaceGroupsToExtract);
+		}
+
 		TSharedPtr<FMeshWrapperAbstract> MeshWrapper = FMeshWrapperAbstract::MakeWrapper(Context, Mesh);
 		return TechSoftUtilities::Tessellate(Representation, Context, *MeshWrapper, bEmptyMesh);
 	}
 
-	bool FTechSoftUtilities::Tessellate(A3DRiRepresentationItem* Representation, const FTessellationContext& Context, UE::Geometry::FDynamicMesh3& Mesh, bool bEmptyMesh)
+	bool FTechSoftUtilities::Tessellate(A3DRiRepresentationItem* Representation, const FTessellationContext& Context, UE::Geometry::FDynamicMesh3& Mesh, bool bSkipDeletedFaces, bool bEmptyMesh)
 	{
 		using namespace UE::CADKernel::MeshUtilities;
 
@@ -131,6 +136,11 @@ namespace UE::CADKernel
 			return false;
 		}
 
+		if (bSkipDeletedFaces)
+		{
+			GetExistingFaceGroups(Mesh, Context.FaceGroupsToExtract);
+		}
+
 		TSharedPtr<FMeshWrapperAbstract> MeshWrapper = FMeshWrapperAbstract::MakeWrapper(Context, Mesh);
 		return TechSoftUtilities::Tessellate(Representation, Context, *MeshWrapper, bEmptyMesh);
 	}
@@ -433,7 +443,7 @@ namespace UE::CADKernel
 		return {};
 	}
 
-	bool TechSoftUtilities::ToRawData(A3DRiRepresentationItem* Representation, int32 MaterialID, TArray<uint8>& RawDataOut)
+	bool FTechSoftUtilities::ToRawData(A3DRiRepresentationItem* Representation, int32 MaterialID, TArray<uint8>& RawDataOut)
 	{
 		// See TechSoftUtils::SaveBodiesToPrcFile
 		// Set attributes on representation not ProductOccurence
@@ -468,6 +478,11 @@ namespace UE::CADKernel
 	{
 		return {};
 	}
+
+	bool FTechSoftUtilities::ToRawData(A3DRiRepresentationItem* Representation, int32 MaterialID, TArray<uint8>& RawDataOut)
+	{
+		return false;
+	}
 }
 #endif
 #endif
diff --git a/Engine/Source/Runtime/Datasmith/CADKernel/Engine/Private/TechSoft/TechSoftUtilities.h b/Engine/Source/Runtime/Datasmith/CADKernel/Engine/Private/TechSoft/TechSoftUtilities.h
--- a/Engine/Source/Runtime/Datasmith/CADKernel/Engine/Private/TechSoft/TechSoftUtilities.h
+++ b/Engine/Source/Runtime/Datasmith/CADKernel/Engine/Private/TechSoft/TechSoftUtilities.h
@@ -27,7 +27,6 @@ namespace UE::CADKernel
 		bool AddRepresentation(A3DRiRepresentationItem* Representation, double ModelUnitToMeter, MeshUtilities::FMeshWrapperAbstract& MeshWrapper);
 		FVector2d GetUVScale(const A3DTopoFace* TopoFace, double TextureUnit);
 		TSharedPtr<FJsonObject> GetJsonObject(A3DEntity* Entity, bool bIsLegacy = false);
-		bool ToRawData(A3DRiRepresentationItem* Representation, int32 MaterialID, TArray<uint8>& RawDataOut);
 
 
 		/*
diff --git a/Engine/Source/Runtime/Datasmith/CADKernel/Engine/Public/CADKernelEngine.h b/Engine/Source/Runtime/Datasmith/CADKernel/Engine/Public/CADKernelEngine.h
--- a/Engine/Source/Runtime/Datasmith/CADKernel/Engine/Public/CADKernelEngine.h
+++ b/Engine/Source/Runtime/Datasmith/CADKernel/Engine/Public/CADKernelEngine.h
@@ -4,7 +4,7 @@
 #include "CoreMinimal.h"
 
 #if PLATFORM_DESKTOP
-#include "ParametricSurfaceData.h"
+#include "CADKernelEngineDefinitions.h"
 #include "TechSoftIncludes.h"
 
 namespace UE::CADKernel
@@ -28,41 +28,20 @@ struct FMeshDescription;
 
 namespace UE::CADKernel
 {
-	struct FTechSoftImportOverrides
-	{
-		bool bLoadGeometryOnly = true;
-		bool bLoadLoadNoDependency = false;
-		bool bLoadStructureOnly = false;
-	};
-
-	struct FMeshExtractionContext
-	{
-		FCADKernelModelParameters ModelParams;
-		FCADKernelMeshParameters MeshParams;
-		TSet<int32> FaceGroupsToExtract;
-		bool bResolveTJunctions = false;
-	};
-
-	struct FTessellationContext : public FMeshExtractionContext
-	{
-		FCADKernelTessellationSettings TessellationSettings;
-		FTessellationContext() {};
-		FTessellationContext(const FCADKernelModelParameters& ModelParams, const FCADKernelMeshParameters& MeshParams, const FCADKernelRetessellationSettings& Settings);
-	};
-
 	class FCADKernelUtilities
 	{
 	public:
 		static bool CADKERNELENGINE_API Save(TSharedPtr<UE::CADKernel::FModel>& Model, const FString& FilePath);
 		static bool CADKERNELENGINE_API Load(TSharedPtr<UE::CADKernel::FModel>& Model, const FString& FilePath);
-		static bool CADKERNELENGINE_API Tessellate(UE::CADKernel::FModel& Model, const FTessellationContext& Context, FMeshDescription& MeshOut, bool bEmptyMesh = false);
-		static bool CADKERNELENGINE_API Tessellate(UE::CADKernel::FModel& Model, const FTessellationContext& Context, UE::Geometry::FDynamicMesh3& MeshOut, bool bEmptyMesh = false);
+		static bool CADKERNELENGINE_API Tessellate(UE::CADKernel::FModel& Model, const FTessellationContext& Context, FMeshDescription& MeshOut, bool bSkipDeletedFaces = false, bool bEmptyMesh = false);
+		static bool CADKERNELENGINE_API Tessellate(UE::CADKernel::FModel& Model, const FTessellationContext& Context, UE::Geometry::FDynamicMesh3& MeshOut, bool bSkipDeletedFaces = false, bool bEmptyMesh = false);
 		
 		/* Note: Meshes are expected to be in their native coordinates system and unit */
 		static void CADKERNELENGINE_API ApplyExtractionContext(const FMeshExtractionContext& Context, FMeshDescription& MeshInOut);
 		static void CADKERNELENGINE_API ApplyExtractionContext(const FMeshExtractionContext& Context, UE::Geometry::FDynamicMesh3& MeshInOut);
 		// Set the proper attributes for FMeshDescription created from CAD data.
 		static void CADKERNELENGINE_API RegisterAttributes(FMeshDescription& MeshInOut, bool bKeepExistingAttribute = false);
+		static bool CADKERNELENGINE_API ToRawData(TSharedPtr<UE::CADKernel::FModel>& Model, TArray<uint8>& RawDataOut);
 	};
 
 	enum class ECADFormat
@@ -193,13 +172,14 @@ namespace UE::CADKernel
 		{
 			return Save(TArray<A3DRiRepresentationItem*>( &Representation, 1 ), FilePath, Attributes);
 		}
-		static CADKERNELENGINE_API bool Tessellate(A3DRiRepresentationItem* Representation, const FTessellationContext& Context, FMeshDescription& MeshOut, bool bEmptyMesh = false);
-		static CADKERNELENGINE_API bool Tessellate(A3DRiRepresentationItem* Representation, const FTessellationContext& Context, UE::Geometry::FDynamicMesh3& MeshOut, bool bEmptyMesh = false);
+		static CADKERNELENGINE_API bool Tessellate(A3DRiRepresentationItem* Representation, const FTessellationContext& Context, FMeshDescription& MeshOut, bool bSkipDeletedFaces = false, bool bEmptyMesh = false);
+		static CADKERNELENGINE_API bool Tessellate(A3DRiRepresentationItem* Representation, const FTessellationContext& Context, UE::Geometry::FDynamicMesh3& MeshOut, bool bSkipDeletedFaces = false, bool bEmptyMesh = false);
 
 		static CADKERNELENGINE_API TSharedPtr<UE::CADKernel::FModel> TechSoftToCADKernel(A3DRiRepresentationItem* Representation, double Unit, double InGeometricTolerance);
 		/* Limitation: Non of the implicit geometry is supported */
 		static CADKERNELENGINE_API A3DRiRepresentationItem* CADKernelToTechSoft(TSharedPtr<UE::CADKernel::FModel>& Model);
 		static CADKERNELENGINE_API A3DRiRepresentationItem* GetRepresentation(const TArray<uint8>& TechSoftRawData);
+		static CADKERNELENGINE_API bool ToRawData(A3DRiRepresentationItem* Representation, int32 MaterialID, TArray<uint8>& RawDataOut);
 
 	};
 
diff --git a/Engine/Source/Runtime/Datasmith/CADKernel/Engine/Public/CADKernelEngineDefinitions.h b/Engine/Source/Runtime/Datasmith/CADKernel/Engine/Public/CADKernelEngineDefinitions.h
new file mode 100644
--- /dev/null
+++ b/Engine/Source/Runtime/Datasmith/CADKernel/Engine/Public/CADKernelEngineDefinitions.h
@@ -0,0 +1,265 @@
+// Copyright Epic Games, Inc. All Rights Reserved.
+
+#pragma once
+
+#include "UObject/ObjectMacros.h"
+
+#include "CADKernelEngineDefinitions.generated.h"
+
+UENUM()
+enum class ECADKernelModelCoordSystem : uint8
+{
+	ZUp_LeftHanded,
+	ZUp_RightHanded,
+	YUp_LeftHanded,
+	YUp_RightHanded,
+	ZUp_RightHanded_FBXLegacy,
+};
+
+UENUM()
+enum class ECADKernelStitchingTechnique : uint8
+{
+	StitchingNone = 0,
+	StitchingHeal,
+	StitchingSew,
+};
+
+USTRUCT(BlueprintType)
+struct CADKERNELENGINE_API FCADKernelModelParameters
+{
+	GENERATED_BODY()
+
+	// value from FDatasmithUtils::ECADKernelModelCoordSystem
+	UPROPERTY()
+	uint8 ModelCoordSys = uint8(ECADKernelModelCoordSystem::ZUp_LeftHanded);
+
+	UPROPERTY()
+	float ModelUnitToCentimeter = 1.; // Model is in centimeters
+
+	UPROPERTY()
+	float ScaleFactor = 1.0f;
+	
+	bool FromJson(const FString& JsonString) const { return false; }
+	FString ToJson() const { return {}; }
+};
+
+USTRUCT(BlueprintType)
+struct CADKERNELENGINE_API FCADKernelMeshParameters
+{
+	GENERATED_BODY()
+
+	UPROPERTY()
+	bool bNeedSwapOrientation = false;
+
+	UPROPERTY()
+	bool bIsSymmetric = false;
+
+	UPROPERTY()
+	FVector SymmetricOrigin = FVector::ZeroVector;
+
+	UPROPERTY()
+	FVector SymmetricNormal = FVector::ZeroVector;
+
+	UPROPERTY()
+	float ScaleFactor = 1.0f;
+
+	bool FromJson(const FString& JsonString) const { return false; }
+	FString ToJson() const { return {}; }
+};
+
+USTRUCT(BlueprintType)
+struct CADKERNELENGINE_API FCADKernelTessellationSettings
+{
+	GENERATED_BODY()
+
+	FCADKernelTessellationSettings(
+			double InChordTolerance = 0.2,
+			double InMaxEdgeLength = 0.0,
+			double InNormalTolerance = 20.0,
+			ECADKernelStitchingTechnique InStitchingTechnique = ECADKernelStitchingTechnique::StitchingSew,
+			double InGeometricTolerance = 0.001,
+			double InStichingTolerance = 0.001,
+			bool bInUseCADKernel = true
+		)
+		: ChordTolerance(InChordTolerance)
+		, MaxEdgeLength(InMaxEdgeLength)
+		, NormalTolerance(InNormalTolerance)
+		, StitchingTechnique(InStitchingTechnique)
+		, bUseCADKernel(bInUseCADKernel)
+		, GeometricTolerance(InGeometricTolerance)
+		, StitchingTolerance(InStichingTolerance)
+	{
+	}
+
+protected:
+	/**
+	 * Maximum distance between any point on a triangle generated by the tessellation process and the actual surface.
+	 * The lower the value the more triangles.
+	 * Default value is 0.2, minimal value is 0.005 cm.
+	 */
+	UPROPERTY(config, EditAnywhere, BlueprintReadWrite, Category = "Geometry & Tessellation Options", meta = (Units = cm, ToolTip = "Maximum distance between any generated triangle and the original surface. Smaller values make more triangles.", ClampMin = "0.005"))
+	double ChordTolerance;
+
+	/**
+	 * Maximum length of edges of triangles generated by the tessellation process.
+	 * The length is in scene/model unit. The smaller the more triangles are generated.
+	 * Value of 0 means no constraint on length of edges
+	 * Default value is 0 to disable this criteria, and 1. cm is its minimal value if enable.
+	 */
+	UPROPERTY(config, EditAnywhere, BlueprintReadWrite, AdvancedDisplay, Category = "Geometry & Tessellation Options", meta = (Units = cm, DisplayName = "Max Edge Length", ToolTip = "Maximum length of any edge in the generated triangles. Smaller values make more triangles.", ClampMin = "0.0"))
+	double MaxEdgeLength;
+
+public:
+	/**
+	 * Maximum angle between the normal of two triangles generated by the tessellation process.
+	 * The angle is expressed in degree. The smaller the more triangles are generated.
+	 * Default value is 20 degrees, min value is 5 degrees.
+	 */
+	UPROPERTY(config, EditAnywhere, BlueprintReadWrite, AdvancedDisplay, Category = "Geometry & Tessellation Options", meta = (Units = deg, ToolTip = "Maximum angle between adjacent triangles. Smaller values make more triangles.", ClampMin = "5.0", ClampMax = "90.0"))
+	double NormalTolerance;
+
+
+	/**
+	 * Stitching technique applied on neighboring surfaces before tessellation.
+	 * None : No stitching applied. This is the default.
+	 * Sewing : Connects surfaces which physically share a boundary but not topologically within a set of objects.
+	 *          This technique can modify the structure of the model by removing and adding objects.
+	 * Healing : Connects surfaces which physically share a boundary but not topologically within an object.
+	 */
+	UPROPERTY(config, EditAnywhere, BlueprintReadWrite, AdvancedDisplay, Category = "Geometry & Tessellation Options", meta = (ToolTip = "Stitching technique applied on model before tessellation. Sewing could impact number of objects."))
+	ECADKernelStitchingTechnique StitchingTechnique;
+
+	/**
+	 * Indicates which tessellator to use CADKernel or TechSoft. By default it is CADKernel.
+	 */
+	UPROPERTY(config, EditAnywhere, BlueprintReadWrite, AdvancedDisplay, Category = "Geometry & Tessellation Options", meta = (DisplayName = "Use CADKernel", Tooltip = "When true uses CADKernel tessellator. Otherwise, TechSoft's tessellator is used. Default is true.", editCondition = "bWithTechSoft == true"))
+	bool bUseCADKernel = true;
+
+	/**
+	 * Indicates which tessellator to use CADKernel or TechSoft. By default it is CADKernel.
+	 */
+	UPROPERTY(config, EditAnywhere, BlueprintReadWrite, AdvancedDisplay, Category = "Geometry & Tessellation Options", meta = (Tooltip = "If true, enables the removal of T-Junctions from the resulting mesh. Default is false."))
+	bool bResolveTJunctions = false;
+
+protected:
+	/**
+	 * Tolerance used to determine if a surface should be tessellate or not.
+	 * Any surface which is narrower than the geometric tolerance
+	 * in one of the iso direction will not be tessellated
+	 * Value is in centimeter
+	 */
+	UPROPERTY(config, EditAnywhere, BlueprintReadWrite, AdvancedDisplay, Category = "Geometry & Tessellation Options", meta = (Units = cm, ToolTip = "Tolerance used to determine if a surface should be tessellated or not."))
+	double GeometricTolerance = 0.001;
+
+	/**
+	 * Tolerance used to determine if two surfaces should be stitched.
+	 * Value is in centimeter
+	 */
+	UPROPERTY(config, EditAnywhere, BlueprintReadWrite, AdvancedDisplay, Category = "Geometry & Tessellation Options", meta = (Units = cm, ToolTip = "Tolerance used to determine if two surfaces should be stitched.", editCondition = "StitchingTechnique!=ECADKernelStitchingTechnique::StitchingNone"))
+	double StitchingTolerance = 0.001;
+
+public:
+	/*
+	 ** Value to convert imported geometrical values to cm which is UE's unit.
+	 ** Default value is 0.1 as CADKernel's and most CAD formats' unit is millimeter.
+	 */
+	// #cadkernel_check: Verify that this new logic is properly used everywhere.
+	UPROPERTY()
+	double UnitMultiplier = 0.1;
+
+	UPROPERTY(Transient)
+	bool bWithTechSoft = bWithHoops;
+
+	static bool bWithHoops;
+
+	bool operator == (const FCADKernelTessellationSettings& Other) const
+	{
+		return FMath::IsNearlyEqual(ChordTolerance, Other.ChordTolerance)
+			&& FMath::IsNearlyEqual(MaxEdgeLength, Other.MaxEdgeLength)
+			&& FMath::IsNearlyEqual(NormalTolerance, Other.NormalTolerance)
+			&& StitchingTechnique == Other.StitchingTechnique
+			&& FMath::IsNearlyEqual(GeometricTolerance, Other.GeometricTolerance)
+			&& FMath::IsNearlyEqual(StitchingTolerance, Other.StitchingTolerance);
+	}
+
+	uint32 GetHash() const
+	{
+		uint32 Hash = uint32(StitchingTechnique);
+		for (double Param : {ChordTolerance, MaxEdgeLength, NormalTolerance, GeometricTolerance, StitchingTolerance})
+		{
+			Hash = HashCombine(Hash, GetTypeHash(Param));
+		}
+		return Hash;
+	}
+	/** Helper functions to get geometrical values in the right unit, cm (native) or mm */
+	double GetGeometricTolerance(bool bInMillimeter = false) const { return bInMillimeter ? GeometricTolerance * 10. : GeometricTolerance; }
+	double GetStitchingTolerance(bool bInMillimeter = false) const { return bInMillimeter ? StitchingTolerance * 10. : StitchingTolerance; }
+	double GetChordTolerance(bool bInMillimeter = false) const { return bInMillimeter ? ChordTolerance * 10. : ChordTolerance; }
+	double GetMaxEdgeLength(bool bInMillimeter = false) const { return bInMillimeter ? MaxEdgeLength * 10. : MaxEdgeLength; }
+
+	bool FromJson(const FString& JsonString) const { return false; }
+	FString ToJson() const { return {}; }
+};
+
+UENUM()
+enum class ECADKernelRetessellationRule : uint8
+{
+	AllFaces = 0,
+	SkipDeletedFaces,
+};
+
+USTRUCT(BlueprintType)
+struct CADKERNELENGINE_API FCADKernelRetessellationSettings : public FCADKernelTessellationSettings
+{
+	GENERATED_BODY()
+
+	FCADKernelRetessellationSettings(double InChordTolerance = 0.2f, double InMaxEdgeLength = 0.0f, double InNormalTolerance = 20.0f, ECADKernelStitchingTechnique InStitchingTechnique = ECADKernelStitchingTechnique::StitchingSew, ECADKernelRetessellationRule InRetessellationRule = ECADKernelRetessellationRule::AllFaces)
+		: FCADKernelTessellationSettings(InChordTolerance, InMaxEdgeLength, InNormalTolerance, InStitchingTechnique)
+		, RetessellationRule(InRetessellationRule)
+	{
+	}
+
+	FCADKernelRetessellationSettings(const FCADKernelTessellationSettings& TessellationSettings, ECADKernelRetessellationRule InRetessellationRule = ECADKernelRetessellationRule::AllFaces)
+		: FCADKernelTessellationSettings(TessellationSettings)
+		, RetessellationRule(InRetessellationRule)
+	{
+	}
+
+	UPROPERTY(config, EditAnywhere, BlueprintReadWrite, AdvancedDisplay, Category = "Retessellation Settings", meta = (ToolTip = "Regenerate deleted surfaces during retesselate or ignore them"))
+	ECADKernelRetessellationRule RetessellationRule = ECADKernelRetessellationRule::AllFaces;
+};
+
+namespace UE::CADKernel
+{
+	struct FUnitConverter
+	{
+		static const double CADKERNELENGINE_API CentimeterToMillimeter;
+		static const double CADKERNELENGINE_API MillimeterToCentimeter;
+		static const double CADKERNELENGINE_API CentimeterToMeter;
+		static const double CADKERNELENGINE_API MeterToCentimeter;
+		static const double CADKERNELENGINE_API MillimeterToMeter;
+		static const double CADKERNELENGINE_API MeterToMillimeter;
+	};
+
+	struct FTechSoftImportOverrides
+	{
+		bool bLoadGeometryOnly = true;
+		bool bLoadLoadNoDependency = false;
+		bool bLoadStructureOnly = false;
+	};
+
+	struct FMeshExtractionContext
+	{
+		FCADKernelModelParameters ModelParams;
+		FCADKernelMeshParameters MeshParams;
+		mutable TSet<int32> FaceGroupsToExtract;
+		bool bResolveTJunctions = false;
+	};
+
+	struct FTessellationContext : public FMeshExtractionContext
+	{
+		FCADKernelTessellationSettings TessellationSettings;
+		CADKERNELENGINE_API FTessellationContext() {};
+		CADKERNELENGINE_API FTessellationContext(const FCADKernelModelParameters& ModelParams, const FCADKernelMeshParameters& MeshParams, const FCADKernelRetessellationSettings& Settings);
+	};
+}
diff --git a/Engine/Source/Runtime/Datasmith/CADKernel/Engine/Public/ParametricSurfaceData.h b/Engine/Source/Runtime/Datasmith/CADKernel/Engine/Public/ParametricSurfaceData.h
deleted file mode 100644
