[Backout] - CL44531039
#rnx
Original CL Desc
-----------------------------------------------------------------
[In Progress] InterchangeCADImporter
- Now UParametricSurfaceData derived from UAssetUserData so it can be embedded into a UStaticMesh object
- Added CADKernelEditor module to remove Engine dependency required as UParametricSurfaceData derives from  UAssetUserData
- Added code to transition out of UDatasmithParametricSurfaceData class so legacy DatasmithCAD imported static meshes can be tessellated
- Updated wire parser to support OpenModel 2025.0
- Enabled TechSoft parser
- Fixed empty FMeshDescription for OpenNurbs parser.

#rnx

diff --git a/Engine/Source/Runtime/Datasmith/CADKernel/Editor/CADKernelEditor.Build.cs b/Engine/Source/Runtime/Datasmith/CADKernel/Editor/CADKernelEditor.Build.cs
deleted file mode 100644
diff --git a/Engine/Source/Runtime/Datasmith/CADKernel/Editor/Private/CADKernelRetessellateAction.cpp b/Engine/Source/Runtime/Datasmith/CADKernel/Editor/Private/CADKernelRetessellateAction.cpp
deleted file mode 100644
diff --git a/Engine/Source/Runtime/Datasmith/CADKernel/Editor/Private/LegacyParametricSurfaceData.cpp b/Engine/Source/Runtime/Datasmith/CADKernel/Editor/Private/LegacyParametricSurfaceData.cpp
deleted file mode 100644
diff --git a/Engine/Source/Runtime/Datasmith/CADKernel/Editor/Private/LegacyParametricSurfaceData.h b/Engine/Source/Runtime/Datasmith/CADKernel/Editor/Private/LegacyParametricSurfaceData.h
deleted file mode 100644
diff --git a/Engine/Source/Runtime/Datasmith/CADKernel/Editor/Private/ParametricSurfaceData.cpp b/Engine/Source/Runtime/Datasmith/CADKernel/Editor/Private/ParametricSurfaceData.cpp
deleted file mode 100644
diff --git a/Engine/Source/Runtime/Datasmith/CADKernel/Editor/Public/CADKernelRetessellateAction.h b/Engine/Source/Runtime/Datasmith/CADKernel/Editor/Public/CADKernelRetessellateAction.h
deleted file mode 100644
diff --git a/Engine/Source/Runtime/Datasmith/CADKernel/Editor/Public/ParametricSurfaceData.h b/Engine/Source/Runtime/Datasmith/CADKernel/Editor/Public/ParametricSurfaceData.h
deleted file mode 100644
diff --git a/Engine/Source/Runtime/Datasmith/CADKernel/Engine/CADKernelEngine.Build.cs b/Engine/Source/Runtime/Datasmith/CADKernel/Engine/CADKernelEngine.Build.cs
--- a/Engine/Source/Runtime/Datasmith/CADKernel/Engine/CADKernelEngine.Build.cs
+++ b/Engine/Source/Runtime/Datasmith/CADKernel/Engine/CADKernelEngine.Build.cs
@@ -8,9 +8,8 @@ namespace UnrealBuildTool.Rules
 		public CADKernelEngine(ReadOnlyTargetRules Target)
 			: base(Target)
 		{
-			bEnableExceptions = true;
 			bUseUnity = false; // TechSoftLibrary.cpp cannot be compiled with the other cpp files.
-			OptimizeCode = CodeOptimization.Never;
+			//OptimizeCode = CodeOptimization.Never;
 
 			PublicDependencyModuleNames.AddRange(
 				new string[]
diff --git a/Engine/Source/Runtime/Datasmith/CADKernel/Engine/Private/CADKernelEngine.cpp b/Engine/Source/Runtime/Datasmith/CADKernel/Engine/Private/CADKernelEngine.cpp
--- a/Engine/Source/Runtime/Datasmith/CADKernel/Engine/Private/CADKernelEngine.cpp
+++ b/Engine/Source/Runtime/Datasmith/CADKernel/Engine/Private/CADKernelEngine.cpp
@@ -18,22 +18,6 @@
 
 DEFINE_LOG_CATEGORY(LogCADKernelEngine);
 
-#ifdef WITH_HOOPS
-bool FCADKernelTessellationSettings::bWithHoops = true;
-#else
-bool FCADKernelTessellationSettings::bWithHoops = false;
-#endif
-
-namespace UE::CADKernel
-{
-	const double FUnitConverter::CentimeterToMillimeter = 10.;
-	const double FUnitConverter::MillimeterToCentimeter = 0.1;
-	const double FUnitConverter::CentimeterToMeter = 0.01;
-	const double FUnitConverter::MeterToCentimeter = 100.;
-	const double FUnitConverter::MillimeterToMeter = 0.001;
-	const double FUnitConverter::MeterToMillimeter = 1000.;
-}
-
 namespace UE::CADKernel
 {
 	using namespace UE::CADKernel::MeshUtilities;
@@ -74,35 +58,20 @@ namespace UE::CADKernel
 		return false;
 	}
 
-	bool FCADKernelUtilities::Tessellate(UE::CADKernel::FModel& Model, const FTessellationContext& Context, FMeshDescription& Mesh, bool bSkipDeletedFaces, bool bEmptyMesh)
+	bool FCADKernelUtilities::Tessellate(UE::CADKernel::FModel& Model, const FTessellationContext& Context, FMeshDescription& Mesh, bool bEmptyMesh)
 	{
 		using namespace UE::CADKernel::MeshUtilities;
 
-		if (bSkipDeletedFaces)
-		{
-			GetExistingFaceGroups(Mesh, Context.FaceGroupsToExtract);
-		}
-
 		TSharedPtr<FMeshWrapperAbstract> MeshWrapper = FMeshWrapperAbstract::MakeWrapper(Context, Mesh);
 		return Private::Tessellate(Model, Context, *MeshWrapper, bEmptyMesh);
 	}
 
-	bool FCADKernelUtilities::Tessellate(UE::CADKernel::FModel& Model, const FTessellationContext& Context, UE::Geometry::FDynamicMesh3& Mesh, bool bSkipDeletedFaces, bool bEmptyMesh)
+	bool FCADKernelUtilities::Tessellate(UE::CADKernel::FModel& Model, const FTessellationContext& Context, UE::Geometry::FDynamicMesh3& Mesh, bool bEmptyMesh)
 	{
 		using namespace UE::CADKernel::MeshUtilities;
 
-		if (bSkipDeletedFaces)
-		{
-			GetExistingFaceGroups(Mesh, Context.FaceGroupsToExtract);
-		}
-
 		TSharedPtr<FMeshWrapperAbstract> MeshWrapper = FMeshWrapperAbstract::MakeWrapper(Context, Mesh);
 		return Private::Tessellate(Model, Context, *MeshWrapper, bEmptyMesh);
 	}
-
-	bool FCADKernelUtilities::ToRawData(TSharedPtr<UE::CADKernel::FModel>& Model, TArray<uint8>& RawDataOut)
-	{
-		return false;
-	}
 }
 #endif
diff --git a/Engine/Source/Runtime/Datasmith/CADKernel/Engine/Private/CADKernelEnginePrivate.cpp b/Engine/Source/Runtime/Datasmith/CADKernel/Engine/Private/CADKernelEnginePrivate.cpp
--- a/Engine/Source/Runtime/Datasmith/CADKernel/Engine/Private/CADKernelEnginePrivate.cpp
+++ b/Engine/Source/Runtime/Datasmith/CADKernel/Engine/Private/CADKernelEnginePrivate.cpp
@@ -83,6 +83,11 @@ namespace UE::CADKernel::Private
 		return AddModelMesh(*ModelMesh, MeshWrapper);
 	}
 
+	bool ToRawData(TSharedPtr<UE::CADKernel::FModel>& Model, TArray<uint8>& RawDataOut)
+	{
+		return false;
+	}
+
 	bool GetFaceTrimmingCurves(const FModel& Model, const FTopologicalFace& Face, TArray<TArray<TArray<FVector>>>& CurvesOut)
 	{
 		using namespace UE::CADKernel;
diff --git a/Engine/Source/Runtime/Datasmith/CADKernel/Engine/Private/CADKernelEnginePrivate.h b/Engine/Source/Runtime/Datasmith/CADKernel/Engine/Private/CADKernelEnginePrivate.h
--- a/Engine/Source/Runtime/Datasmith/CADKernel/Engine/Private/CADKernelEnginePrivate.h
+++ b/Engine/Source/Runtime/Datasmith/CADKernel/Engine/Private/CADKernelEnginePrivate.h
@@ -25,6 +25,67 @@ namespace UE::CADKernel
 
 	namespace Private
 	{
+
+		template<typename MeshType>
+		bool Tessellate(UParametricSurfaceData& Data, MeshType& MeshOut)
+		{
+			using namespace UE::CADKernel;
+
+			bool bTessellationSuccessful = false;
+
+			FTessellationContext Context(Data.GetModelParameters(), Data.GetMeshParameters(), Data.GetLastTessellationSettings());
+
+			if (Context.TessellationSettings.bUseCADKernel)
+			{
+				TSharedPtr<FModel> Model = Data.GetModel();
+				if (Model.IsValid())
+				{
+					bTessellationSuccessful = FCADKernelUtilities::Tessellate(*Model, Context, MeshOut);
+				}
+			}
+			else if (FTechSoftLibrary::Initialize())
+			{
+				A3DRiRepresentationItem* Representation = Data.GetRepresentation();
+				bTessellationSuccessful = FTechSoftUtilities::Tessellate(Representation, Context, MeshOut);
+			}
+
+			return bTessellationSuccessful;
+		}
+
+		template<typename MeshType>
+		bool Retessellate(UParametricSurfaceData& Data, const FCADKernelRetessellationSettings& Settings, MeshType& MeshOut)
+		{
+			using namespace UE::CADKernel;
+			using namespace UE::CADKernel::MeshUtilities;
+
+			bool bTessellationSuccessful = false;
+
+			FTessellationContext Context(Data.GetModelParameters(), Data.GetMeshParameters(), Settings);
+
+			if (Settings.RetessellationRule == ECADKernelRetessellationRule::SkipDeletedFaces)
+			{
+				GetExistingFaceGroups(MeshOut, Context.FaceGroupsToExtract);
+			}
+
+			if (Context.TessellationSettings.bUseCADKernel)
+			{
+				TSharedPtr<FModel> Model = Data.GetModel();
+				if (Model.IsValid())
+				{
+					bTessellationSuccessful = FCADKernelUtilities::Tessellate(*Model, Context, MeshOut, true);
+				}
+			}
+			else if(FTechSoftLibrary::Initialize())
+			{
+				if (A3DRiRepresentationItem* Representation = Data.GetRepresentation())
+				{
+					bTessellationSuccessful = FTechSoftUtilities::Tessellate(Representation, Context, MeshOut, true);
+				}
+			}
+
+			return bTessellationSuccessful;
+		}
+
 		bool Tessellate(FModel& Model, const UE::CADKernel::FTessellationContext& Context, UE::CADKernel::MeshUtilities::FMeshWrapperAbstract& MeshWrapper, bool bEmptyMesh = false);
 
 		// For future use
@@ -33,6 +94,7 @@ namespace UE::CADKernel
 		bool GetFaceTrimming3DPolylines(const UE::CADKernel::FModel& Model, const UE::CADKernel::FTopologicalFace& Face, TArray<TArray<FVector>>& PolylinesOut);
 
 		bool AddModelMesh(const UE::CADKernel::FModelMesh& ModelMesh, UE::CADKernel::MeshUtilities::FMeshWrapperAbstract& MeshWrapper);
+		bool ToRawData(TSharedPtr<UE::CADKernel::FModel>& Model, TArray<uint8>& RawDataOut);
 	}
 }
 #endif
diff --git a/Engine/Source/Runtime/Datasmith/CADKernel/Engine/Private/MeshUtilities.h b/Engine/Source/Runtime/Datasmith/CADKernel/Engine/Private/MeshUtilities.h
--- a/Engine/Source/Runtime/Datasmith/CADKernel/Engine/Private/MeshUtilities.h
+++ b/Engine/Source/Runtime/Datasmith/CADKernel/Engine/Private/MeshUtilities.h
@@ -4,7 +4,7 @@
 
 #include "CADKernelEngine.h"
 #if PLATFORM_DESKTOP
-#include "CADKernelEngineDefinitions.h"
+#include "ParametricSurfaceData.h"
 
 #include "IntVectorTypes.h"
 #include "StaticMeshAttributes.h"
diff --git a/Engine/Source/Runtime/Datasmith/CADKernel/Engine/Private/ParametricSurfaceData.cpp b/Engine/Source/Runtime/Datasmith/CADKernel/Engine/Private/ParametricSurfaceData.cpp
new file mode 100644
--- /dev/null
+++ b/Engine/Source/Runtime/Datasmith/CADKernel/Engine/Private/ParametricSurfaceData.cpp
@@ -0,0 +1,156 @@
+// Copyright Epic Games, Inc. All Rights Reserved.
+
+#include "ParametricSurfaceData.h"
+
+#include "CADKernelEnginePrivate.h"
+#include "MeshUtilities.h"
+#include "TechSoft/TechSoftUtilities.h"
+
+#include "Core/Session.h"
+#include "Topo/Model.h"
+
+#include "Misc/FileHelper.h"
+#include "Misc/Paths.h"
+#include "UObject/EnterpriseObjectVersion.h"
+
+#ifdef WITH_HOOPS
+bool FCADKernelTessellationSettings::bWithHoops = true;
+#else
+bool FCADKernelTessellationSettings::bWithHoops = false;
+#endif
+
+#define LOCTEXT_NAMESPACE "ParametricSurfaceData"
+
+const double FUnitConverter::CentimeterToMillimeter = 10.;
+const double FUnitConverter::MillimeterToCentimeter = 0.1;
+const double FUnitConverter::CentimeterToMeter = 0.01;
+const double FUnitConverter::MeterToCentimeter = 100.;
+const double FUnitConverter::MillimeterToMeter = 0.001;
+const double FUnitConverter::MeterToMillimeter = 1000.;
+
+bool UParametricSurfaceData::SetFromFile(const TCHAR* FilePath, bool bForTechSoft)
+{
+	if (FPaths::FileExists(FilePath))
+	{
+		TArray<uint8> ByteArray;
+
+		if (FFileHelper::LoadFileToArray(ByteArray, FilePath))
+		{
+			TArray<uint8>& RawData = bForTechSoft ? TechSoftRawData : CADKernelRawData;
+			RawData = MoveTemp(ByteArray);
+			return true;
+		}
+	}
+
+	return false;
+}
+
+void UParametricSurfaceData::Serialize(FArchive& Ar)
+{
+	Ar.UsingCustomVersion(FEnterpriseObjectVersion::GUID);
+
+	if (Ar.IsLoading() && Ar.CustomVer(FEnterpriseObjectVersion::GUID) < FEnterpriseObjectVersion::AddedParametricSurfaceData)
+	{
+		// #cad_import_todo : Put code to deserialize UDatasmithParametricSurfaceData
+		 return;
+	}
+
+	Super::Serialize(Ar);
+
+	Ar << CADKernelRawData;
+	Ar << TechSoftRawData;
+}
+
+TSharedPtr<UE::CADKernel::FModel> UParametricSurfaceData::GetModel()
+{
+#if PLATFORM_DESKTOP
+	using namespace UE::CADKernel;
+
+	if (CADKernelRawData.IsEmpty())
+	{
+		return {};
+	}
+
+	TSharedRef<FSession> Session = MakeShared<FSession>(LastTessellationSettings.GetGeometricTolerance(true));
+	Session->AddDatabase(CADKernelRawData);
+
+	return Session->GetModelAsShared();
+#else
+	return {};
+#endif
+}
+
+A3DRiRepresentationItem* UParametricSurfaceData::GetRepresentation()
+{
+#ifdef WITH_HOOPS
+	using namespace UE::CADKernel;
+
+	if (TechSoftRawData.IsEmpty())
+	{
+		return {};
+	}
+
+	return FTechSoftUtilities::GetRepresentation(TechSoftRawData);
+#else
+	return nullptr;
+#endif
+}
+
+
+bool UParametricSurfaceData::Tessellate(UE::Geometry::FDynamicMesh3& MeshOut)
+{
+#if PLATFORM_DESKTOP
+	return UE::CADKernel::Private::Tessellate(*this, MeshOut);
+#else
+	return false;
+#endif
+}
+
+bool UParametricSurfaceData::Tessellate(FMeshDescription& MeshOut)
+{
+#if PLATFORM_DESKTOP
+	return UE::CADKernel::Private::Tessellate(*this, MeshOut);
+#else
+	return false;
+#endif
+}
+
+bool UParametricSurfaceData::Retessellate(const FCADKernelRetessellationSettings& Settings, UE::Geometry::FDynamicMesh3& MeshOut)
+{
+#if PLATFORM_DESKTOP
+	return UE::CADKernel::Private::Retessellate(*this, Settings, MeshOut);
+#else
+	return false;
+#endif
+}
+
+bool UParametricSurfaceData::Retessellate(const FCADKernelRetessellationSettings& Settings, FMeshDescription& MeshOut)
+{
+#if PLATFORM_DESKTOP
+	return UE::CADKernel::Private::Retessellate(*this, Settings, MeshOut);
+#else
+	return false;
+#endif
+}
+
+bool UParametricSurfaceData::SetModel(TSharedPtr<UE::CADKernel::FModel>& Model, double UnitModelToCentimeter)
+{
+#if PLATFORM_DESKTOP
+	ModelParameters.ModelUnitToCentimeter = UnitModelToCentimeter;
+	return UE::CADKernel::Private::ToRawData(Model, CADKernelRawData);
+#else
+	return false;
+#endif
+}
+
+bool UParametricSurfaceData::SetRepresentation(A3DRiRepresentationItem* Representation, int32 MaterialID, double UnitRepresentationToCentimeter)
+{
+#ifdef WITH_HOOPS
+	ModelParameters.ModelUnitToCentimeter = UnitRepresentationToCentimeter;
+	return UE::CADKernel::TechSoftUtilities::ToRawData(Representation, MaterialID, TechSoftRawData);
+#else
+	return false;
+#endif
+}
+
+#undef LOCTEXT_NAMESPACE // "ParametricSurfaceData"
diff --git a/Engine/Source/Runtime/Datasmith/CADKernel/Engine/Private/TechSoft/TechSoftUtilities.cpp b/Engine/Source/Runtime/Datasmith/CADKernel/Engine/Private/TechSoft/TechSoftUtilities.cpp
--- a/Engine/Source/Runtime/Datasmith/CADKernel/Engine/Private/TechSoft/TechSoftUtilities.cpp
+++ b/Engine/Source/Runtime/Datasmith/CADKernel/Engine/Private/TechSoft/TechSoftUtilities.cpp
@@ -109,7 +109,7 @@ namespace UE::CADKernel
 		return true;
 	}
 
-	bool FTechSoftUtilities::Tessellate(A3DRiRepresentationItem* Representation, const FTessellationContext& Context, FMeshDescription& Mesh, bool bSkipDeletedFaces, bool bEmptyMesh)
+	bool FTechSoftUtilities::Tessellate(A3DRiRepresentationItem* Representation, const FTessellationContext& Context, FMeshDescription& Mesh, bool bEmptyMesh)
 	{
 		using namespace UE::CADKernel::MeshUtilities;
 
@@ -118,16 +118,11 @@ namespace UE::CADKernel
 			return false;
 		}
 
-		if (bSkipDeletedFaces)
-		{
-			GetExistingFaceGroups(Mesh, Context.FaceGroupsToExtract);
-		}
-
 		TSharedPtr<FMeshWrapperAbstract> MeshWrapper = FMeshWrapperAbstract::MakeWrapper(Context, Mesh);
 		return TechSoftUtilities::Tessellate(Representation, Context, *MeshWrapper, bEmptyMesh);
 	}
 
-	bool FTechSoftUtilities::Tessellate(A3DRiRepresentationItem* Representation, const FTessellationContext& Context, UE::Geometry::FDynamicMesh3& Mesh, bool bSkipDeletedFaces, bool bEmptyMesh)
+	bool FTechSoftUtilities::Tessellate(A3DRiRepresentationItem* Representation, const FTessellationContext& Context, UE::Geometry::FDynamicMesh3& Mesh, bool bEmptyMesh)
 	{
 		using namespace UE::CADKernel::MeshUtilities;
 
@@ -136,11 +131,6 @@ namespace UE::CADKernel
 			return false;
 		}
 
-		if (bSkipDeletedFaces)
-		{
-			GetExistingFaceGroups(Mesh, Context.FaceGroupsToExtract);
-		}
-
 		TSharedPtr<FMeshWrapperAbstract> MeshWrapper = FMeshWrapperAbstract::MakeWrapper(Context, Mesh);
 		return TechSoftUtilities::Tessellate(Representation, Context, *MeshWrapper, bEmptyMesh);
 	}
@@ -443,7 +433,7 @@ namespace UE::CADKernel
 		return {};
 	}
 
-	bool FTechSoftUtilities::ToRawData(A3DRiRepresentationItem* Representation, int32 MaterialID, TArray<uint8>& RawDataOut)
+	bool TechSoftUtilities::ToRawData(A3DRiRepresentationItem* Representation, int32 MaterialID, TArray<uint8>& RawDataOut)
 	{
 		// See TechSoftUtils::SaveBodiesToPrcFile
 		// Set attributes on representation not ProductOccurence
@@ -478,11 +468,6 @@ namespace UE::CADKernel
 	{
 		return {};
 	}
-
-	bool FTechSoftUtilities::ToRawData(A3DRiRepresentationItem* Representation, int32 MaterialID, TArray<uint8>& RawDataOut)
-	{
-		return false;
-	}
 }
 #endif
 #endif
diff --git a/Engine/Source/Runtime/Datasmith/CADKernel/Engine/Private/TechSoft/TechSoftUtilities.h b/Engine/Source/Runtime/Datasmith/CADKernel/Engine/Private/TechSoft/TechSoftUtilities.h
--- a/Engine/Source/Runtime/Datasmith/CADKernel/Engine/Private/TechSoft/TechSoftUtilities.h
+++ b/Engine/Source/Runtime/Datasmith/CADKernel/Engine/Private/TechSoft/TechSoftUtilities.h
@@ -27,6 +27,7 @@ namespace UE::CADKernel
 		bool AddRepresentation(A3DRiRepresentationItem* Representation, double ModelUnitToMeter, MeshUtilities::FMeshWrapperAbstract& MeshWrapper);
 		FVector2d GetUVScale(const A3DTopoFace* TopoFace, double TextureUnit);
 		TSharedPtr<FJsonObject> GetJsonObject(A3DEntity* Entity, bool bIsLegacy = false);
+		bool ToRawData(A3DRiRepresentationItem* Representation, int32 MaterialID, TArray<uint8>& RawDataOut);
 
 
 		/*
diff --git a/Engine/Source/Runtime/Datasmith/CADKernel/Engine/Public/CADKernelEngine.h b/Engine/Source/Runtime/Datasmith/CADKernel/Engine/Public/CADKernelEngine.h
--- a/Engine/Source/Runtime/Datasmith/CADKernel/Engine/Public/CADKernelEngine.h
+++ b/Engine/Source/Runtime/Datasmith/CADKernel/Engine/Public/CADKernelEngine.h
@@ -4,7 +4,7 @@
 #include "CoreMinimal.h"
 
 #if PLATFORM_DESKTOP
-#include "CADKernelEngineDefinitions.h"
+#include "ParametricSurfaceData.h"
 #include "TechSoftIncludes.h"
 
 namespace UE::CADKernel
@@ -28,20 +28,41 @@ struct FMeshDescription;
 
 namespace UE::CADKernel
 {
+	struct FTechSoftImportOverrides
+	{
+		bool bLoadGeometryOnly = true;
+		bool bLoadLoadNoDependency = false;
+		bool bLoadStructureOnly = false;
+	};
+
+	struct FMeshExtractionContext
+	{
+		FCADKernelModelParameters ModelParams;
+		FCADKernelMeshParameters MeshParams;
+		TSet<int32> FaceGroupsToExtract;
+		bool bResolveTJunctions = false;
+	};
+
+	struct FTessellationContext : public FMeshExtractionContext
+	{
+		FCADKernelTessellationSettings TessellationSettings;
+		FTessellationContext() {};
+		FTessellationContext(const FCADKernelModelParameters& ModelParams, const FCADKernelMeshParameters& MeshParams, const FCADKernelRetessellationSettings& Settings);
+	};
+
 	class FCADKernelUtilities
 	{
 	public:
 		static bool CADKERNELENGINE_API Save(TSharedPtr<UE::CADKernel::FModel>& Model, const FString& FilePath);
 		static bool CADKERNELENGINE_API Load(TSharedPtr<UE::CADKernel::FModel>& Model, const FString& FilePath);
-		static bool CADKERNELENGINE_API Tessellate(UE::CADKernel::FModel& Model, const FTessellationContext& Context, FMeshDescription& MeshOut, bool bSkipDeletedFaces = false, bool bEmptyMesh = false);
-		static bool CADKERNELENGINE_API Tessellate(UE::CADKernel::FModel& Model, const FTessellationContext& Context, UE::Geometry::FDynamicMesh3& MeshOut, bool bSkipDeletedFaces = false, bool bEmptyMesh = false);
+		static bool CADKERNELENGINE_API Tessellate(UE::CADKernel::FModel& Model, const FTessellationContext& Context, FMeshDescription& MeshOut, bool bEmptyMesh = false);
+		static bool CADKERNELENGINE_API Tessellate(UE::CADKernel::FModel& Model, const FTessellationContext& Context, UE::Geometry::FDynamicMesh3& MeshOut, bool bEmptyMesh = false);
 		
 		/* Note: Meshes are expected to be in their native coordinates system and unit */
 		static void CADKERNELENGINE_API ApplyExtractionContext(const FMeshExtractionContext& Context, FMeshDescription& MeshInOut);
 		static void CADKERNELENGINE_API ApplyExtractionContext(const FMeshExtractionContext& Context, UE::Geometry::FDynamicMesh3& MeshInOut);
 		// Set the proper attributes for FMeshDescription created from CAD data.
 		static void CADKERNELENGINE_API RegisterAttributes(FMeshDescription& MeshInOut, bool bKeepExistingAttribute = false);
-		static bool CADKERNELENGINE_API ToRawData(TSharedPtr<UE::CADKernel::FModel>& Model, TArray<uint8>& RawDataOut);
 	};
 
 	enum class ECADFormat
@@ -172,14 +193,13 @@ namespace UE::CADKernel
 		{
 			return Save(TArray<A3DRiRepresentationItem*>( &Representation, 1 ), FilePath, Attributes);
 		}
-		static CADKERNELENGINE_API bool Tessellate(A3DRiRepresentationItem* Representation, const FTessellationContext& Context, FMeshDescription& MeshOut, bool bSkipDeletedFaces = false, bool bEmptyMesh = false);
-		static CADKERNELENGINE_API bool Tessellate(A3DRiRepresentationItem* Representation, const FTessellationContext& Context, UE::Geometry::FDynamicMesh3& MeshOut, bool bSkipDeletedFaces = false, bool bEmptyMesh = false);
+		static CADKERNELENGINE_API bool Tessellate(A3DRiRepresentationItem* Representation, const FTessellationContext& Context, FMeshDescription& MeshOut, bool bEmptyMesh = false);
+		static CADKERNELENGINE_API bool Tessellate(A3DRiRepresentationItem* Representation, const FTessellationContext& Context, UE::Geometry::FDynamicMesh3& MeshOut, bool bEmptyMesh = false);
 
 		static CADKERNELENGINE_API TSharedPtr<UE::CADKernel::FModel> TechSoftToCADKernel(A3DRiRepresentationItem* Representation, double Unit, double InGeometricTolerance);
 		/* Limitation: Non of the implicit geometry is supported */
 		static CADKERNELENGINE_API A3DRiRepresentationItem* CADKernelToTechSoft(TSharedPtr<UE::CADKernel::FModel>& Model);
 		static CADKERNELENGINE_API A3DRiRepresentationItem* GetRepresentation(const TArray<uint8>& TechSoftRawData);
-		static CADKERNELENGINE_API bool ToRawData(A3DRiRepresentationItem* Representation, int32 MaterialID, TArray<uint8>& RawDataOut);
 
 	};
 
diff --git a/Engine/Source/Runtime/Datasmith/CADKernel/Engine/Public/CADKernelEngineDefinitions.h b/Engine/Source/Runtime/Datasmith/CADKernel/Engine/Public/CADKernelEngineDefinitions.h
deleted file mode 100644
diff --git a/Engine/Source/Runtime/Datasmith/CADKernel/Engine/Public/ParametricSurfaceData.h b/Engine/Source/Runtime/Datasmith/CADKernel/Engine/Public/ParametricSurfaceData.h
new file mode 100644
--- /dev/null
+++ b/Engine/Source/Runtime/Datasmith/CADKernel/Engine/Public/ParametricSurfaceData.h
@@ -0,0 +1,323 @@
+// Copyright Epic Games, Inc. All Rights Reserved.
+
+#pragma once
+
+#include "CoreMinimal.h"
+
+#include "TechSoftIncludes.h"
+
+#include "ParametricSurfaceData.generated.h"
+
+namespace UE::CADKernel
+{
+	class FModel;
+}
+
+namespace UE::Geometry
+{
+	class FDynamicMesh3;
+}
+
+class UStaticMesh;
+struct FMeshDescription;
+
+typedef void A3DRiRepresentationItem;
+typedef void A3DAsmModelFile;
+
+struct FUnitConverter
+{
+	static const double CADKERNELENGINE_API CentimeterToMillimeter;
+	static const double CADKERNELENGINE_API MillimeterToCentimeter;
+	static const double CADKERNELENGINE_API CentimeterToMeter;
+	static const double CADKERNELENGINE_API MeterToCentimeter;
+	static const double CADKERNELENGINE_API MillimeterToMeter;
+	static const double CADKERNELENGINE_API MeterToMillimeter;
+};
+
+UENUM()
+enum class ECADKernelModelCoordSystem : uint8
+{
+	ZUp_LeftHanded,
+	ZUp_RightHanded,
+	YUp_LeftHanded,
+	YUp_RightHanded,
+	ZUp_RightHanded_FBXLegacy,
+};
+
+UENUM()
+enum class ECADKernelStitchingTechnique : uint8
+{
+	StitchingNone = 0,
+	StitchingHeal,
+	StitchingSew,
+};
+
+USTRUCT(BlueprintType)
+struct CADKERNELENGINE_API FCADKernelModelParameters
+{
+	GENERATED_BODY()
+
+	// value from FDatasmithUtils::ECADKernelModelCoordSystem
+	UPROPERTY()
+	uint8 ModelCoordSys = uint8(ECADKernelModelCoordSystem::ZUp_LeftHanded);
+
+	UPROPERTY()
+	float ModelUnitToCentimeter = 1.; // Model is in centimeters
+
+	UPROPERTY()
+	float ScaleFactor = 1.0f;
+	
+	bool FromJson(const FString& JsonString) const { return false; }
+	FString ToJson() const { return {}; }
+};
+
+USTRUCT(BlueprintType)
+struct CADKERNELENGINE_API FCADKernelMeshParameters
+{
+	GENERATED_BODY()
+
+	UPROPERTY()
+	bool bNeedSwapOrientation = false;
+
+	UPROPERTY()
+	bool bIsSymmetric = false;
+
+	UPROPERTY()
+	FVector SymmetricOrigin = FVector::ZeroVector;
+
+	UPROPERTY()
+	FVector SymmetricNormal = FVector::ZeroVector;
+
+	UPROPERTY()
+	float ScaleFactor = 1.0f;
+
+	bool FromJson(const FString& JsonString) const { return false; }
+	FString ToJson() const { return {}; }
+};
+
+USTRUCT(BlueprintType)
+struct CADKERNELENGINE_API FCADKernelTessellationSettings
+{
+	GENERATED_BODY()
+
+	FCADKernelTessellationSettings(
+			double InChordTolerance = 0.2,
+			double InMaxEdgeLength = 0.0,
+			double InNormalTolerance = 20.0,
+			ECADKernelStitchingTechnique InStitchingTechnique = ECADKernelStitchingTechnique::StitchingSew,
+			double InGeometricTolerance = 0.001,
+			double InStichingTolerance = 0.001,
+			bool bInUseCADKernel = true
+		)
+		: ChordTolerance(InChordTolerance)
+		, MaxEdgeLength(InMaxEdgeLength)
+		, NormalTolerance(InNormalTolerance)
+		, StitchingTechnique(InStitchingTechnique)
+		, bUseCADKernel(bInUseCADKernel)
+		, GeometricTolerance(InGeometricTolerance)
+		, StitchingTolerance(InStichingTolerance)
+	{
+	}
+
+protected:
+	/**
+	 * Maximum distance between any point on a triangle generated by the tessellation process and the actual surface.
+	 * The lower the value the more triangles.
+	 * Default value is 0.2, minimal value is 0.005 cm.
+	 */
+	UPROPERTY(config, EditAnywhere, BlueprintReadWrite, Category = "Geometry & Tessellation Options", meta = (Units = cm, ToolTip = "Maximum distance between any generated triangle and the original surface. Smaller values make more triangles.", ClampMin = "0.005"))
+	double ChordTolerance;
+
+	/**
+	 * Maximum length of edges of triangles generated by the tessellation process.
+	 * The length is in scene/model unit. The smaller the more triangles are generated.
+	 * Value of 0 means no constraint on length of edges
+	 * Default value is 0 to disable this criteria, and 1. cm is its minimal value if enable.
+	 */
+	UPROPERTY(config, EditAnywhere, BlueprintReadWrite, AdvancedDisplay, Category = "Geometry & Tessellation Options", meta = (Units = cm, DisplayName = "Max Edge Length", ToolTip = "Maximum length of any edge in the generated triangles. Smaller values make more triangles.", ClampMin = "0.0"))
+	double MaxEdgeLength;
+
+public:
+	/**
+	 * Maximum angle between the normal of two triangles generated by the tessellation process.
+	 * The angle is expressed in degree. The smaller the more triangles are generated.
+	 * Default value is 20 degrees, min value is 5 degrees.
+	 */
+	UPROPERTY(config, EditAnywhere, BlueprintReadWrite, AdvancedDisplay, Category = "Geometry & Tessellation Options", meta = (Units = deg, ToolTip = "Maximum angle between adjacent triangles. Smaller values make more triangles.", ClampMin = "5.0", ClampMax = "90.0"))
+	double NormalTolerance;
+
+
+	/**
+	 * Stitching technique applied on neighboring surfaces before tessellation.
+	 * None : No stitching applied. This is the default.
+	 * Sewing : Connects surfaces which physically share a boundary but not topologically within a set of objects.
+	 *          This technique can modify the structure of the model by removing and adding objects.
+	 * Healing : Connects surfaces which physically share a boundary but not topologically within an object.
+	 */
+	UPROPERTY(config, EditAnywhere, BlueprintReadWrite, AdvancedDisplay, Category = "Geometry & Tessellation Options", meta = (ToolTip = "Stitching technique applied on model before tessellation. Sewing could impact number of objects."))
+	ECADKernelStitchingTechnique StitchingTechnique;
+
+	/**
+	 * Indicates which tessellator to use CADKernel or TechSoft. By default it is CADKernel.
+	 */
+	UPROPERTY(config, EditAnywhere, BlueprintReadWrite, AdvancedDisplay, Category = "Geometry & Tessellation Options", meta = (DisplayName = "Use CADKernel", Tooltip = "When true uses CADKernel tessellator. Otherwise, TechSoft's tessellator is used. Default is true.", editCondition = "bWithTechSoft == true"))
+	bool bUseCADKernel = true;
+
+	/**
+	 * Indicates which tessellator to use CADKernel or TechSoft. By default it is CADKernel.
+	 */
+	UPROPERTY(config, EditAnywhere, BlueprintReadWrite, AdvancedDisplay, Category = "Geometry & Tessellation Options", meta = (Tooltip = "If true, enables the removal of T-Junctions from the resulting mesh. Default is false."))
+	bool bResolveTJunctions = false;
+
+protected:
+	/**
+	 * Tolerance used to determine if a surface should be tessellate or not.
+	 * Any surface which is narrower than the geometric tolerance
+	 * in one of the iso direction will not be tessellated
+	 * Value is in centimeter
+	 */
+	UPROPERTY(config, EditAnywhere, BlueprintReadWrite, AdvancedDisplay, Category = "Geometry & Tessellation Options", meta = (Units = cm, ToolTip = "Tolerance used to determine if a surface should be tessellated or not."))
+	double GeometricTolerance = 0.001;
+
+	/**
+	 * Tolerance used to determine if two surfaces should be stitched.
+	 * Value is in centimeter
+	 */
+	UPROPERTY(config, EditAnywhere, BlueprintReadWrite, AdvancedDisplay, Category = "Geometry & Tessellation Options", meta = (Units = cm, ToolTip = "Tolerance used to determine if two surfaces should be stitched.", editCondition = "StitchingTechnique!=ECADKernelStitchingTechnique::StitchingNone"))
+	double StitchingTolerance = 0.001;
+
+public:
+	/*
+	 ** Value to convert imported geometrical values to cm which is UE's unit.
+	 ** Default value is 0.1 as CADKernel's and most CAD formats' unit is millimeter.
+	 */
+	// #cadkernel_check: Verify that this new logic is properly used everywhere.
+	UPROPERTY()
+	double UnitMultiplier = 0.1;
+
+	UPROPERTY(Transient)
+	bool bWithTechSoft = bWithHoops;
+
+	static bool bWithHoops;
+
+	bool operator == (const FCADKernelTessellationSettings& Other) const
+	{
+		return FMath::IsNearlyEqual(ChordTolerance, Other.ChordTolerance)
+			&& FMath::IsNearlyEqual(MaxEdgeLength, Other.MaxEdgeLength)
+			&& FMath::IsNearlyEqual(NormalTolerance, Other.NormalTolerance)
+			&& StitchingTechnique == Other.StitchingTechnique
+			&& FMath::IsNearlyEqual(GeometricTolerance, Other.GeometricTolerance)
+			&& FMath::IsNearlyEqual(StitchingTolerance, Other.StitchingTolerance);
+	}
+
+	uint32 GetHash() const
+	{
+		uint32 Hash = uint32(StitchingTechnique);
+		for (double Param : {ChordTolerance, MaxEdgeLength, NormalTolerance, GeometricTolerance, StitchingTolerance})
+		{
+			Hash = HashCombine(Hash, GetTypeHash(Param));
+		}
+		return Hash;
+	}
+	/** Helper functions to get geometrical values in the right unit, cm (native) or mm */
+	double GetGeometricTolerance(bool bInMillimeter = false) const { return bInMillimeter ? GeometricTolerance * 10. : GeometricTolerance; }
+	double GetStitchingTolerance(bool bInMillimeter = false) const { return bInMillimeter ? StitchingTolerance * 10. : StitchingTolerance; }
+	double GetChordTolerance(bool bInMillimeter = false) const { return bInMillimeter ? ChordTolerance * 10. : ChordTolerance; }
+	double GetMaxEdgeLength(bool bInMillimeter = false) const { return bInMillimeter ? MaxEdgeLength * 10. : MaxEdgeLength; }
+
+	bool FromJson(const FString& JsonString) const { return false; }
+	FString ToJson() const { return {}; }
+};
+
+UENUM()
+enum class ECADKernelRetessellationRule : uint8
+{
+	AllFaces = 0,
+	SkipDeletedFaces,
+};
+
+USTRUCT(BlueprintType)
+struct CADKERNELENGINE_API FCADKernelRetessellationSettings : public FCADKernelTessellationSettings
+{
+	GENERATED_BODY()
+
+	FCADKernelRetessellationSettings(double InChordTolerance = 0.2f, double InMaxEdgeLength = 0.0f, double InNormalTolerance = 20.0f, ECADKernelStitchingTechnique InStitchingTechnique = ECADKernelStitchingTechnique::StitchingSew, ECADKernelRetessellationRule InRetessellationRule = ECADKernelRetessellationRule::AllFaces)
+		: FCADKernelTessellationSettings(InChordTolerance, InMaxEdgeLength, InNormalTolerance, InStitchingTechnique)
+		, RetessellationRule(InRetessellationRule)
+	{
+	}
+
+	FCADKernelRetessellationSettings(const FCADKernelTessellationSettings& TessellationSettings, ECADKernelRetessellationRule InRetessellationRule = ECADKernelRetessellationRule::AllFaces)
+		: FCADKernelTessellationSettings(TessellationSettings)
+		, RetessellationRule(InRetessellationRule)
+	{
+	}
+
+	UPROPERTY(config, EditAnywhere, BlueprintReadWrite, AdvancedDisplay, Category = "Retessellation Settings", meta = (ToolTip = "Regenerate deleted surfaces during retesselate or ignore them"))
+	ECADKernelRetessellationRule RetessellationRule = ECADKernelRetessellationRule::AllFaces;
+};
+
+UCLASS(meta = (DisplayName = "CADKernel Parametric Surface Data"))
+class CADKERNELENGINE_API UParametricSurfaceData : public UObject
+{
+	GENERATED_BODY()
+
+public:
+
+	virtual bool IsValid() { return CADKernelRawData.Num() > 0 || TechSoftRawData.Num() > 0; }
+
+	void SetModelParameters(const FCADKernelModelParameters& InModelParameters) { ModelParameters = InModelParameters; }
+	const FCADKernelModelParameters& GetModelParameters() const { return ModelParameters; }
+	FCADKernelModelParameters& GetModelParameters() { return ModelParameters; }
+
+	void SetMeshParameters(const FCADKernelMeshParameters& InMeshParameters) { MeshParameters = InMeshParameters; }
+	const FCADKernelMeshParameters& GetMeshParameters() const { return MeshParameters; }
+	FCADKernelMeshParameters& GetMeshParameters() { return MeshParameters; }
+
+	void SetLastTessellationSettings(const FCADKernelTessellationSettings& InTessellationSettings) { LastTessellationSettings = InTessellationSettings; }
+	const FCADKernelTessellationSettings& GetLastTessellationSettings() const { return LastTessellationSettings; }
+
+
+	virtual bool SetFromFile(const TCHAR* FilePath, bool bForTechSoft = false);
+	
+	bool SetModel(TSharedPtr<UE::CADKernel::FModel>& Model, double UnitModelToCentimeter = 1.);
+	bool SetRepresentation(A3DRiRepresentationItem* Representation, int32 MaterialID, double UnitRepresentationToCentimeter = 1.);
+
+	TSharedPtr<UE::CADKernel::FModel> GetModel();
+	A3DRiRepresentationItem* GetRepresentation();
+
+	bool Tessellate(UE::Geometry::FDynamicMesh3& MeshOut);
+	bool Tessellate(FMeshDescription& MeshOut);
+
+	bool Retessellate(const FCADKernelRetessellationSettings& Settings, UE::Geometry::FDynamicMesh3& MeshOut);
+	bool Retessellate(const FCADKernelRetessellationSettings& Settings, FMeshDescription& MeshOut);
+
+protected:
+	virtual void Serialize(FArchive& Ar) override;
+
+protected:
+	UPROPERTY()
+	FCADKernelModelParameters ModelParameters;
+
+	UPROPERTY()
+	FCADKernelMeshParameters MeshParameters;
+
+	UPROPERTY(EditAnywhere, Category = NURBS)
+	FCADKernelTessellationSettings LastTessellationSettings;
+
+	TArray<uint8> CADKernelRawData;
+	TArray<uint8> TechSoftRawData;
+
+private:
+	/*
+	 ** The SetRawData is only for internal use to facilitate the transition
+	 ** out of UDatasmithParametricSurfaceData into the new UParametricSurfaceData class
+	 */
+	virtual void SetRawData(const TArray<uint8>& InRawData, bool bForTechSoft = false)
+	{
+		TArray<uint8>& RawData = bForTechSoft ? TechSoftRawData : CADKernelRawData;
+		RawData = InRawData;
+	}
+	friend class UTechSoftParametricSurfaceData;
+	friend class UCADKernelParametricSurfaceData;
+};
