Fixed Interchange CAD import issue

#rnx

diff --git a/Engine/Source/Runtime/Datasmith/CADKernel/Engine/Private/MeshDescriptionWrapper.cpp b/Engine/Source/Runtime/Datasmith/CADKernel/Engine/Private/MeshDescriptionWrapper.cpp
--- a/Engine/Source/Runtime/Datasmith/CADKernel/Engine/Private/MeshDescriptionWrapper.cpp
+++ b/Engine/Source/Runtime/Datasmith/CADKernel/Engine/Private/MeshDescriptionWrapper.cpp
@@ -103,12 +103,11 @@ namespace UE::CADKernel::MeshUtilities
 
 		FVertexID3 NewVertices;
 		FVertexInstanceID3 NewVertexInstanceIDs;
-		FVertexInstanceID3 VertexInstancesIDs;
-		TArray<FVertexInstanceID> VertexInstancesIDArray(VertexInstancesIDs.ABC, 3);
+		TArray<FVertexInstanceID, TInlineAllocator<3>> VertexInstancesIDs;
 
 		for (const FPolygonID PolygonID : Mesh.Polygons().GetElementIDs())
 		{
-			Mesh.GetPolygonVertexInstances(PolygonID, VertexInstancesIDArray);
+			Mesh.GetPolygonVertexInstances(PolygonID, VertexInstancesIDs);
 
 			NewVertices[0] = VertexMapping[Mesh.GetVertexInstanceVertex(VertexInstancesIDs[2])];
 			NewVertices[1] = VertexMapping[Mesh.GetVertexInstanceVertex(VertexInstancesIDs[1])];
