Replaced FString and FName constructions with TEXT() and empty string literals.

diff --git a/Engine/Source/Runtime/AVEncoder/Private/AudioEncoderFactory.cpp b/Engine/Source/Runtime/AVEncoder/Private/AudioEncoderFactory.cpp
--- a/Engine/Source/Runtime/AVEncoder/Private/AudioEncoderFactory.cpp
+++ b/Engine/Source/Runtime/AVEncoder/Private/AudioEncoderFactory.cpp
@@ -93,7 +93,7 @@ namespace AVEncoder
 				
 				if (Str.IsEmpty())
 				{
-					return FString(TEXT("None"));
+					return FString("None");
 				}
 				else
 				{
diff --git a/Engine/Source/Runtime/AudioExtensions/Private/IAudioEndpoint.cpp b/Engine/Source/Runtime/AudioExtensions/Private/IAudioEndpoint.cpp
--- a/Engine/Source/Runtime/AudioExtensions/Private/IAudioEndpoint.cpp
+++ b/Engine/Source/Runtime/AudioExtensions/Private/IAudioEndpoint.cpp
@@ -136,13 +136,13 @@ FName IAudioEndpointFactory::GetEndpointTypeName()
 
 FName IAudioEndpointFactory::GetTypeNameForDefaultEndpoint()
 {
-	static FName DefaultEndpointName = FName(TEXT("Default Endpoint"));
+	static FName DefaultEndpointName = FName("Default Endpoint");
 	return DefaultEndpointName;
 }
 
 FName IAudioEndpointFactory::GetModularFeatureName()
 {
-	static FName ModularFeatureName = TEXT("External Audio Endpoint");
+	static FName ModularFeatureName = "External Audio Endpoint";
 	return ModularFeatureName;
 }
 
diff --git a/Engine/Source/Runtime/AudioExtensions/Private/ISoundfieldEndpoint.cpp b/Engine/Source/Runtime/AudioExtensions/Private/ISoundfieldEndpoint.cpp
--- a/Engine/Source/Runtime/AudioExtensions/Private/ISoundfieldEndpoint.cpp
+++ b/Engine/Source/Runtime/AudioExtensions/Private/ISoundfieldEndpoint.cpp
@@ -159,7 +159,7 @@ TArray<FName> ISoundfieldEndpointFactory::GetAllSoundfieldEndpointTypes()
 
 FName ISoundfieldEndpointFactory::DefaultSoundfieldEndpointName()
 {
-	static FName DefaultEndpointName = FName(TEXT("Default Soundfield Endpoint"));
+	static FName DefaultEndpointName = FName("Default Soundfield Endpoint");
 	return DefaultEndpointName;
 }
 
diff --git a/Engine/Source/Runtime/AudioExtensions/Private/ISoundfieldFormat.cpp b/Engine/Source/Runtime/AudioExtensions/Private/ISoundfieldFormat.cpp
--- a/Engine/Source/Runtime/AudioExtensions/Private/ISoundfieldFormat.cpp
+++ b/Engine/Source/Runtime/AudioExtensions/Private/ISoundfieldFormat.cpp
@@ -9,19 +9,19 @@
 
 FName ISoundfieldFactory::GetFormatNameForNoEncoding()
 {
-	static FName NoEncodingFormatName = FName(TEXT("No Encoding"));
+	static FName NoEncodingFormatName = FName("No Encoding");
 	return NoEncodingFormatName;
 }
 
 FName ISoundfieldFactory::GetFormatNameForInheritedEncoding()
 {
-	static FName InheritedFormatName = FName(TEXT("Inherited Encoding"));
+	static FName InheritedFormatName = FName("Inherited Encoding");
 	return InheritedFormatName;
 }
 
 FName ISoundfieldFactory::GetModularFeatureName()
 {
-	static FName SoundfieldFactoryName = FName(TEXT("Soundfield Format"));
+	static FName SoundfieldFactoryName = FName("Soundfield Format");
 	return SoundfieldFactoryName;
 }
 
diff --git a/Engine/Source/Runtime/AudioExtensions/Public/IAudioExtensionPlugin.h b/Engine/Source/Runtime/AudioExtensions/Public/IAudioExtensionPlugin.h
--- a/Engine/Source/Runtime/AudioExtensions/Public/IAudioExtensionPlugin.h
+++ b/Engine/Source/Runtime/AudioExtensions/Public/IAudioExtensionPlugin.h
@@ -205,7 +205,7 @@ public:
 	*/
 	virtual FString GetDisplayName()
 	{
-		static FString DisplayName = FString(TEXT("Generic Audio Plugin"));
+		static FString DisplayName = FString("Generic Audio Plugin");
 		return DisplayName;
 	}
 
@@ -253,14 +253,14 @@ public:
 	// IModularFeature
 	static FName GetModularFeatureName()
 	{
-		static FName AudioExtFeatureName = FName(TEXT("AudioSpatializationPlugin"));
+		static FName AudioExtFeatureName = FName("AudioSpatializationPlugin");
 		return AudioExtFeatureName;
 	}
 
 	/* Begin IAudioPluginWithMetadata implementation */
 	virtual FString GetDisplayName() override
 	{
-		static FString DisplayName = FString(TEXT("Generic Audio Spatialization Plugin"));
+		static FString DisplayName = FString("Generic Audio Spatialization Plugin");
 		return DisplayName;
 	}
 	/* End IAudioPluginWithMetadata implementation */
@@ -415,14 +415,14 @@ public:
 	// IModularFeature
 	static FName GetModularFeatureName()
 	{
-		static FName AudioExtFeatureName = FName(TEXT("AudioSourceDataOverridePlugin"));
+		static FName AudioExtFeatureName = FName("AudioSourceDataOverridePlugin");
 		return AudioExtFeatureName;
 	}
 
 	/* Begin IAudioPluginWithMetadata implementation */
 	virtual FString GetDisplayName() override
 	{
-		static FString DisplayName = FString(TEXT("Generic Audio Source Data Override Plugin"));
+		static FString DisplayName = FString("Generic Audio Source Data Override Plugin");
 		return DisplayName;
 	}
 	/* End IAudioPluginWithMetadata implementation */
@@ -499,14 +499,14 @@ public:
 	// IModularFeature
 	static FName GetModularFeatureName()
 	{
-		static FName AudioExtFeatureName = FName(TEXT("AudioOcclusionPlugin"));
+		static FName AudioExtFeatureName = FName("AudioOcclusionPlugin");
 		return AudioExtFeatureName;
 	}
 
 	/* Begin IAudioPluginWithMetadata implementation */
 	virtual FString GetDisplayName() override
 	{
-		static FString DisplayName = FString(TEXT("Generic Audio Occlusion Plugin"));
+		static FString DisplayName = FString("Generic Audio Occlusion Plugin");
 		return DisplayName;
 	}
 	/* End IAudioPluginWithMetadata implementation */
@@ -574,7 +574,7 @@ public:
 	// IModularFeature
 	static FName GetModularFeatureName()
 	{
-		static FName AudioExtFeatureName = FName(TEXT("AudioModulationPlugin"));
+		static FName AudioExtFeatureName = FName("AudioModulationPlugin");
 		return AudioExtFeatureName;
 	}
 
@@ -610,14 +610,14 @@ public:
 	// IModularFeature
 	static FName GetModularFeatureName()
 	{
-		static FName AudioExtFeatureName = FName(TEXT("AudioReverbPlugin"));
+		static FName AudioExtFeatureName = FName("AudioReverbPlugin");
 		return AudioExtFeatureName;
 	}
 
 	/* Begin IAudioPluginWithMetadata implementation */
 	virtual FString GetDisplayName() override
 	{
-		static FString DisplayName = FString(TEXT("Generic Audio Reverb Plugin"));
+		static FString DisplayName = FString("Generic Audio Reverb Plugin");
 		return DisplayName;
 	}
 	/* End IAudioPluginWithMetadata implementation */
diff --git a/Engine/Source/Runtime/AudioExtensions/Public/ISoundfieldEndpoint.h b/Engine/Source/Runtime/AudioExtensions/Public/ISoundfieldEndpoint.h
--- a/Engine/Source/Runtime/AudioExtensions/Public/ISoundfieldEndpoint.h
+++ b/Engine/Source/Runtime/AudioExtensions/Public/ISoundfieldEndpoint.h
@@ -178,7 +178,7 @@ public:
 	/** This is the FName used to register Soundfield Endpoint factories with the modular feature system. */
 	static FName GetModularFeatureName()
 	{
-		static FName SoundfieldEndpointName = FName(TEXT("Soundfield Endpoint"));
+		static FName SoundfieldEndpointName = FName("Soundfield Endpoint");
 		return SoundfieldEndpointName;
 	}
 
diff --git a/Engine/Source/Runtime/AudioLink/AudioLinkEngine/Public/IAudioLinkFactory.h b/Engine/Source/Runtime/AudioLink/AudioLinkEngine/Public/IAudioLinkFactory.h
--- a/Engine/Source/Runtime/AudioLink/AudioLinkEngine/Public/IAudioLinkFactory.h
+++ b/Engine/Source/Runtime/AudioLink/AudioLinkEngine/Public/IAudioLinkFactory.h
@@ -122,7 +122,7 @@ public:
 
 	inline static FName GetModularFeatureName()
 	{
-		static FName FeatureName = FName(TEXT("AudioLink Factory"));
+		static FName FeatureName = FName("AudioLink Factory");
 		return FeatureName;
 	}
 
diff --git a/Engine/Source/Runtime/AudioMixer/Private/AudioMixerDevice.cpp b/Engine/Source/Runtime/AudioMixer/Private/AudioMixerDevice.cpp
--- a/Engine/Source/Runtime/AudioMixer/Private/AudioMixerDevice.cpp
+++ b/Engine/Source/Runtime/AudioMixer/Private/AudioMixerDevice.cpp
@@ -1325,7 +1325,7 @@ namespace Audio
 		// 2. If Unset or not found, fallback to engine asset
 		if (!NewSubmix)
 		{
-			static const FString EngineSubmixDir = TEXT("/Engine/EngineSounds/Submixes");
+			static const FString EngineSubmixDir = "/Engine/EngineSounds/Submixes";
 
 			InObjectPath = FString::Printf(TEXT("%s/%s.%s"), *EngineSubmixDir, *InDefaultName, *InDefaultName);
 			NewSubmix = Cast<USoundSubmix>(InObjectPath.TryLoad());
diff --git a/Engine/Source/Runtime/AudioMixer/Private/AudioMixerSubmix.cpp b/Engine/Source/Runtime/AudioMixer/Private/AudioMixerSubmix.cpp
--- a/Engine/Source/Runtime/AudioMixer/Private/AudioMixerSubmix.cpp
+++ b/Engine/Source/Runtime/AudioMixer/Private/AudioMixerSubmix.cpp
@@ -227,7 +227,7 @@ namespace Audio
 
 		if (RecoverRecordingOnShutdownCVar && OwningSubmixObject.IsValid() && bIsRecording)
 		{
-			FString InterruptedFileName = TEXT("InterruptedRecording.wav");
+			FString InterruptedFileName = "InterruptedRecording.wav";
 			UE_LOG(LogAudioMixer, Warning, TEXT("Recording of Submix %s was interrupted. Saving interrupted recording as %s."), *(OwningSubmixObject->GetName()), *InterruptedFileName);
 			if (const USoundSubmix* SoundSubmix = Cast<const USoundSubmix>(OwningSubmixObject))
 			{
diff --git a/Engine/Source/Runtime/AudioMixer/Private/SoundFileIO/SoundFileIOManagerImpl.cpp b/Engine/Source/Runtime/AudioMixer/Private/SoundFileIO/SoundFileIOManagerImpl.cpp
--- a/Engine/Source/Runtime/AudioMixer/Private/SoundFileIO/SoundFileIOManagerImpl.cpp
+++ b/Engine/Source/Runtime/AudioMixer/Private/SoundFileIO/SoundFileIOManagerImpl.cpp
@@ -103,16 +103,16 @@ namespace Audio
 #if WITH_SNDFILE_IO 
 	#if PLATFORM_WINDOWS
 		#if PLATFORM_CPU_ARM_FAMILY && !PLATFORM_WINDOWS_ARM64EC
-			const FString PlatformPath = TEXT("Win64/arm64/");
+			const FString PlatformPath = "Win64/arm64/";
 		#else
-			const FString PlatformPath = TEXT("Win64/");
+			const FString PlatformPath = "Win64/";
 		#endif // PLATFORM_CPU_ARM_FAMILY && !PLATFORM_WINDOWS_ARM64EC
-		const FString DllName = TEXT("libsndfile-1.dll");
+		const FString DllName = "libsndfile-1.dll";
 	#elif PLATFORM_MAC //PLATFORM_WINDOWS
-		const FString PlatformPath = TEXT("Mac/");
-		const FString DllName = TEXT("libsndfile.1.dylib");
+		const FString PlatformPath = "Mac/";
+		const FString DllName = "libsndfile.1.dylib";
 	#elif PLATFORM_LINUX //PLATFORM_MAC
-		const FString PlatformPath = TEXT("Linux/");
+		const FString PlatformPath = "Linux/";
 		const FString DllName = ("libsndfile.so.1");
 	#else //PLATFORM_LINUX
 		#pragma message ("Platform not supported");
@@ -120,7 +120,7 @@ namespace Audio
 		const FString DllName;
 	#endif //PLATFORM_LINUX
 
-		const FString Path = FPaths::EngineDir() / FString(TEXT("Binaries/ThirdParty/libsndfile/")) / PlatformPath;
+		const FString Path = FPaths::EngineDir() / FString("Binaries/ThirdParty/libsndfile/") / PlatformPath;
 		FPlatformProcess::PushDllDirectory(*Path);
 		DllHandle = FPlatformProcess::GetDllHandle(*(Path + DllName));
 		FPlatformProcess::PopDllDirectory(*Path);
diff --git a/Engine/Source/Runtime/AudioMixerCore/Public/AudioMixerTypes.h b/Engine/Source/Runtime/AudioMixerCore/Public/AudioMixerTypes.h
--- a/Engine/Source/Runtime/AudioMixerCore/Public/AudioMixerTypes.h
+++ b/Engine/Source/Runtime/AudioMixerCore/Public/AudioMixerTypes.h
@@ -45,18 +45,18 @@ namespace Audio {
 	}
 
 	// Indicates a platform-specific format
-	inline FName NAME_PLATFORM_SPECIFIC(TEXT("PLATFORM_SPECIFIC"));
-	inline FName NAME_PROJECT_DEFINED(TEXT("PROJECT_DEFINED"));
+	inline FName NAME_PLATFORM_SPECIFIC("PLATFORM_SPECIFIC");
+	inline FName NAME_PROJECT_DEFINED("PROJECT_DEFINED");
 
 	// Supported on all platforms:
-	inline FName NAME_BINKA(TEXT("BINKA"));
-	inline FName NAME_ADPCM(TEXT("ADPCM"));
-	inline FName NAME_PCM(TEXT("PCM"));
-	inline FName NAME_OPUS(TEXT("OPUS"));
-	inline FName NAME_RADA(TEXT("RADA"));
+	inline FName NAME_BINKA("BINKA");
+	inline FName NAME_ADPCM("ADPCM");
+	inline FName NAME_PCM("PCM");
+	inline FName NAME_OPUS("OPUS");
+	inline FName NAME_RADA("RADA");
 
 	// Not yet supported on all platforms as a selectable option so is included under "platform specific" enumeration for now. 
-	inline FName NAME_OGG(TEXT("OGG"));
+	inline FName NAME_OGG("OGG");
 }
 
 struct FAudioPlatformSettings
diff --git a/Engine/Source/Runtime/AugmentedReality/Private/ARTypes.cpp b/Engine/Source/Runtime/AugmentedReality/Private/ARTypes.cpp
--- a/Engine/Source/Runtime/AugmentedReality/Private/ARTypes.cpp
+++ b/Engine/Source/Runtime/AugmentedReality/Private/ARTypes.cpp
@@ -53,22 +53,22 @@ void UARCandidateImage::PostEditChangeProperty(FPropertyChangedEvent& PropertyCh
 			}
 		}
 		// If the texture has changed enforce the aspect ratio on the physical size
-		if (PropertyChangedEvent.Property->GetFName() == FName(TEXT("CandidateTexture")))
+		if (PropertyChangedEvent.Property->GetFName() == FName("CandidateTexture"))
 		{
 			Height = Width * AspectRatio;
 		}
 		// Adjust width if they changed the height
-		else if (PropertyChangedEvent.Property->GetFName() == FName(TEXT("Height")))
+		else if (PropertyChangedEvent.Property->GetFName() == FName("Height"))
 		{
 			Width = Height * (1.f / AspectRatio);
 		}
 		// Adjust height if they changed the width
-		else if (PropertyChangedEvent.Property->GetFName() == FName(TEXT("Width")))
+		else if (PropertyChangedEvent.Property->GetFName() == FName("Width"))
 		{
 			Height = Width * AspectRatio;
 		}
 		// Adjust the sizes if they switched the orientation
-		else if (PropertyChangedEvent.Property->GetFName() == FName(TEXT("Orientation")))
+		else if (PropertyChangedEvent.Property->GetFName() == FName("Orientation"))
 		{
 			if (Orientation == EARCandidateImageOrientation::Landscape)
 			{
diff --git a/Engine/Source/Runtime/AugmentedReality/Public/ARDependencyHandler.h b/Engine/Source/Runtime/AugmentedReality/Public/ARDependencyHandler.h
--- a/Engine/Source/Runtime/AugmentedReality/Public/ARDependencyHandler.h
+++ b/Engine/Source/Runtime/AugmentedReality/Public/ARDependencyHandler.h
@@ -87,7 +87,7 @@ class UARDependencyHandler : public UObject, public IModularFeature
 public:
 	static FName GetModularFeatureName()
 	{
-		static FName FeatureName = FName(TEXT("ARDependencyHandler"));
+		static FName FeatureName = FName("ARDependencyHandler");
 		return FeatureName;
 	}
 	
diff --git a/Engine/Source/Runtime/AugmentedReality/Public/ARGeoTrackingSupport.h b/Engine/Source/Runtime/AugmentedReality/Public/ARGeoTrackingSupport.h
--- a/Engine/Source/Runtime/AugmentedReality/Public/ARGeoTrackingSupport.h
+++ b/Engine/Source/Runtime/AugmentedReality/Public/ARGeoTrackingSupport.h
@@ -93,7 +93,7 @@ class UARGeoTrackingSupport : public UObject, public IModularFeature
 public:
 	static FName GetModularFeatureName()
 	{
-		static FName FeatureName = FName(TEXT("ARGeoTrackingSupport"));
+		static FName FeatureName = FName("ARGeoTrackingSupport");
 		return FeatureName;
 	}
 	
diff --git a/Engine/Source/Runtime/AugmentedReality/Public/ARSessionConfigCookSupport.h b/Engine/Source/Runtime/AugmentedReality/Public/ARSessionConfigCookSupport.h
--- a/Engine/Source/Runtime/AugmentedReality/Public/ARSessionConfigCookSupport.h
+++ b/Engine/Source/Runtime/AugmentedReality/Public/ARSessionConfigCookSupport.h
@@ -13,7 +13,7 @@ class IARSessionConfigCookSupport : public IModularFeature
 public:
 	static FName GetModularFeatureName()
 	{
-		static FName FeatureName = FName(TEXT("ARSessionConfigCookingSupport"));
+		static FName FeatureName = FName("ARSessionConfigCookingSupport");
 		return FeatureName;
 	}
 
diff --git a/Engine/Source/Runtime/AugmentedReality/Public/ARSupportInterface.h b/Engine/Source/Runtime/AugmentedReality/Public/ARSupportInterface.h
--- a/Engine/Source/Runtime/AugmentedReality/Public/ARSupportInterface.h
+++ b/Engine/Source/Runtime/AugmentedReality/Public/ARSupportInterface.h
@@ -34,7 +34,7 @@ public:
 
 	static FName GetModularFeatureName()
 	{
-		static const FName ModularFeatureName = FName(TEXT("ARSystem"));
+		static const FName ModularFeatureName = FName("ARSystem");
 		return ModularFeatureName;
 	}
 
diff --git a/Engine/Source/Runtime/AutomationTest/Private/AutomationTestExcludelist.cpp b/Engine/Source/Runtime/AutomationTest/Private/AutomationTestExcludelist.cpp
--- a/Engine/Source/Runtime/AutomationTest/Private/AutomationTestExcludelist.cpp
+++ b/Engine/Source/Runtime/AutomationTest/Private/AutomationTestExcludelist.cpp
@@ -19,7 +19,7 @@ DEFINE_LOG_CATEGORY_STATIC(LogAutomationTestExcludelist, Log, All);
 
 namespace
 {
-	const FString FunctionalTestsPreFix = TEXT("Project.Functional Tests.");
+	const FString FunctionalTestsPreFix = "Project.Functional Tests.";
 
 	void SortExcludelist(TMap<FString, FAutomationTestExcludelistEntry>& List)
 	{
@@ -41,8 +41,8 @@ namespace
 			});
 	}
 
-	const FString TicketTrackerURLHashtagPropertyName = TEXT("URLHashtag");
-	const FString TicketTrackerURLBasePropertyName = TEXT("URLBase");
+	const FString TicketTrackerURLHashtagPropertyName = "URLHashtag";
+	const FString TicketTrackerURLBasePropertyName = "URLBase";
 } // anonymous namespace
 
 #if WITH_EDITOR
@@ -285,7 +285,7 @@ FString UAutomationTestExcludelist::GetConfigTaskTrackerHashtag() const
 
 FString UAutomationTestExcludelist::GetBeautifiedTaskTrackerTicketTagSuffix() const
 {
-	static const FString DefaultTaskTrackerTagSuffix = TEXT("unknown");
+	static const FString DefaultTaskTrackerTagSuffix = "unknown";
 
 	FString TaskTrackerTicketTagSuffix = DefaultConfig->GetTaskTrackerURLHashtag();
 	TaskTrackerTicketTagSuffix.TrimStartAndEndInline();
diff --git a/Engine/Source/Runtime/AutomationWorker/Private/AutomationWorkerModule.cpp b/Engine/Source/Runtime/AutomationWorker/Private/AutomationWorkerModule.cpp
--- a/Engine/Source/Runtime/AutomationWorker/Private/AutomationWorkerModule.cpp
+++ b/Engine/Source/Runtime/AutomationWorker/Private/AutomationWorkerModule.cpp
@@ -555,7 +555,7 @@ TSet<FName> GetRHIForAutomation()
 #if WITH_ENGINE
 	FString FeatureLevel = LexToString(GMaxRHIFeatureLevel);
 #else
-	FString FeatureLevel = TEXT("N/A");
+	FString FeatureLevel = "N/A";
 #endif
 
 	if (RHI.IsEmpty())
diff --git a/Engine/Source/Runtime/Cbor/Tests/CborTests.cpp b/Engine/Source/Runtime/Cbor/Tests/CborTests.cpp
--- a/Engine/Source/Runtime/Cbor/Tests/CborTests.cpp
+++ b/Engine/Source/Runtime/Cbor/Tests/CborTests.cpp
@@ -169,7 +169,7 @@ GROUP_TEST_CASE("CBor", "CBor::FCbor::Smoke Test", "[Smoke]")
 
 		// String
 
-		FString TestString(TEXT("ANSIString"));
+		FString TestString("ANSIString");
 
 		Writer.WriteValue(TestString);
 		CHECK(Reader.ReadNext(Context) == true);
diff --git a/Engine/Source/Runtime/CinematicCamera/Private/CineCameraComponent.cpp b/Engine/Source/Runtime/CinematicCamera/Private/CineCameraComponent.cpp
--- a/Engine/Source/Runtime/CinematicCamera/Private/CineCameraComponent.cpp
+++ b/Engine/Source/Runtime/CinematicCamera/Private/CineCameraComponent.cpp
@@ -232,7 +232,7 @@ void UCineCameraComponent::PostEditChangeProperty(FPropertyChangedEvent& Propert
 	// set focus plane color in case that's what changed
 	if (DebugFocusPlaneMID)
 	{
-		DebugFocusPlaneMID->SetVectorParameterValue(FName(TEXT("Color")), FocusSettings.DebugFocusPlaneColor.ReinterpretAsLinear());
+		DebugFocusPlaneMID->SetVectorParameterValue(FName("Color"), FocusSettings.DebugFocusPlaneColor.ReinterpretAsLinear());
 	}
 
 	// reset interpolation if the user changes anything
@@ -799,7 +799,7 @@ void UCineCameraComponent::CreateDebugFocusPlane()
 			DebugFocusPlaneMID = DebugFocusPlaneComponent->CreateAndSetMaterialInstanceDynamicFromMaterial(0, FocusPlaneVisualizationMaterial);
 			if (DebugFocusPlaneMID)
 			{
-				DebugFocusPlaneMID->SetVectorParameterValue(FName(TEXT("Color")), FocusSettings.DebugFocusPlaneColor.ReinterpretAsLinear());
+				DebugFocusPlaneMID->SetVectorParameterValue(FName("Color"), FocusSettings.DebugFocusPlaneColor.ReinterpretAsLinear());
 			}
 		}
 	}
diff --git a/Engine/Source/Runtime/CinematicCamera/Private/CineCameraSettings.cpp b/Engine/Source/Runtime/CinematicCamera/Private/CineCameraSettings.cpp
--- a/Engine/Source/Runtime/CinematicCamera/Private/CineCameraSettings.cpp
+++ b/Engine/Source/Runtime/CinematicCamera/Private/CineCameraSettings.cpp
@@ -211,7 +211,7 @@ void UCineCameraSettings::RecalcSensorAspectRatios()
 
 void UCineCameraSettings::CopyOldConfigSettings()
 {	
-	const FString SettingsConfigSection(TEXT("/Script/CinematicCamera.CineCameraSettings"));
+	const FString SettingsConfigSection("/Script/CinematicCamera.CineCameraSettings");
 
 	FString OldDefaultLensPresetName;
 	if (GConfig->GetString(*CineCameraConfigSection, TEXT("DefaultLensPresetName"), OldDefaultLensPresetName, GEngineIni))
diff --git a/Engine/Source/Runtime/ClientPilot/Private/ClientPilotBlackboard.cpp b/Engine/Source/Runtime/ClientPilot/Private/ClientPilotBlackboard.cpp
--- a/Engine/Source/Runtime/ClientPilot/Private/ClientPilotBlackboard.cpp
+++ b/Engine/Source/Runtime/ClientPilot/Private/ClientPilotBlackboard.cpp
@@ -11,7 +11,7 @@ void UClientPilotBlackboard::InitializeFromProfile(FString ProfileCategoryAndNam
 
 FString UClientPilotBlackboard::GetStringValue(FString KeyName)
 {
-	FString RetVal = TEXT("");
+	FString RetVal;
 	if (Blackboard.Contains(KeyName))
 	{
 		RetVal = Blackboard[KeyName];
diff --git a/Engine/Source/Runtime/CoreOnline/Public/Online/NestedVariant.h b/Engine/Source/Runtime/CoreOnline/Public/Online/NestedVariant.h
--- a/Engine/Source/Runtime/CoreOnline/Public/Online/NestedVariant.h
+++ b/Engine/Source/Runtime/CoreOnline/Public/Online/NestedVariant.h
@@ -17,7 +17,7 @@
  * using FNestedVariant = TNestedVariant<FString, bool, int64, double, FString>; // The first FString is the TMap key type.
  * FNestedVariant::FMapRef Data = FNestedVariant::FMap::CreateVariant();
  *
- * Data->AddVariant(TEXT("MyString"), FString(TEXT("1234-5678-9000")));
+ * Data->AddVariant(TEXT("MyString"), FString("1234-5678-9000"));
  * Data->AddVariant(TEXT("MyInt"), 42LL);
  * Data->AddVariant(TEXT("MyObject"), FNestedVariant::FMap::CreateVariant())->AddVariant(TEXT("MyNestedObjectDouble"), 1.0);
  * FNestedVariant::FArrayRef& Array = Data->AddVariant(TEXT("MyArray"), FNestedVariant::FArray::CreateVariant());
diff --git a/Engine/Source/Runtime/CrashReportCore/Private/Android/AndroidErrorReport.cpp b/Engine/Source/Runtime/CrashReportCore/Private/Android/AndroidErrorReport.cpp
--- a/Engine/Source/Runtime/CrashReportCore/Private/Android/AndroidErrorReport.cpp
+++ b/Engine/Source/Runtime/CrashReportCore/Private/Android/AndroidErrorReport.cpp
@@ -34,8 +34,8 @@ void FAndroidErrorReport::ShutDown()
 	// This way if anything goes wrong during processing it will not continue to cause issues on subsequent runs.
 
 	static const FMergeInfo MergeTypesOfInterest[] = { 
-		FMergeInfo(FString(TEXT("AllThreads")), FMergeInfo::EMergeType::Threads),
-		FMergeInfo(FString(TEXT("PlatformProperties")), FMergeInfo::EMergeType::PlatformProperties)
+		FMergeInfo(FString("AllThreads"), FMergeInfo::EMergeType::Threads),
+		FMergeInfo(FString("PlatformProperties"), FMergeInfo::EMergeType::PlatformProperties)
 	};
 
 	for (const FMergeInfo& MergeSource : MergeTypesOfInterest)
diff --git a/Engine/Source/Runtime/CrashReportCore/Private/CrashDescription.cpp b/Engine/Source/Runtime/CrashReportCore/Private/CrashDescription.cpp
--- a/Engine/Source/Runtime/CrashReportCore/Private/CrashDescription.cpp
+++ b/Engine/Source/Runtime/CrashReportCore/Private/CrashDescription.cpp
@@ -512,7 +512,7 @@ void FPrimaryCrashProperties::MakeCrashEventAttributes(TArray<FAnalyticsEventAtt
 	}
 	else
 	{
-		FString KeyName = FString(TEXT("EngineData.InvalidXML"));
+		FString KeyName = FString("EngineData.InvalidXML");
 		OutCrashAttributes.Add(FAnalyticsEventAttribute(*KeyName, true));
 	}
 }
@@ -670,7 +670,7 @@ FCrashWERContext::FCrashWERContext( const FString& WERXMLFilepath )
 			{
 				BranchName = ParsedParameters9[0].Replace( TEXT( "+" ), TEXT( "/" ), ESearchCase::CaseSensitive);
 
-				const FString DepotRoot = TEXT( "//depot/" );
+				const FString DepotRoot = "//depot/";
 				if (BranchName.StartsWith( DepotRoot ))
 				{
 					BranchName.MidInline( DepotRoot.Len(), MAX_int32, EAllowShrinking::No );
diff --git a/Engine/Source/Runtime/CrashReportCore/Private/CrashUpload.cpp b/Engine/Source/Runtime/CrashReportCore/Private/CrashUpload.cpp
--- a/Engine/Source/Runtime/CrashReportCore/Private/CrashUpload.cpp
+++ b/Engine/Source/Runtime/CrashReportCore/Private/CrashUpload.cpp
@@ -33,10 +33,10 @@ namespace CrashUploadDefs
 {
 	const float PingTimeoutSeconds = 5.f;
 
-	const FString APIKey(TEXT("CrashReporter"));
-	const FString AppEnvironmentInternal(TEXT("Dev"));
-	const FString AppEnvironmentExternal(TEXT("Release"));
-	const FString UploadType(TEXT("crashreports"));
+	const FString APIKey("CrashReporter");
+	const FString AppEnvironmentInternal("Dev");
+	const FString AppEnvironmentExternal("Release");
+	const FString UploadType("crashreports");
 }
 
 enum class ECompressedCrashFileHeader
@@ -699,7 +699,7 @@ bool FCrashUploadToReceiver::ParseServerResponse(FHttpResponsePtr Response, bool
 		return false;
 	}
 	// Turn the snippet into a complete XML document, to keep the XML parser happy
-	FXmlFile ParsedResponse(FString(TEXT("<Root>")) + Response->GetContentAsString() + TEXT("</Root>"), EConstructMethod::ConstructFromBuffer);
+	FXmlFile ParsedResponse(FString("<Root>") + Response->GetContentAsString() + TEXT("</Root>"), EConstructMethod::ConstructFromBuffer);
 	UE_LOG(CrashReportCoreLog, Log, TEXT("Response->GetContentAsString(): '%s'"), *Response->GetContentAsString());
 	if (!ParsedResponse.IsValid())
 	{
diff --git a/Engine/Source/Runtime/CrashReportCore/Private/GenericErrorReport.cpp b/Engine/Source/Runtime/CrashReportCore/Private/GenericErrorReport.cpp
--- a/Engine/Source/Runtime/CrashReportCore/Private/GenericErrorReport.cpp
+++ b/Engine/Source/Runtime/CrashReportCore/Private/GenericErrorReport.cpp
@@ -382,6 +382,6 @@ FString FGenericErrorReport::FindCrashedAppName() const
 FString FGenericErrorReport::FindCrashedAppPath() const
 {
 	UE_LOG( LogTemp, Warning, TEXT( "FGenericErrorReport::FindCrashedAppPath not implemented on this platform" ) );
-	return FString( TEXT( "GenericAppPath" ) );
+	return FString( "GenericAppPath" );
 }
 
diff --git a/Engine/Source/Runtime/D3D12RHI/Private/D3D12RayTracing.cpp b/Engine/Source/Runtime/D3D12RHI/Private/D3D12RayTracing.cpp
--- a/Engine/Source/Runtime/D3D12RHI/Private/D3D12RayTracing.cpp
+++ b/Engine/Source/Runtime/D3D12RHI/Private/D3D12RayTracing.cpp
@@ -230,7 +230,7 @@ static void DumpRayTracingGeometries(EDumpRayTracingGeometryMode Mode, int32 Num
 		return GetGeometrySize(A) > GetGeometrySize(B);
 	});
 
-	FName CategoryName(TEXT("D3D12RayTracing"));
+	FName CategoryName("D3D12RayTracing");
 	uint64 TotalSizeBytes = 0;
 	uint64 TopSizeBytes = 0;
 	BufferedOutput.CategorizedLogf(CategoryName, ELogVerbosity::Log, TEXT("Tracked FD3D12RayTracingGeometry objects"));
@@ -3612,7 +3612,7 @@ FD3D12RayTracingGeometry::FD3D12RayTracingGeometry(FRHICommandListBase& RHICmdLi
 {
 	INC_DWORD_STAT(STAT_D3D12RayTracingAllocatedBLAS);
 
-	static const FName NAME_BLAS(TEXT("BLAS"));
+	static const FName NAME_BLAS("BLAS");
 
 	DebugName = !Initializer.DebugName.IsNone() ? Initializer.DebugName : NAME_BLAS;
 	OwnerName = Initializer.OwnerName;
@@ -3778,7 +3778,7 @@ void FD3D12RayTracingGeometry::Swap(FD3D12RayTracingGeometry& Other)
 
 	Initializer = Other.Initializer;
 
-	DebugName = !Initializer.DebugName.IsNone() ? Initializer.DebugName : FName(TEXT("BLAS"));
+	DebugName = !Initializer.DebugName.IsNone() ? Initializer.DebugName : FName("BLAS");
 
 	checkf(Initializer.Segments.Num() > 0, TEXT("Ray tracing geometry must be initialized with at least one segment."));
 
diff --git a/Engine/Source/Runtime/D3D12RHI/Private/D3D12Resources.cpp b/Engine/Source/Runtime/D3D12RHI/Private/D3D12Resources.cpp
--- a/Engine/Source/Runtime/D3D12RHI/Private/D3D12Resources.cpp
+++ b/Engine/Source/Runtime/D3D12RHI/Private/D3D12Resources.cpp
@@ -1495,7 +1495,7 @@ bool FD3D12ResourceLocation::OnAllocationMoved(
 	FD3D12Resource* CurrentResource = GetResource();
 	FD3D12PoolAllocator* NewAllocator = GetPoolAllocator();
 
-	UE_TRACE_METADATA_SCOPE_ASSET_FNAME(CurrentResource->GetName(), FName(TEXT("FD3D12ResourceLocation::OnAllocationMoved")), NAME_None);
+	UE_TRACE_METADATA_SCOPE_ASSET_FNAME(CurrentResource->GetName(), FName("FD3D12ResourceLocation::OnAllocationMoved"), NAME_None);
 
 	// Textures don't have valid GPUVirtualAddress when IsTrackingAllAllocations() is false, so don't do memory trace in this case.
 	const bool bTrackingAllAllocations = GetParentDevice()->GetParentAdapter()->IsTrackingAllAllocations();
diff --git a/Engine/Source/Runtime/D3D12RHI/Private/D3D12Texture.cpp b/Engine/Source/Runtime/D3D12RHI/Private/D3D12Texture.cpp
--- a/Engine/Source/Runtime/D3D12RHI/Private/D3D12Texture.cpp
+++ b/Engine/Source/Runtime/D3D12RHI/Private/D3D12Texture.cpp
@@ -2003,7 +2003,7 @@ FRHILockTextureResult FD3D12Texture::Lock(FRHICommandListBase& RHICmdList, const
 {
 	SCOPE_CYCLE_COUNTER(STAT_D3D12LockTextureTime);
 
-	const static FName RHITextureLockName(TEXT("FRHITexture Lock"));
+	const static FName RHITextureLockName("FRHITexture Lock");
 	UE_TRACE_METADATA_SCOPE_ASSET_FNAME(GetName(), RHITextureLockName, GetOwnerName());
 
 	FD3D12Device* Device = GetParentDevice();
diff --git a/Engine/Source/Runtime/D3D12RHI/Private/D3D12Util.cpp b/Engine/Source/Runtime/D3D12RHI/Private/D3D12Util.cpp
--- a/Engine/Source/Runtime/D3D12RHI/Private/D3D12Util.cpp
+++ b/Engine/Source/Runtime/D3D12RHI/Private/D3D12Util.cpp
@@ -166,7 +166,7 @@ static FString GetD3D12ErrorString(HRESULT ErrorCode, ID3D12Device* Device)
 	if (ErrorCode == DXGI_ERROR_DEVICE_REMOVED && Device)
 	{
 		HRESULT hResDeviceRemoved = Device->GetDeviceRemovedReason();
-		ErrorCodeText += FString(TEXT(" with Reason: ")) + GetD3D12DeviceHungErrorString(hResDeviceRemoved);
+		ErrorCodeText += FString(" with Reason: ") + GetD3D12DeviceHungErrorString(hResDeviceRemoved);
 	}
 
 	return ErrorCodeText;
@@ -176,7 +176,7 @@ static FString GetD3D12ErrorString(HRESULT ErrorCode, ID3D12Device* Device)
 
 static FString GetD3D12TextureFlagString(uint32 TextureFlags)
 {
-	FString TextureFormatText = TEXT("");
+	FString TextureFormatText;
 
 	if (TextureFlags & D3D12_RESOURCE_FLAG_ALLOW_RENDER_TARGET)
 	{
diff --git a/Engine/Source/Runtime/Datasmith/CADKernel/Base/Public/Core/System.h b/Engine/Source/Runtime/Datasmith/CADKernel/Base/Public/Core/System.h
--- a/Engine/Source/Runtime/Datasmith/CADKernel/Base/Public/Core/System.h
+++ b/Engine/Source/Runtime/Datasmith/CADKernel/Base/Public/Core/System.h
@@ -20,7 +20,7 @@ class FSystem
 
 protected:
 
-	FString ProductName = FString(TEXT("CADKernel"));
+	FString ProductName = FString("CADKernel");
 
 	FVisu DefaultVisu;
 	FVisu* Viewer;
diff --git a/Engine/Source/Runtime/Datasmith/CADKernel/Base/Public/Mesh/MeshEnum.h b/Engine/Source/Runtime/Datasmith/CADKernel/Base/Public/Mesh/MeshEnum.h
--- a/Engine/Source/Runtime/Datasmith/CADKernel/Base/Public/Mesh/MeshEnum.h
+++ b/Engine/Source/Runtime/Datasmith/CADKernel/Base/Public/Mesh/MeshEnum.h
@@ -45,13 +45,13 @@ inline FString ToString(EGridSpace Space)
 	switch (Space)
 	{
 	case Default2D:
-		return FString(TEXT("Default2D"));
+		return FString("Default2D");
 	case Scaled:
-		return FString(TEXT("Scaled"));
+		return FString("Scaled");
 	case UniformScaled:
-		return FString(TEXT("UniformScaled"));
+		return FString("UniformScaled");
 	default:
-		return FString(TEXT(""));
+		return FString();
 	}
 }
 
diff --git a/Engine/Source/Runtime/Datasmith/CADKernel/Base/Public/UI/Message.h b/Engine/Source/Runtime/Datasmith/CADKernel/Base/Public/UI/Message.h
--- a/Engine/Source/Runtime/Datasmith/CADKernel/Base/Public/UI/Message.h
+++ b/Engine/Source/Runtime/Datasmith/CADKernel/Base/Public/UI/Message.h
@@ -36,7 +36,7 @@ public:
 	static void Error(const FmtType& Text, Types... Args)
 	{
 		static_assert((TIsValidVariadicFunctionArg<Types>::Value && ...), "Invalid argument(s) passed to FString::Printf");
-		FString CompleteText = TEXT("ERROR: ");
+		FString CompleteText = "ERROR: ";
 		CompleteText += Text;
 		VPrintf(Log, *CompleteText, Args...);
 	}
@@ -45,7 +45,7 @@ public:
 	static void Warning(const FmtType& Text, Types... Args)
 	{
 		static_assert((TIsValidVariadicFunctionArg<Types>::Value && ...), "Invalid argument(s) passed to FString::Printf");
-		FString CompleteText = TEXT("WARNING: ");
+		FString CompleteText = "WARNING: ";
 		CompleteText += Text;
 		VPrintf(Log, *CompleteText, Args...);
 	}
diff --git a/Engine/Source/Runtime/Datasmith/CADKernel/Editor/Private/CADKernelRetessellateAction.cpp b/Engine/Source/Runtime/Datasmith/CADKernel/Editor/Private/CADKernelRetessellateAction.cpp
--- a/Engine/Source/Runtime/Datasmith/CADKernel/Editor/Private/CADKernelRetessellateAction.cpp
+++ b/Engine/Source/Runtime/Datasmith/CADKernel/Editor/Private/CADKernelRetessellateAction.cpp
@@ -26,7 +26,7 @@ DEFINE_LOG_CATEGORY_STATIC(LogCADKernelRetessellation, Log, All)
 
 #define LOCTEXT_NAMESPACE "CADKernelRetessellateAction"
 
-static const FName CADKernelTessellationTag(TEXT("CADKernel"));
+static const FName CADKernelTessellationTag("CADKernel");
 
 static bool GCADKernelRetessellation = false;
 FAutoConsoleVariableRef GCADKernelDebugUseEngine(
diff --git a/Engine/Source/Runtime/Datasmith/DatasmithCore/Private/DatasmithSceneXmlReader.cpp b/Engine/Source/Runtime/Datasmith/DatasmithCore/Private/DatasmithSceneXmlReader.cpp
--- a/Engine/Source/Runtime/Datasmith/DatasmithCore/Private/DatasmithSceneXmlReader.cpp
+++ b/Engine/Source/Runtime/Datasmith/DatasmithCore/Private/DatasmithSceneXmlReader.cpp
@@ -1728,7 +1728,7 @@ void FDatasmithSceneXmlReader::ParseUEPbrMaterial(FXmlNode* InNode, TSharedPtr<
 
 	OutElement->SetParentLabel( *InNode->GetAttribute( DATASMITH_PARENTMATERIALLABEL ) );
 
-	FString ExpressionsTag = TEXT("Expressions");
+	FString ExpressionsTag = "Expressions";
 
 	FXmlNode* const* ExpressionsNode = Algo::FindByPredicate( InNode->GetChildrenNodes(), [ ExpressionsTag ]( const FXmlNode* Node ) -> bool
 	{
diff --git a/Engine/Source/Runtime/Datasmith/DatasmithCore/Private/DatasmithSceneXmlWriter.cpp b/Engine/Source/Runtime/Datasmith/DatasmithCore/Private/DatasmithSceneXmlWriter.cpp
--- a/Engine/Source/Runtime/Datasmith/DatasmithCore/Private/DatasmithSceneXmlWriter.cpp
+++ b/Engine/Source/Runtime/Datasmith/DatasmithCore/Private/DatasmithSceneXmlWriter.cpp
@@ -222,7 +222,7 @@ void FDatasmithSceneXmlWriterImpl::WriteValue(FArchive& Archive, int32 Indent, c
 			XmlString += TEXT("\" desc=\"") + Desc;
 		}
 
-		XmlString += FString( TEXT("\"/>") ) + LINE_TERMINATOR;
+		XmlString += FString( "\"/>" ) + LINE_TERMINATOR;
 		SerializeToArchive( Archive, XmlString );
 	}
 }
@@ -231,7 +231,7 @@ void FDatasmithSceneXmlWriterImpl::AppendXMLChild(FString& XmlString, int32 Inde
 {
 	AppendIndent(XmlString, Indent);
 	XmlString += FString::Printf( TEXT("<%s %s=\"%s\""), Tag, AttributeName, AttributeValue );
-	XmlString += FString( TEXT("/>") ) + LINE_TERMINATOR;
+	XmlString += FString( "/>" ) + LINE_TERMINATOR;
 }
 
 void FDatasmithSceneXmlWriterImpl::WriteTexture(FArchive& Archive, int32 Indent, const TCHAR* Prefix, const TCHAR* Name, FDatasmithTextureSampler UV)
@@ -244,7 +244,7 @@ void FDatasmithSceneXmlWriterImpl::WriteTexture(FArchive& Archive, int32 Indent,
 		BaseName += TEXT("_Tex");
 	}
 
-	FString XmlString = TEXT("<") + FString(Prefix) + TEXT(" tex=\"") + SanitizeXMLText( FDatasmithUtils::SanitizeFileName( *BaseName ) ) + FString( TEXT("\" ") );
+	FString XmlString = TEXT("<") + FString(Prefix) + TEXT(" tex=\"") + SanitizeXMLText( FDatasmithUtils::SanitizeFileName( *BaseName ) ) + FString( "\" " );
 	SerializeToArchive( Archive, XmlString );
 
 	XmlString = TEXT("coordinate=\"") + FString::FromInt( UV.CoordinateIndex ) + TEXT("\" ");
@@ -259,7 +259,7 @@ void FDatasmithSceneXmlWriterImpl::WriteTexture(FArchive& Archive, int32 Indent,
 	XmlString += TEXT("channel=\"") + FString::FromInt( UV.OutputChannel ) + TEXT("\" ");
 	XmlString += FString::Printf( TEXT("inv=\"%u\" "), UV.bInvert ? 1 : 0 );
 	XmlString += FString::Printf(TEXT("cropped=\"%u\""), UV.bCroppedTexture ? 1 : 0);
-	XmlString += FString( TEXT("/>") ) + LINE_TERMINATOR;
+	XmlString += FString( "/>" ) + LINE_TERMINATOR;
 
 	SerializeToArchive( Archive, XmlString );
 }
@@ -268,7 +268,7 @@ void FDatasmithSceneXmlWriterImpl::WriteTransform(const TSharedPtr< IDatasmithAc
 {
 	WriteIndent(Archive, Indent);
 
-	FString XmlString = FString( TEXT("<Transform ") ) +
+	FString XmlString = FString( "<Transform " ) +
 		TEXT("tx=\"") + FString::SanitizeFloat( ActorElement->GetTranslation().X ) +
 		TEXT("\" ty=\"") + FString::SanitizeFloat( ActorElement->GetTranslation().Y ) +
 		TEXT("\" tz=\"") + FString::SanitizeFloat( ActorElement->GetTranslation().Z ) + TEXT("\" ");
@@ -293,7 +293,7 @@ void FDatasmithSceneXmlWriterImpl::WriteTransform(const FTransform& Transform, F
 
 	FVector Location = Transform.GetLocation();
 
-	FString XmlString = FString(TEXT("<Transform ")) +
+	FString XmlString = FString("<Transform ") +
 		TEXT("tx=\"") + FString::SanitizeFloat(Location.X) +
 		TEXT("\" ty=\"") + FString::SanitizeFloat(Location.Y) +
 		TEXT("\" tz=\"") + FString::SanitizeFloat(Location.Z) + TEXT("\" ");
@@ -329,7 +329,7 @@ void FDatasmithSceneXmlWriterImpl::WriteMeshElement(const TSharedPtr< IDatasmith
 
 	FString XmlString = TEXT("<") + FString(DATASMITH_STATICMESHNAME) + TEXT(" name=\"") + SanitizeXMLText( MeshElement->GetName() ) + TEXT("\"");
 	XmlString += TEXT(" label=\"") + SanitizeXMLText( MeshElement->GetLabel() ) + TEXT("\"");
-	XmlString += FString( TEXT(">") ) + LINE_TERMINATOR;
+	XmlString += FString( ">" ) + LINE_TERMINATOR;
 	SerializeToArchive(Archive, XmlString);
 
 	WriteIndent(Archive, Indent + 1);
@@ -339,7 +339,7 @@ void FDatasmithSceneXmlWriterImpl::WriteMeshElement(const TSharedPtr< IDatasmith
 	if (MeshElement->GetArea() > 0)
 	{
 		WriteIndent(Archive, Indent+1);
-		XmlString = FString(TEXT("<Size ")) +
+		XmlString = FString("<Size ") +
 			TEXT("a=\"") + FString::SanitizeFloat(MeshElement->GetArea()) +
 			TEXT("\" x=\"") + FString::SanitizeFloat(MeshElement->GetWidth()) +
 			TEXT("\" y=\"") + FString::SanitizeFloat(MeshElement->GetDepth()) +
@@ -394,7 +394,7 @@ void FDatasmithSceneXmlWriterImpl::WriteLevelSequenceElement(const TSharedRef< I
 	WriteIndent(Archive, Indent);
 
 	FString XmlString = TEXT("<") + FString(DATASMITH_LEVELSEQUENCENAME) + TEXT(" name=\"") + SanitizeXMLText( SequenceElement->GetName() ) + TEXT("\"");
-	XmlString += FString( TEXT(">") ) + LINE_TERMINATOR;
+	XmlString += FString( ">" ) + LINE_TERMINATOR;
 
 	SerializeToArchive(Archive, XmlString);
 
@@ -419,7 +419,7 @@ void FDatasmithSceneXmlWriterImpl::WriteLevelVariantSetsElement( const TSharedRe
 	WriteIndent( Archive, Indent );
 
 	FString XmlString = TEXT( "<" ) + FString( DATASMITH_LEVELVARIANTSETSNAME ) + TEXT( " name=\"" ) + SanitizeXMLText( LevelVariantSetsElement->GetName() ) + TEXT( "\"" );
-	XmlString += FString( TEXT( ">" ) ) + LINE_TERMINATOR;
+	XmlString += FString( ">" ) + LINE_TERMINATOR;
 
 	SerializeToArchive( Archive, XmlString );
 
@@ -443,7 +443,7 @@ void FDatasmithSceneXmlWriterImpl::WriteVariantSetElement( const TSharedRef< IDa
 	WriteIndent( Archive, Indent );
 
 	FString XmlString = TEXT( "<" ) + FString( DATASMITH_VARIANTSETNAME ) + TEXT( " name=\"" ) + SanitizeXMLText( VariantSetElement->GetName() ) + TEXT( "\"" );
-	XmlString += FString( TEXT( ">" ) ) + LINE_TERMINATOR;
+	XmlString += FString( ">" ) + LINE_TERMINATOR;
 
 	SerializeToArchive( Archive, XmlString );
 
@@ -467,7 +467,7 @@ void FDatasmithSceneXmlWriterImpl::WriteVariantElement( const TSharedRef< IDatas
 	WriteIndent( Archive, Indent );
 
 	FString XmlString = TEXT( "<" ) + FString( DATASMITH_VARIANTNAME ) + TEXT( " name=\"" ) + SanitizeXMLText( VariantElement->GetName() ) + TEXT( "\"" );
-	XmlString += FString( TEXT( ">" ) ) + LINE_TERMINATOR;
+	XmlString += FString( ">" ) + LINE_TERMINATOR;
 
 	SerializeToArchive( Archive, XmlString );
 
@@ -498,7 +498,7 @@ void FDatasmithSceneXmlWriterImpl::WriteActorBindingElement( const TSharedRef< I
 	}
 
 	FString XmlString = TEXT( "<" ) + FString( DATASMITH_ACTORBINDINGNAME ) + TEXT( " actor=\"" ) + SanitizeXMLText( ActorName ) + TEXT( "\"" );
-	XmlString += FString( TEXT( ">" ) ) + LINE_TERMINATOR;
+	XmlString += FString( ">" ) + LINE_TERMINATOR;
 
 	SerializeToArchive( Archive, XmlString );
 
@@ -625,7 +625,7 @@ void FDatasmithSceneXmlWriterImpl::WriteBaseActorElement(const TSharedPtr< IData
 
 	FString XmlString = TEXT("<") + FString(DATASMITH_ACTORNAME) + TEXT(" name=\"") + SanitizeXMLText(FDatasmithUtils::SanitizeFileName(ActorElement->GetName())) + TEXT("\"");
 
-	XmlString += GetLabelAndLayer(ActorElement) + FString(TEXT(">")) + LINE_TERMINATOR;
+	XmlString += GetLabelAndLayer(ActorElement) + FString(">") + LINE_TERMINATOR;
 
 	SerializeToArchive(Archive, XmlString);
 
@@ -680,7 +680,7 @@ void FDatasmithSceneXmlWriterImpl::WriteActorChildren(const TSharedPtr< IDatasmi
 			WriteActorElement(ActorElement->GetChild(i), Archive, Indent + 2);
 		}
 
-		XmlString = FString(TEXT("</children>")) + LINE_TERMINATOR;
+		XmlString = FString("</children>") + LINE_TERMINATOR;
 
 		WriteIndent(Archive, Indent + 1);
 		SerializeToArchive(Archive, XmlString);
@@ -719,7 +719,7 @@ void FDatasmithSceneXmlWriterImpl::WriteBeginOfMeshActorElement(const TSharedPtr
 
 	XmlString += GetLabelAndLayer(MeshActorElement);
 
-	XmlString += FString(TEXT(">")) + LINE_TERMINATOR;
+	XmlString += FString(">") + LINE_TERMINATOR;
 
 	SerializeToArchive(Archive, XmlString);
 
@@ -731,7 +731,7 @@ void FDatasmithSceneXmlWriterImpl::WriteBeginOfMeshActorElement(const TSharedPtr
 	{
 		XmlString += TEXT(" name=\"") + SanitizeXMLText(MeshActorElement->GetStaticMeshPathName()) + TEXT("\"");
 	}
-	XmlString += FString(TEXT("/>")) + LINE_TERMINATOR;
+	XmlString += FString("/>") + LINE_TERMINATOR;
 
 	SerializeToArchive(Archive, XmlString);
 
@@ -793,13 +793,13 @@ void FDatasmithSceneXmlWriterImpl::WriteCustomActorElement(const TSharedPtr< IDa
 {
 	WriteIndent(Archive, Indent);
 
-	FString XmlString = TEXT("<");
+	FString XmlString = "<";
 	XmlString += FString(XmlNodeName) + TEXT(" name=\"") + SanitizeXMLText(FDatasmithUtils::SanitizeFileName(CustomActorElement->GetName())) + TEXT("\" ");
 	XmlString += FString(DATASMITH_CUSTOMACTORPATHNAME) + TEXT("=\"") + SanitizeXMLText( CustomActorElement->GetClassOrPathName() ) + TEXT("\" ");
 
 	XmlString += GetLabelAndLayer( CustomActorElement );
 
-	XmlString += FString(TEXT(">")) + LINE_TERMINATOR;
+	XmlString += FString(">") + LINE_TERMINATOR;
 	SerializeToArchive( Archive, XmlString );
 
 	WriteTransform( CustomActorElement, Archive, Indent + 1 );
@@ -819,12 +819,12 @@ void FDatasmithSceneXmlWriterImpl::WriteLandscapeActorElement(const TSharedPtr<
 {
 	WriteIndent(Archive, Indent);
 
-	FString XmlString = TEXT("<");
+	FString XmlString = "<";
 	XmlString += FString(DATASMITH_LANDSCAPENAME) + TEXT(" name=\"") + SanitizeXMLText( LandscapeActorElement->GetName() ) + TEXT("\" ");
 
 	XmlString += GetLabelAndLayer( LandscapeActorElement );
 
-	XmlString += FString(TEXT(">")) + LINE_TERMINATOR;
+	XmlString += FString(">") + LINE_TERMINATOR;
 	SerializeToArchive(Archive, XmlString);
 
 	WriteTransform( LandscapeActorElement, Archive, Indent + 1 );
@@ -896,7 +896,7 @@ void FDatasmithSceneXmlWriterImpl::WriteLightActorElement(const TSharedPtr< IDat
 		TEXT("\" G=\"") + FString::SanitizeFloat( LightActorElement->GetColor().G ) +
 		TEXT("\" B=\"") + FString::SanitizeFloat( LightActorElement->GetColor().B );
 
-	XmlString += FString(TEXT("\"/>")) + LINE_TERMINATOR;
+	XmlString += FString("\"/>") + LINE_TERMINATOR;
 
 	SerializeToArchive( Archive, XmlString );
 
@@ -1358,10 +1358,10 @@ void FDatasmithSceneXmlWriterImpl::WriteMaterialElement( TSharedPtr< IDatasmithM
 
 	WriteIndent(Archive, Indent);
 
-	FString XmlString = TEXT("<");
+	FString XmlString = "<";
 	XmlString += FString(DATASMITH_MATERIALNAME) + TEXT(" name=\"") + SanitizeXMLText(MaterialElement->GetName()) + TEXT("\"");
 	XmlString += TEXT(" label=\"") + SanitizeXMLText(MaterialElement->GetLabel()) + TEXT("\"");
-	XmlString += FString(TEXT(">")) + LINE_TERMINATOR;
+	XmlString += FString(">") + LINE_TERMINATOR;
 
 	SerializeToArchive( Archive, XmlString );
 
@@ -1398,7 +1398,7 @@ void FDatasmithSceneXmlWriterImpl::WriteMaterialInstanceElement(TSharedPtr< IDat
 {
 	WriteIndent(Archive, Indent);
 
-	FString XmlString = TEXT("<");
+	FString XmlString = "<";
 	XmlString += FString(DATASMITH_MATERIALINSTANCENAME) + TEXT(" name=\"") + SanitizeXMLText(FDatasmithUtils::SanitizeFileName(MaterialInstanceElement->GetName())) + TEXT("\" ");
 	XmlString += TEXT(" label=\"") + SanitizeXMLText(MaterialInstanceElement->GetLabel()) + TEXT("\" ");
 	XmlString += FString(DATASMITH_MATERIALINSTANCETYPE) + TEXT("=\"") + FString::FromInt( (int32)MaterialInstanceElement->GetMaterialType() ) + TEXT("\" ");
@@ -1409,7 +1409,7 @@ void FDatasmithSceneXmlWriterImpl::WriteMaterialInstanceElement(TSharedPtr< IDat
 		XmlString += FString(DATASMITH_MATERIALINSTANCEPATHNAME) + TEXT("=\"") + SanitizeXMLText( MaterialInstanceElement->GetCustomMaterialPathName() ) + TEXT("\" ");
 	}
 
-	XmlString += FString(TEXT(">")) + LINE_TERMINATOR;
+	XmlString += FString(">") + LINE_TERMINATOR;
 	SerializeToArchive( Archive, XmlString );
 
 	WriteKeyValueProperties( *MaterialInstanceElement, Archive, Indent );
@@ -1423,11 +1423,11 @@ void FDatasmithSceneXmlWriterImpl::WriteDecalMaterialElement(TSharedPtr< IDatasm
 {
 	WriteIndent(Archive, Indent);
 
-	FString XmlString = TEXT("<");
+	FString XmlString = "<";
 	XmlString += FString(DATASMITH_DECALMATERIALNAME) + TEXT(" name=\"") + SanitizeXMLText(FDatasmithUtils::SanitizeFileName(DecalMaterialElement->GetName())) + TEXT("\" ");
 	XmlString += TEXT(" label=\"") + SanitizeXMLText(DecalMaterialElement->GetLabel()) + TEXT("\" ");
 
-	XmlString += FString(TEXT(">")) + LINE_TERMINATOR;
+	XmlString += FString(">") + LINE_TERMINATOR;
 	SerializeToArchive( Archive, XmlString );
 
 	WriteIndent(Archive, Indent + 1);
@@ -1447,7 +1447,7 @@ void FDatasmithSceneXmlWriterImpl::WriteUEPbrMaterialElement( const TSharedRef<
 {
 	WriteIndent( Archive, Indent );
 
-	FString XmlString = TEXT("<");
+	FString XmlString = "<";
 	XmlString += FString(DATASMITH_UEPBRMATERIALNAME) + TEXT(" name=\"") + SanitizeXMLText( MaterialElement->GetName() ) + TEXT("\"");
 
 	if ( FCString::Strcmp( MaterialElement->GetLabel(), MaterialElement->GetName() ) != 0 )
@@ -1460,7 +1460,7 @@ void FDatasmithSceneXmlWriterImpl::WriteUEPbrMaterialElement( const TSharedRef<
 		XmlString += TEXT(" ") + FString( DATASMITH_PARENTMATERIALLABEL ) + TEXT("=\"") + SanitizeXMLText( MaterialElement->GetParentLabel() ) + TEXT("\"");
 	}
 
-	XmlString += FString( TEXT(">") ) + LINE_TERMINATOR;
+	XmlString += FString( ">" ) + LINE_TERMINATOR;
 	SerializeToArchive( Archive, XmlString );
 
 	WriteUEPbrMaterialExpressions( MaterialElement, Archive, Indent + 1 );
@@ -1510,7 +1510,7 @@ void FDatasmithSceneXmlWriterImpl::WriteUEPbrMaterialElement( const TSharedRef<
 void FDatasmithSceneXmlWriterImpl::WriteUEPbrMaterialExpressions( const TSharedRef< IDatasmithUEPbrMaterialElement >& MaterialElement, FArchive& Archive, int32 Indent )
 {
 	WriteIndent( Archive, Indent );
-	FString XmlString = FString( TEXT("<Expressions>") ) + LINE_TERMINATOR;
+	FString XmlString = FString( "<Expressions>" ) + LINE_TERMINATOR;
 	SerializeToArchive( Archive, XmlString );
 
 	for ( int32 ExpressionIndex = 0; ExpressionIndex < MaterialElement->GetExpressionsCount(); ++ExpressionIndex )
@@ -1571,7 +1571,7 @@ void FDatasmithSceneXmlWriterImpl::WriteUEPbrMaterialExpressions( const TSharedR
 	}
 
 	WriteIndent( Archive, Indent );
-	XmlString = TEXT("</") + FString( TEXT("Expressions") ) + TEXT(">") + LINE_TERMINATOR;
+	XmlString = TEXT("</") + FString( "Expressions" ) + TEXT(">") + LINE_TERMINATOR;
 	SerializeToArchive( Archive, XmlString );
 }
 
@@ -1592,21 +1592,21 @@ void FDatasmithSceneXmlWriterImpl::WriteUEPbrMaterialExpressionInput( const TSha
 
 void FDatasmithSceneXmlWriterImpl::AppendMaterialExpressionAttributes( const IDatasmithMaterialExpression& Expression, FString& XmlString )
 {
-	XmlString += FString( TEXT("Name=\"") ) + Expression.GetName() + FString( TEXT("\"") );
+	XmlString += FString( "Name=\"" ) + Expression.GetName() + FString( "\"" );
 }
 
 void FDatasmithSceneXmlWriterImpl::WriteMaterialExpressionFlattenNormal( const TSharedRef< IDatasmithUEPbrMaterialElement >& MaterialElement, const IDatasmithMaterialExpressionFlattenNormal& FlattenNormal, FArchive& Archive, int32 Indent )
 {
 	WriteIndent( Archive, Indent );
 
-	FString XmlString = FString( TEXT("<FlattenNormal>") ) + LINE_TERMINATOR;
+	FString XmlString = FString( "<FlattenNormal>" ) + LINE_TERMINATOR;
 	SerializeToArchive( Archive, XmlString );
 
 	WriteUEPbrMaterialExpressionInput( MaterialElement, FlattenNormal.GetNormal(), Archive, Indent + 1 );
 	WriteUEPbrMaterialExpressionInput( MaterialElement, FlattenNormal.GetFlatness(), Archive, Indent + 1 );
 
 	WriteIndent( Archive, Indent );
-	XmlString = TEXT("</") + FString( TEXT("FlattenNormal") ) + TEXT(">") + LINE_TERMINATOR;
+	XmlString = TEXT("</") + FString( "FlattenNormal" ) + TEXT(">") + LINE_TERMINATOR;
 
 	SerializeToArchive( Archive, XmlString );
 }
@@ -1615,16 +1615,16 @@ void FDatasmithSceneXmlWriterImpl::WriteMaterialExpressionTexture( const TShared
 {
 	WriteIndent( Archive, Indent );
 
-	FString XmlString = TEXT("<Texture ");
+	FString XmlString = "<Texture ";
 	AppendMaterialExpressionAttributes( TextureExpression, XmlString );
-	XmlString += FString( TEXT(" PathName=\"") ) + TextureExpression.GetTexturePathName() + FString( TEXT("\">") ) + LINE_TERMINATOR;
+	XmlString += FString( " PathName=\"" ) + TextureExpression.GetTexturePathName() + FString( "\">" ) + LINE_TERMINATOR;
 
 	SerializeToArchive( Archive, XmlString );
 
 	WriteUEPbrMaterialExpressionInput( MaterialElement, TextureExpression.GetInputCoordinate(), Archive, Indent + 1 );
 
 	WriteIndent( Archive, Indent );
-	XmlString = TEXT("</") + FString( TEXT("Texture") ) + TEXT(">") + LINE_TERMINATOR;
+	XmlString = TEXT("</") + FString( "Texture" ) + TEXT(">") + LINE_TERMINATOR;
 
 	SerializeToArchive( Archive, XmlString );
 }
@@ -1645,10 +1645,10 @@ void FDatasmithSceneXmlWriterImpl::WriteMaterialExpressionBool( const TSharedRef
 {
 	WriteIndent( Archive, Indent );
 
-	FString XmlString = TEXT("<Bool ");
+	FString XmlString = "<Bool ";
 	AppendMaterialExpressionAttributes( BoolExpression, XmlString );
 	XmlString += TEXT(" constant=\"" ) + LexToString( BoolExpression.GetBool() ) + TEXT("\"");
-	XmlString += FString( TEXT("/>") ) + LINE_TERMINATOR;
+	XmlString += FString( "/>" ) + LINE_TERMINATOR;
 
 	SerializeToArchive( Archive, XmlString );
 }
@@ -1657,10 +1657,10 @@ void FDatasmithSceneXmlWriterImpl::WriteMaterialExpressionColor( const TSharedRe
 {
 	WriteIndent( Archive, Indent );
 
-	FString XmlString = TEXT("<Color ");
+	FString XmlString = "<Color ";
 	AppendMaterialExpressionAttributes( ColorExpression, XmlString );
 	XmlString += TEXT(" constant=\"" ) + ColorExpression.GetColor().ToString() + TEXT("\"");
-	XmlString += FString( TEXT("/>") ) + LINE_TERMINATOR;
+	XmlString += FString( "/>" ) + LINE_TERMINATOR;
 
 	SerializeToArchive( Archive, XmlString );
 }
@@ -1669,10 +1669,10 @@ void FDatasmithSceneXmlWriterImpl::WriteMaterialExpressionScalar( const TSharedR
 {
 	WriteIndent( Archive, Indent );
 
-	FString XmlString = TEXT("<Scalar ");
+	FString XmlString = "<Scalar ";
 	AppendMaterialExpressionAttributes( ScalarExpression, XmlString );
 	XmlString += TEXT(" constant=\"" ) + FString::SanitizeFloat( ScalarExpression.GetScalar() ) + TEXT("\"");
-	XmlString += FString( TEXT("/>") ) + LINE_TERMINATOR;
+	XmlString += FString( "/>" ) + LINE_TERMINATOR;
 
 	SerializeToArchive( Archive, XmlString );
 }
@@ -1681,9 +1681,9 @@ void FDatasmithSceneXmlWriterImpl::WriteMaterialExpressionGeneric( const TShared
 {
 	WriteIndent( Archive, Indent );
 
-	FString XmlString = TEXT("<") + FString( GenericExpression.GetExpressionName() ) + FString( TEXT(" ") );
+	FString XmlString = TEXT("<") + FString( GenericExpression.GetExpressionName() ) + FString( " " );
 	AppendMaterialExpressionAttributes( GenericExpression, XmlString );
-	XmlString += FString( TEXT(">") ) + LINE_TERMINATOR;
+	XmlString += FString( ">" ) + LINE_TERMINATOR;
 	SerializeToArchive( Archive, XmlString );
 
 	WriteKeyValueProperties( GenericExpression, Archive, Indent );
@@ -1704,7 +1704,7 @@ void FDatasmithSceneXmlWriterImpl::WriteMaterialExpressionFunctionCall( const TS
 	WriteIndent( Archive, Indent );
 
 	FString XmlString = TEXT("<FunctionCall Function=\"") + FString( FunctionCall.GetFunctionPathName() ) + TEXT("\"");
-	XmlString += FString( TEXT(">") ) + LINE_TERMINATOR;
+	XmlString += FString( ">" ) + LINE_TERMINATOR;
 	SerializeToArchive( Archive, XmlString );
 
 	for ( int32 InputIndex = 0 ; InputIndex < FunctionCall.GetInputCount(); ++InputIndex )
@@ -1713,7 +1713,7 @@ void FDatasmithSceneXmlWriterImpl::WriteMaterialExpressionFunctionCall( const TS
 	}
 
 	WriteIndent( Archive, Indent );
-	XmlString = FString( TEXT("</FunctionCall>") ) + LINE_TERMINATOR;
+	XmlString = FString( "</FunctionCall>" ) + LINE_TERMINATOR;
 
 	SerializeToArchive( Archive, XmlString );
 }
@@ -1833,7 +1833,7 @@ void FDatasmithSceneXmlWriterImpl::WriteMetaDataElement(const TSharedPtr< IDatas
 		}
 	}
 
-	FString XmlString = TEXT("<");
+	FString XmlString = "<";
 	XmlString += FString(DATASMITH_METADATANAME) + TEXT(" name=\"") + ElementName + TEXT("\" ");
 	XmlString += FString(DATASMITH_REFERENCENAME) + TEXT("=\"") + TypeString + TEXT(".") + ReferenceName + TEXT("\">") + LINE_TERMINATOR;
 	SerializeToArchive(Archive, XmlString);
@@ -1902,7 +1902,7 @@ void FDatasmithSceneXmlWriter::Serialize( TSharedRef< IDatasmithScene > Datasmit
 	// explicitly set the locale to have consistent exports
 	FDatasmithLocaleScope CLocaleScope;
 
-	FString XmlString = FString( TEXT("<DatasmithUnrealScene>") ) + LINE_TERMINATOR;
+	FString XmlString = FString( "<DatasmithUnrealScene>" ) + LINE_TERMINATOR;
 	FDatasmithSceneXmlWriterImpl::SerializeToArchive( Archive, XmlString );
 
 	FDatasmithSceneXmlWriterImpl::WriteIndent(Archive, 1);
diff --git a/Engine/Source/Runtime/Datasmith/DatasmithCore/Public/DatasmithCloth.h b/Engine/Source/Runtime/Datasmith/DatasmithCore/Public/DatasmithCloth.h
--- a/Engine/Source/Runtime/Datasmith/DatasmithCore/Public/DatasmithCloth.h
+++ b/Engine/Source/Runtime/Datasmith/DatasmithCore/Public/DatasmithCloth.h
@@ -151,7 +151,7 @@ public:
 class UE_DEPRECATED(5.5, "The experimental Cloth importer is no longer supported.") IDatasmithClothFactoryClassesProvider : public IModularFeature
 {
 public:
-	inline static const FName FeatureName = TEXT("IDatasmithClothFactoryClassesProvider");
+	inline static const FName FeatureName = "IDatasmithClothFactoryClassesProvider";
 
 	IDatasmithClothFactoryClassesProvider() = default;
 	virtual ~IDatasmithClothFactoryClassesProvider() = default;
diff --git a/Engine/Source/Runtime/DeveloperSettings/Private/Engine/DeveloperSettings.cpp b/Engine/Source/Runtime/DeveloperSettings/Private/Engine/DeveloperSettings.cpp
--- a/Engine/Source/Runtime/DeveloperSettings/Private/Engine/DeveloperSettings.cpp
+++ b/Engine/Source/Runtime/DeveloperSettings/Private/Engine/DeveloperSettings.cpp
@@ -17,11 +17,11 @@ UDeveloperSettings::UDeveloperSettings(const FObjectInitializer& ObjectInitializ
 
 FName UDeveloperSettings::GetContainerName() const
 {
-	static const FName ProjectName(TEXT("Project"));
-	static const FName EditorName(TEXT("Editor"));
+	static const FName ProjectName("Project");
+	static const FName EditorName("Editor");
 
-	static const FName EditorSettingsName(TEXT("EditorSettings"));
-	static const FName EditorPerProjectUserSettingsName(TEXT("EditorPerProjectUserSettings"));
+	static const FName EditorSettingsName("EditorSettings");
+	static const FName EditorPerProjectUserSettingsName("EditorPerProjectUserSettings");
 
 	FName ConfigName = GetClass()->ClassConfigName;
 
@@ -35,9 +35,9 @@ FName UDeveloperSettings::GetContainerName() const
 
 FName UDeveloperSettings::GetCategoryName() const
 {
-	static const FName GeneralName(TEXT("General"));
-	static const FName EditorSettingsName(TEXT("EditorSettings"));
-	static const FName EditorPerProjectUserSettingsName(TEXT("EditorPerProjectUserSettings"));
+	static const FName GeneralName("General");
+	static const FName EditorSettingsName("EditorSettings");
+	static const FName EditorPerProjectUserSettingsName("EditorPerProjectUserSettings");
 
 	if ( CategoryName != NAME_None )
 	{
@@ -105,7 +105,7 @@ TSharedPtr<SWidget> UDeveloperSettings::GetCustomSettingsWidget() const
 
 #if WITH_EDITOR
 
-static FName DeveloperSettingsConsoleVariableMetaFName(TEXT("ConsoleVariable"));
+static FName DeveloperSettingsConsoleVariableMetaFName("ConsoleVariable");
 
 void UDeveloperSettings::ImportConsoleVariableValues()
 {
@@ -182,7 +182,7 @@ void UDeveloperSettings::ExportValuesToConsoleVariables(FProperty* PropertyThatC
 			{
 				// Reduce the amount of log spam for read-only properties. 
 				// We assume that if property requires restart it is very likely it needs to stay read-only and therefore no need to log a warning.
-				static const FName ConfigRestartRequiredKey = TEXT("ConfigRestartRequired");
+				static const FName ConfigRestartRequiredKey = "ConfigRestartRequired";
 				if (!PropertyThatChanged->GetBoolMetaData(ConfigRestartRequiredKey))
 				{
 					UE_LOG(LogInit, Warning, TEXT("CVar named '%s' marked up in %s was not found or is set to read-only"), *CVarName, *GetClass()->GetName());
