Removed all debugging and reportingactive code from the CADKernel module as it was preventing its usage in multi-threading mode
- Still remaining code between ifdef to clean up
Deleted unnecessary files

#rnx

diff --git a/Engine/Plugins/Enterprise/DatasmithCADImporter/Source/CADInterfaces/Private/TechSoftBridge.cpp b/Engine/Plugins/Enterprise/DatasmithCADImporter/Source/CADInterfaces/Private/TechSoftBridge.cpp
--- a/Engine/Plugins/Enterprise/DatasmithCADImporter/Source/CADInterfaces/Private/TechSoftBridge.cpp
+++ b/Engine/Plugins/Enterprise/DatasmithCADImporter/Source/CADInterfaces/Private/TechSoftBridge.cpp
@@ -28,8 +28,6 @@
 
 #include "Utils/StringUtil.h"
 
-#include "UI/Display.h"
-
 #ifdef CADKERNEL_DEV
 #include "CADFileReport.h"
 #endif
@@ -897,11 +895,6 @@ void FTechSoftBridge::AddFace(const A3DTopoFace* A3DFace, UE::CADKernel::EOrient
 
 			if (!ExternalLoop->IsExternal())
 			{
-#ifdef CADKERNEL_DEV
-				FCADFileReport::Get().WrongExternalLoopCount++;
-				FMessage::Printf(Log, TEXT("ERROR: Face %d had wrong external loop\n"), Face->GetId());
-#endif
-
 				for (TSharedPtr<FTopologicalLoop>& Loop : Loops)
 				{
 					if (Loop->IsExternal())
@@ -935,10 +928,6 @@ void FTechSoftBridge::AddFace(const A3DTopoFace* A3DFace, UE::CADKernel::EOrient
 
 	if (Face->GetLoops().Num() == 0)
 	{
-#ifdef CADKERNEL_DEV
-		FMessage::Printf(EVerboseLevel::Log, TEXT("A Face is degenerate, this face is ignored\n"));
-		FCADFileReport::Get().FailedFaceCount++;
-#endif
 		Face->SetAsDegenerated();
 		Face->Delete();
 		return;
diff --git a/Engine/Plugins/Enterprise/DatasmithCADImporter/Source/DatasmithOpenNurbsTranslator/Private/OpenNurbsBRepToCADKernelConverter.cpp b/Engine/Plugins/Enterprise/DatasmithCADImporter/Source/DatasmithOpenNurbsTranslator/Private/OpenNurbsBRepToCADKernelConverter.cpp
--- a/Engine/Plugins/Enterprise/DatasmithCADImporter/Source/DatasmithOpenNurbsTranslator/Private/OpenNurbsBRepToCADKernelConverter.cpp
+++ b/Engine/Plugins/Enterprise/DatasmithCADImporter/Source/DatasmithOpenNurbsTranslator/Private/OpenNurbsBRepToCADKernelConverter.cpp
@@ -295,7 +295,6 @@ TSharedPtr<UE::CADKernel::FTopologicalFace> FOpenNurbsBRepToCADKernelConverter::
 		Face->SetAsDegenerated();
 		Face->Delete();
 
-		FMessage::Printf(EVerboseLevel::Log, TEXT("A Face is degenerate, this face is ignored\n"));
 		return TSharedPtr<FTopologicalFace>();
 	}
 
diff --git a/Engine/Plugins/Enterprise/DatasmithCADImporter/Source/WireInterface/Private/AliasModelToCADKernelConverter.cpp b/Engine/Plugins/Enterprise/DatasmithCADImporter/Source/WireInterface/Private/AliasModelToCADKernelConverter.cpp
--- a/Engine/Plugins/Enterprise/DatasmithCADImporter/Source/WireInterface/Private/AliasModelToCADKernelConverter.cpp
+++ b/Engine/Plugins/Enterprise/DatasmithCADImporter/Source/WireInterface/Private/AliasModelToCADKernelConverter.cpp
@@ -260,7 +260,6 @@ TSharedPtr<FTopologicalFace> FAliasModelToCADKernelConverter::AddTrimRegion(cons
 
 	if (Loops.Num() == 0)
 	{
-		FMessage::Printf(EVerboseLevel::Log, TEXT("The Face %s is degenerate, this face is ignored\n"), TrimRegion.name());
 		return TSharedPtr<FTopologicalFace>();
 	}
 
@@ -275,7 +274,6 @@ TSharedPtr<FTopologicalFace> FAliasModelToCADKernelConverter::AddTrimRegion(cons
 		Face->SetAsDegenerated();
 		Face->Delete();
 
-		FMessage::Printf(EVerboseLevel::Log, TEXT("The Face %s is degenerate, this face is ignored\n"), TrimRegion.name());
 		return TSharedPtr<FTopologicalFace>();
 	}
 
diff --git a/Engine/Source/Runtime/Datasmith/CADKernel/Base/CADKernel.Build.cs b/Engine/Source/Runtime/Datasmith/CADKernel/Base/CADKernel.Build.cs
--- a/Engine/Source/Runtime/Datasmith/CADKernel/Base/CADKernel.Build.cs
+++ b/Engine/Source/Runtime/Datasmith/CADKernel/Base/CADKernel.Build.cs
@@ -10,7 +10,8 @@ namespace UnrealBuildTool.Rules
 		{
 			//OptimizeCode = CodeOptimization.Never;
 			CppCompileWarningSettings.DeterministicWarningLevel = WarningLevel.Off; // __DATE__ in Private/CADKernel/Core/System.cpp
-			
+
+			PublicIncludePaths.Add(Path.Combine(ModuleDirectory, "Private"));
 			PublicDefinitions.Add("CADKERNEL_THINZONE=0");
 
 			PublicDependencyModuleNames.AddRange(
diff --git a/Engine/Source/Runtime/Datasmith/CADKernel/Base/Private/Core/CADEntity.cpp b/Engine/Source/Runtime/Datasmith/CADKernel/Base/Private/Core/CADEntity.cpp
--- a/Engine/Source/Runtime/Datasmith/CADKernel/Base/Private/Core/CADEntity.cpp
+++ b/Engine/Source/Runtime/Datasmith/CADKernel/Base/Private/Core/CADEntity.cpp
@@ -6,7 +6,6 @@
 
 #include "Core/Database.h"
 #include "Core/Session.h"
-#include "Core/System.h"
 #include "Geo/Curves/Curve.h"
 #include "Geo/Surfaces/Surface.h"
 #include "Mesh/Criteria/Criterion.h"
diff --git a/Engine/Source/Runtime/Datasmith/CADKernel/Base/Private/Core/Database.cpp b/Engine/Source/Runtime/Datasmith/CADKernel/Base/Private/Core/Database.cpp
--- a/Engine/Source/Runtime/Datasmith/CADKernel/Base/Private/Core/Database.cpp
+++ b/Engine/Source/Runtime/Datasmith/CADKernel/Base/Private/Core/Database.cpp
@@ -61,17 +61,17 @@ void FDatabase::RemoveEntity(FEntity& Entity)
 	FIdent EntityId = Entity.GetId();
 	if (EntityId == 0)
 	{
-		FMessage::Printf(Debug, TEXT("Warning in FDataBasse::RemoveEntity: the entity is not in the database.\n"), EntityId);
+		UE_LOG(LogCADKernelBase, Verbose, TEXT("Warning in FDataBasse::RemoveEntity: the entity is not in the database.\n"));
 	}
 	else if (EntityId >= (FIdent)DatabaseEntities.Num())
 	{
-		FMessage::Printf(Debug, TEXT("Warning in FDataBasse::RemoveEntity: the Id=%d is not yet defined, the entity has never exist.\n"), EntityId);
+		UE_LOG(LogCADKernelBase, Verbose, TEXT("Warning in FDataBasse::RemoveEntity: the Id=%d is not yet defined, the entity has never exist.\n"), EntityId);
 	}
 	else if (DatabaseEntities[EntityId].IsValid())
 	{
 		if (DatabaseEntities[EntityId].Get() != &Entity)
 		{
-			FMessage::Printf(Debug, TEXT("Warning in FDataBasse::RemoveEntity: the entity with Id=%d is not in the database.\n"), EntityId);
+			UE_LOG(LogCADKernelBase, Verbose, TEXT("Warning in FDataBasse::RemoveEntity: the entity with Id=%d is not in the database.\n"), EntityId);
 			return;
 		}
 		DatabaseEntities[EntityId] = TSharedPtr<FEntity>();
@@ -79,7 +79,7 @@ void FDatabase::RemoveEntity(FEntity& Entity)
 	}
 	else
 	{
-		FMessage::Printf(Debug, TEXT("Warning in FDataBasse::RemoveEntity: the entity with Id=%d is already deleted.\n"), EntityId);
+		UE_LOG(LogCADKernelBase, Verbose, TEXT("Warning in FDataBasse::RemoveEntity: the entity with Id=%d is already deleted.\n"), EntityId);
 	}
 }
 
@@ -87,7 +87,7 @@ void FDatabase::RemoveEntity(FIdent EntityId)
 {
 	if (EntityId >= (FIdent)DatabaseEntities.Num())
 	{
-		FMessage::Printf(Debug, TEXT("Warning in FDataBasse::RemoveEntity: the Id=%d is not yet defined, the entity has never exist.\n"), EntityId);
+		UE_LOG(LogCADKernelBase, Verbose, TEXT("Warning in FDataBasse::RemoveEntity: the Id=%d is not yet defined, the entity has never exist.\n"), EntityId);
 	}
 	else if (DatabaseEntities[EntityId].IsValid())
 	{
@@ -96,7 +96,7 @@ void FDatabase::RemoveEntity(FIdent EntityId)
 	}
 	else
 	{
-		FMessage::Printf(Debug, TEXT("Warning in FDataBasse::RemoveEntity: the entity with Id=%d is already deleted.\n"), EntityId);
+		UE_LOG(LogCADKernelBase, Verbose, TEXT("Warning in FDataBasse::RemoveEntity: the entity with Id=%d is already deleted.\n"), EntityId);
 	}
 }
 
@@ -104,12 +104,12 @@ TSharedPtr<FEntity> FDatabase::GetEntity(FIdent EntityId) const
 {
 	if (EntityId >= (FIdent)DatabaseEntities.Num())
 	{
-		FMessage::Printf(Debug, TEXT("Warning in FDataBasse::GetEntity: the Id=%d is not yet defined\n"), EntityId);
+		UE_LOG(LogCADKernelBase, Verbose, TEXT("Warning in FDataBasse::GetEntity: the Id=%d is not yet defined\n"), EntityId);
 		return TSharedPtr<FEntity>();
 	}
 	if (!DatabaseEntities[EntityId].IsValid())
 	{
-		FMessage::Printf(Debug, TEXT("Warning in FDataBasse::GetEntity: the entity with Id=%d no longer exists.\n"), EntityId);
+		UE_LOG(LogCADKernelBase, Verbose, TEXT("Warning in FDataBasse::GetEntity: the entity with Id=%d no longer exists.\n"), EntityId);
 	}
 	return DatabaseEntities[EntityId];
 }
@@ -296,8 +296,6 @@ void FDatabase::SerializeSelection(FCADKernelArchive& Ar, const TArray<FIdent>&
 	bIsRecursiveSerialization = false;
 
 	RemoveEntity(TmpModel->GetId());
-
-	FMessage::Printf(Log, TEXT("End Serialisation of %d entities\n"), Index);
 }
 
 void FDatabase::Serialize(FCADKernelArchive& Ar)
@@ -314,11 +312,8 @@ void FDatabase::Serialize(FCADKernelArchive& Ar)
 	Ar << DatabaseSize;
 
 	int32 Index = 0;
-	FProgress Progess(DatabaseEntities.Num());
 	for (TSharedPtr<FEntity> Entity : DatabaseEntities)
 	{
-		Progess.Increase();
-
 		if (!Entity.IsValid() || Entity->IsDeleted())
 		{
 			EEntity Type = EEntity::None;
@@ -331,7 +326,6 @@ void FDatabase::Serialize(FCADKernelArchive& Ar)
 		Entity->Serialize(Ar);
 		Index++;
 	}
-	FMessage::Printf(Log, TEXT("End Serialisation of %d %d entity\n"), DatabaseEntities.Num(), Index);
 }
 
 void FDatabase::Deserialize(FCADKernelArchive& Ar)
@@ -352,11 +346,9 @@ void FDatabase::Deserialize(FCADKernelArchive& Ar)
 	ArchiveEntities.Init(TSharedPtr<FEntity>(), ArchiveSize);
 
 	int64 TotalSize = Ar.TotalSize();
-	FProgress Progress(ArchiveSize);
 	int32 Index = 0;
 	for( ; Ar.Tell() < TotalSize; ++Index)
 	{
-		Progress.Increase();
 		TSharedPtr<FEntity> Entity = FEntity::Deserialize(Ar);
 	}
 
@@ -368,8 +360,6 @@ void FDatabase::Deserialize(FCADKernelArchive& Ar)
 
 	ArchiveIdToWaitingSharedPointers.Empty();
 	ArchiveEntities.Empty();
-	
-	FMessage::Printf(Log, TEXT("End Deserialisation of %d %d entity\n"), DatabaseEntities.Num(), ArchiveSize);
 }
 
 void FDatabase::CleanArchiveEntities()
@@ -396,9 +386,6 @@ void FDatabase::CleanArchiveEntities()
 			{
 				TSharedPtr<FVertexLink> VertexLink = StaticCastSharedPtr<FVertexLink>(Entity);
 				VertexLink->CleanLink();
-#ifdef CADKERNEL_DEV
- 				ensureCADKernel(VertexLink->GetTwinEntityCount());
-#endif
 				break;
 			}
 			case EEntity::Body:
@@ -538,19 +525,6 @@ uint32 FDatabase::SpawnEntityIdent(FEntity& Entity, bool bInForceSpawning)
 	return EntityCount;
 }
 
-
-// =========================================================================================================================================================================================================
-// =========================================================================================================================================================================================================
-// =========================================================================================================================================================================================================
-//
-//
-//                                                                            NOT YET REVIEWED
-//
-//
-// =========================================================================================================================================================================================================
-// =========================================================================================================================================================================================================
-// =========================================================================================================================================================================================================
-
 void FDatabase::ExpandSelection(const TArray<TSharedPtr<FEntity>>& Entities, const TSet<EEntity>& Filters, TArray<TSharedPtr<FEntity>>& Selection) const
 {
 	TSet<TSharedPtr<FEntity>> EntitySet;
@@ -630,11 +604,6 @@ void FDatabase::ExpandSelection(TSharedPtr<FEntity> Entity, const TSet<EEntity>&
 	case EEntity::MeshModel:
 	{
 		ensureCADKernel(false);
-		//TSharedPtr<FModelMesh> MeshModel = StaticCastSharedPtr<FModelMesh>(Entity);
-		//for (const FMesh* Mesh : MeshModel->GetMeshes())
-		//{
-		//	ExpandSelection(Mesh, Filter, Selection);
-		//}
 		break;
 	}
 
diff --git a/Engine/Source/Runtime/Datasmith/CADKernel/Base/Private/Core/Factory.h b/Engine/Source/Runtime/Datasmith/CADKernel/Base/Private/Core/Factory.h
--- a/Engine/Source/Runtime/Datasmith/CADKernel/Base/Private/Core/Factory.h
+++ b/Engine/Source/Runtime/Datasmith/CADKernel/Base/Private/Core/Factory.h
@@ -3,9 +3,6 @@
 
 #include "Core/Types.h"
 #include "Core/Group.h"
-#ifdef CADKERNEL_DEV
-#include "UI/DefineForDebug.h"
-#endif
 
 #include "Containers/List.h"
 
diff --git a/Engine/Source/Runtime/Datasmith/CADKernel/Base/Private/Core/Group.cpp b/Engine/Source/Runtime/Datasmith/CADKernel/Base/Private/Core/Group.cpp
--- a/Engine/Source/Runtime/Datasmith/CADKernel/Base/Private/Core/Group.cpp
+++ b/Engine/Source/Runtime/Datasmith/CADKernel/Base/Private/Core/Group.cpp
@@ -2,8 +2,6 @@
 
 #include "Group.h"
 
-#include "Core/System.h"
-
 namespace UE::CADKernel
 { 
 
diff --git a/Engine/Source/Runtime/Datasmith/CADKernel/Base/Private/Core/Session.cpp b/Engine/Source/Runtime/Datasmith/CADKernel/Base/Private/Core/Session.cpp
--- a/Engine/Source/Runtime/Datasmith/CADKernel/Base/Private/Core/Session.cpp
+++ b/Engine/Source/Runtime/Datasmith/CADKernel/Base/Private/Core/Session.cpp
@@ -16,7 +16,7 @@ bool FSession::SaveDatabase(const TCHAR* FileName)
 	TSharedPtr<FCADKernelArchive> Archive = FCADKernelArchive::CreateArchiveWriter(*this, FileName);
 	if (!Archive.IsValid())
 	{
-		FMessage::Printf(Log, TEXT("The archive file %s is corrupted\n"), FileName);
+		UE_LOG(LogCADKernelBase, Warning, TEXT("The archive file %s is corrupted\n"), FileName);
 		return false;
 	}
 
@@ -31,7 +31,7 @@ bool FSession::SaveDatabase(TArray<uint8>& Bytes)
 	TSharedPtr<FCADKernelArchive> Archive = MakeShared<FCADKernelArchive>(*this, Bytes);
 	if (!Archive.IsValid())
 	{
-		FMessage::Printf(Log, TEXT("The archive is corrupted\n"));
+		UE_LOG(LogCADKernelBase, Warning, TEXT("The archive is corrupted\n"));
 		return false;
 	}
 
@@ -89,7 +89,7 @@ bool FSession::LoadDatabase(const TCHAR* FilePath)
 	TSharedPtr<FCADKernelArchive> Archive = FCADKernelArchive::CreateArchiveReader(*this, FilePath);
 	if (!Archive.IsValid())
 	{
-		FMessage::Printf(Log, TEXT("The archive file %s is corrupted\n"), FilePath);
+		UE_LOG(LogCADKernelBase, Warning, TEXT("The archive file %s is corrupted\n"), FilePath);
 		return false;
 	}
 
@@ -115,7 +115,7 @@ bool FSession::LoadDatabase(const TArray<uint8>& Bytes)
 	TSharedPtr<FCADKernelArchive> Archive = MakeShared<FCADKernelArchive>(*this, Bytes);
 	if (!Archive.IsValid())
 	{
-		FMessage::Printf(Log, TEXT("The archive buffer is corrupted\n"));
+		UE_LOG(LogCADKernelBase, Warning, TEXT("The archive buffer is corrupted\n"));
 		return false;
 	}
 
diff --git a/Engine/Source/Runtime/Datasmith/CADKernel/Base/Private/Core/System.cpp b/Engine/Source/Runtime/Datasmith/CADKernel/Base/Private/Core/System.cpp
deleted file mode 100644
diff --git a/Engine/Source/Runtime/Datasmith/CADKernel/Base/Private/Geo/Curves/CompositeCurve.cpp b/Engine/Source/Runtime/Datasmith/CADKernel/Base/Private/Geo/Curves/CompositeCurve.cpp
--- a/Engine/Source/Runtime/Datasmith/CADKernel/Base/Private/Geo/Curves/CompositeCurve.cpp
+++ b/Engine/Source/Runtime/Datasmith/CADKernel/Base/Private/Geo/Curves/CompositeCurve.cpp
@@ -2,7 +2,6 @@
 
 #include "Geo/Curves/CompositeCurve.h"
 
-#include "Core/System.h"
 #include "Math/Boundary.h"
 #include "Geo/GeoPoint.h"
 #include "UI/Message.h"
diff --git a/Engine/Source/Runtime/Datasmith/CADKernel/Base/Private/Geo/Curves/Curve.cpp b/Engine/Source/Runtime/Datasmith/CADKernel/Base/Private/Geo/Curves/Curve.cpp
--- a/Engine/Source/Runtime/Datasmith/CADKernel/Base/Private/Geo/Curves/Curve.cpp
+++ b/Engine/Source/Runtime/Datasmith/CADKernel/Base/Private/Geo/Curves/Curve.cpp
@@ -1,14 +1,12 @@
 // Copyright Epic Games, Inc. All Rights Reserved.
 #include "Geo/Curves/Curve.h"
 
-#include "Core/System.h"
 #include "Geo/Curves/BoundedCurve.h"
 #include "Geo/Sampling/Polyline.h"
 #include "Geo/Sampling/PolylineTools.h"
 #include "Geo/Sampling/SurfacicPolyline.h"
 #include "Geo/Sampler/SamplerOnChord.h"
 #include "Geo/Sampler/SamplerOnParam.h"
-#include "UI/Display.h"
 
 namespace UE::CADKernel
 {
@@ -82,7 +80,7 @@ TSharedPtr<FCurve> FCurve::Rebound(const FLinearBoundary& InBoundary)
 {
 	if (FMath::IsNearlyEqual(InBoundary.Min, GetUMin()) && FMath::IsNearlyEqual(InBoundary.Max, GetUMax()))
 	{
-		FMessage::Printf(Debug, TEXT("Invalid rebound (UMin and UMax are nearly equal) on curve %d\n"), GetId());
+		UE_LOG(LogCADKernelBase, Verbose, TEXT("Invalid rebound (UMin and UMax are nearly equal) on curve %d\n"), GetId());
 		TSharedPtr<FEntity> Entity = AsShared();
 		return StaticCastSharedPtr<FCurve>(Entity);
 	}
@@ -106,13 +104,13 @@ TSharedPtr<FCurve> FCurve::MakeBoundedCurve(const FLinearBoundary& InBoundary)
 
 	if (NewBoundary.IsDegenerated())
 	{
-		FMessage::Printf(Log, TEXT("Invalid bounds (u1=%f u2=%f) on curve %d\n"), NewBoundary.Min, NewBoundary.Max, GetId());
+		UE_LOG(LogCADKernelBase, Warning, TEXT("Invalid bounds (u1=%f u2=%f) on curve %d\n"), NewBoundary.Min, NewBoundary.Max, GetId());
 		return TSharedPtr<FCurve>();
 	}
 
 	if (FMath::IsNearlyEqual(NewBoundary.Min, GetUMin()) && FMath::IsNearlyEqual(NewBoundary.Max, GetUMax()))
 	{
-		FMessage::Printf(Debug, TEXT("Invalid rebound (UMin and UMax are nearly equal) on curve %d\n"), GetId());
+		UE_LOG(LogCADKernelBase, Verbose, TEXT("Invalid rebound (UMin and UMax are nearly equal) on curve %d\n"), GetId());
 		return TSharedPtr<FCurve>();
 	}
 
diff --git a/Engine/Source/Runtime/Datasmith/CADKernel/Base/Private/Geo/Curves/SurfacicCurve.cpp b/Engine/Source/Runtime/Datasmith/CADKernel/Base/Private/Geo/Curves/SurfacicCurve.cpp
--- a/Engine/Source/Runtime/Datasmith/CADKernel/Base/Private/Geo/Curves/SurfacicCurve.cpp
+++ b/Engine/Source/Runtime/Datasmith/CADKernel/Base/Private/Geo/Curves/SurfacicCurve.cpp
@@ -2,7 +2,6 @@
 
 #include "Geo/Curves/SurfacicCurve.h"
 
-#include "Core/System.h"
 #include "Geo/GeoPoint.h"
 #include "Geo/Sampling/Polyline.h"
 #include "Geo/Sampling/SurfacicPolyline.h"
diff --git a/Engine/Source/Runtime/Datasmith/CADKernel/Base/Private/Geo/Sampler/SamplerAbstract.h b/Engine/Source/Runtime/Datasmith/CADKernel/Base/Private/Geo/Sampler/SamplerAbstract.h
--- a/Engine/Source/Runtime/Datasmith/CADKernel/Base/Private/Geo/Sampler/SamplerAbstract.h
+++ b/Engine/Source/Runtime/Datasmith/CADKernel/Base/Private/Geo/Sampler/SamplerAbstract.h
@@ -65,7 +65,7 @@ public:
 	 */
 	virtual void Sample()
 	{
-		SamplingInitalizing();
+		SamplingInitializing();
 		RunSampling();
 	}
 
@@ -90,7 +90,7 @@ protected:
 	/**
 	 * Method to call to generate the curve sampling
 	 */
-	virtual void SamplingInitalizing()
+	virtual void SamplingInitializing()
 	{
 		IsOptimalSegments.Empty(100);
 		Sampling.Empty(100);
@@ -502,8 +502,8 @@ public:
 		double ParamMed = ErrorSum[1] / EvaluationCount;
 		double ParamStandartDeviation = sqrt(SquareErrorSum[1] / EvaluationCount - FMath::Square(ParamMed));
 
-		FMessage::Printf(EVerboseLevel::Log, TEXT("Chord Desired Med Sd Min Max %f %f %f %f %f"), MaxError, ChordMed, ChordStandartDeviation, MinError[0], MaxError[0]);
-		FMessage::Printf(EVerboseLevel::Log, TEXT("Param Desired Med Sd Min Max %f %f %f %f %f"), MaxError, ParamMed, ParamStandartDeviation, MinError[1], MaxError[1]);
+		UE_LOG(LogCADKernelBase, Verbose, TEXT("Chord Desired Med Sd Min Max %f %f %f %f %f"), MaxError, ChordMed, ChordStandartDeviation, MinError[0], MaxError[0]);
+		UE_LOG(LogCADKernelBase, Verbose, TEXT("Param Desired Med Sd Min Max %f %f %f %f %f"), MaxError, ParamMed, ParamStandartDeviation, MinError[1], MaxError[1]);
 	}
 #endif
 };
diff --git a/Engine/Source/Runtime/Datasmith/CADKernel/Base/Private/Geo/Sampler/SamplerOnParam.h b/Engine/Source/Runtime/Datasmith/CADKernel/Base/Private/Geo/Sampler/SamplerOnParam.h
--- a/Engine/Source/Runtime/Datasmith/CADKernel/Base/Private/Geo/Sampler/SamplerOnParam.h
+++ b/Engine/Source/Runtime/Datasmith/CADKernel/Base/Private/Geo/Sampler/SamplerOnParam.h
@@ -245,7 +245,7 @@ public:
 
 		Set(InIsoType, Middle, SurfaceBoundary[Other(InIsoType)]);
 
-		SamplingInitalizing();
+		SamplingInitializing();
 		RunSampling();
 
 		double Step = SurfaceBoundary[InIsoType].Length() / 6;
@@ -285,7 +285,7 @@ public:
 		Swap(SurfaceSampling[Other(InIsoType)], Sampling.Coordinates);
 	}
 
-	virtual void SamplingInitalizing() override
+	virtual void SamplingInitializing() override
 	{
 
 #ifdef DEBUG_CURVE_SAMPLING
diff --git a/Engine/Source/Runtime/Datasmith/CADKernel/Base/Private/Geo/Sampling/SurfacicPolyline.cpp b/Engine/Source/Runtime/Datasmith/CADKernel/Base/Private/Geo/Sampling/SurfacicPolyline.cpp
--- a/Engine/Source/Runtime/Datasmith/CADKernel/Base/Private/Geo/Sampling/SurfacicPolyline.cpp
+++ b/Engine/Source/Runtime/Datasmith/CADKernel/Base/Private/Geo/Sampling/SurfacicPolyline.cpp
@@ -2,7 +2,6 @@
 
 #include "Geo/Sampling/SurfacicPolyline.h"
 
-#include "Core/System.h"
 #include "Geo/Curves/Curve.h"
 #include "Geo/Surfaces/Surface.h"
 #include "Geo/Sampler/SamplerOnParam.h"
diff --git a/Engine/Source/Runtime/Datasmith/CADKernel/Base/Private/Geo/Surfaces/OffsetSurface.cpp b/Engine/Source/Runtime/Datasmith/CADKernel/Base/Private/Geo/Surfaces/OffsetSurface.cpp
--- a/Engine/Source/Runtime/Datasmith/CADKernel/Base/Private/Geo/Surfaces/OffsetSurface.cpp
+++ b/Engine/Source/Runtime/Datasmith/CADKernel/Base/Private/Geo/Surfaces/OffsetSurface.cpp
@@ -2,10 +2,8 @@
 
 #include "Geo/Surfaces/OffsetSurface.h"
 
-#include "Core/System.h"
 #include "Geo/GeoPoint.h"
 #include "Geo/Sampling/SurfacicSampling.h"
-#include "UI/Display.h"
 
 namespace UE::CADKernel
 {
diff --git a/Engine/Source/Runtime/Datasmith/CADKernel/Base/Private/Geo/Surfaces/PlaneSurface.cpp b/Engine/Source/Runtime/Datasmith/CADKernel/Base/Private/Geo/Surfaces/PlaneSurface.cpp
--- a/Engine/Source/Runtime/Datasmith/CADKernel/Base/Private/Geo/Surfaces/PlaneSurface.cpp
+++ b/Engine/Source/Runtime/Datasmith/CADKernel/Base/Private/Geo/Surfaces/PlaneSurface.cpp
@@ -2,7 +2,6 @@
 
 #include "Geo/Surfaces/PlaneSurface.h"
 
-#include "Core/System.h"
 #include "Geo/GeoPoint.h"
 #include "Geo/Sampling/SurfacicSampling.h"
 
diff --git a/Engine/Source/Runtime/Datasmith/CADKernel/Base/Private/Geo/Surfaces/Surface.cpp b/Engine/Source/Runtime/Datasmith/CADKernel/Base/Private/Geo/Surfaces/Surface.cpp
--- a/Engine/Source/Runtime/Datasmith/CADKernel/Base/Private/Geo/Surfaces/Surface.cpp
+++ b/Engine/Source/Runtime/Datasmith/CADKernel/Base/Private/Geo/Surfaces/Surface.cpp
@@ -2,8 +2,6 @@
 
 #include "Geo/Surfaces/Surface.h"
 
-#include "Core/System.h"
-
 #include "Geo/GeoPoint.h"
 #include "Geo/Sampler/SamplerOnParam.h"
 #include "Geo/Sampling/SurfacicPolyline.h"
diff --git a/Engine/Source/Runtime/Datasmith/CADKernel/Base/Private/Math/BSpline.cpp b/Engine/Source/Runtime/Datasmith/CADKernel/Base/Private/Math/BSpline.cpp
--- a/Engine/Source/Runtime/Datasmith/CADKernel/Base/Private/Math/BSpline.cpp
+++ b/Engine/Source/Runtime/Datasmith/CADKernel/Base/Private/Math/BSpline.cpp
@@ -7,7 +7,6 @@
 #include "Geo/Surfaces/NURBSSurface.h"
 #include "Math/MatrixH.h"
 #include "Mesh/Structure/Grid.h"
-#include "UI/Display.h"
 #include "Utils/Util.h"
 
 #include <algorithm>
diff --git a/Engine/Source/Runtime/Datasmith/CADKernel/Base/Private/Math/Geometry.cpp b/Engine/Source/Runtime/Datasmith/CADKernel/Base/Private/Math/Geometry.cpp
--- a/Engine/Source/Runtime/Datasmith/CADKernel/Base/Private/Math/Geometry.cpp
+++ b/Engine/Source/Runtime/Datasmith/CADKernel/Base/Private/Math/Geometry.cpp
@@ -2,7 +2,6 @@
 
 #include "Math/Geometry.h"
 
-#include "Core/System.h"
 #include "Math/Geometry.h"
 #include "Math/Point.h"
 #include "Math/MatrixH.h"
diff --git a/Engine/Source/Runtime/Datasmith/CADKernel/Base/Private/Math/MatrixH.cpp b/Engine/Source/Runtime/Datasmith/CADKernel/Base/Private/Math/MatrixH.cpp
--- a/Engine/Source/Runtime/Datasmith/CADKernel/Base/Private/Math/MatrixH.cpp
+++ b/Engine/Source/Runtime/Datasmith/CADKernel/Base/Private/Math/MatrixH.cpp
@@ -116,20 +116,6 @@ void FMatrixH::Inverse()
 	InverseMatrixN(Matrix, 4);
 }
 
-void FMatrixH::Print(EVerboseLevel level) const
-{
-	FMessage::Printf(level, TEXT(" - Matrix\n"));
-	for (int32 Row = 0; Row < 4; Row++)
-	{
-		FMessage::Printf(level, TEXT("	- "));
-		for (int32 Column = 0; Column < 4; Column++)
-		{
-			FMessage::Printf(level, TEXT("%f "), Get(Row, Column));
-		}
-		FMessage::Printf(level, TEXT("\n"));
-	}
-}
-
 void InverseMatrixN(double* Matrice, int32 Rank)
 {
 	const double One = 1.0;
diff --git a/Engine/Source/Runtime/Datasmith/CADKernel/Base/Private/Mesh/Criteria/CriteriaGrid.cpp b/Engine/Source/Runtime/Datasmith/CADKernel/Base/Private/Mesh/Criteria/CriteriaGrid.cpp
--- a/Engine/Source/Runtime/Datasmith/CADKernel/Base/Private/Mesh/Criteria/CriteriaGrid.cpp
+++ b/Engine/Source/Runtime/Datasmith/CADKernel/Base/Private/Mesh/Criteria/CriteriaGrid.cpp
@@ -4,11 +4,6 @@
 #include "Mesh/Criteria/Criterion.h"
 #include "Geo/Surfaces/Surface.h"
 #include "Topo/TopologicalFace.h"
-#include "UI/Display.h"
-
-#ifdef CADKERNEL_DEV
-#include "UI/DefineForDebug.h"
-#endif
 
 namespace UE::CADKernel
 {
diff --git a/Engine/Source/Runtime/Datasmith/CADKernel/Base/Private/Mesh/Criteria/Criterion.cpp b/Engine/Source/Runtime/Datasmith/CADKernel/Base/Private/Mesh/Criteria/Criterion.cpp
--- a/Engine/Source/Runtime/Datasmith/CADKernel/Base/Private/Mesh/Criteria/Criterion.cpp
+++ b/Engine/Source/Runtime/Datasmith/CADKernel/Base/Private/Mesh/Criteria/Criterion.cpp
@@ -89,7 +89,7 @@ TSharedPtr<FCriterion> FCriterion::CreateCriterion(ECriterion Type, double Value
 	case ECriterion::CADCurvature:
 		return FEntity::MakeShared<FCurvatureCriterion>();
 	default:
-		ERROR_FUNCTION_CALL_NOT_EXPECTED;
+		UE_LOG(LogCADKernelBase, Warning, TEXT("CALL of \"FCriterion::CreateCriterion\" NOT EXPECTED at line %d of the file Criterion.cpp\n"), __LINE__);
 		return TSharedPtr<FCriterion>();
 	}
 }
@@ -107,7 +107,7 @@ double FCriterion::DefaultValue(ECriterion Type)
 	case ECriterion::Angle:
 		return FAngleCriterion::DefaultValue();
 	default:
-		ERROR_FUNCTION_CALL_NOT_EXPECTED;
+		UE_LOG(LogCADKernelBase, Warning, TEXT("CALL of \"FCriterion::DefaultValue\" NOT EXPECTED at line %d of the file Criterion.cpp\n"), __LINE__);
 		return 0.0;
 	}
 }
diff --git a/Engine/Source/Runtime/Datasmith/CADKernel/Base/Private/Mesh/Criteria/SizeCriterion.cpp b/Engine/Source/Runtime/Datasmith/CADKernel/Base/Private/Mesh/Criteria/SizeCriterion.cpp
--- a/Engine/Source/Runtime/Datasmith/CADKernel/Base/Private/Mesh/Criteria/SizeCriterion.cpp
+++ b/Engine/Source/Runtime/Datasmith/CADKernel/Base/Private/Mesh/Criteria/SizeCriterion.cpp
@@ -1,7 +1,6 @@
 // Copyright Epic Games, Inc. All Rights Reserved.
 #include "Mesh/Criteria/SizeCriterion.h"
 
-#include "Core/System.h"
 #include "Geo/GeoPoint.h"
 #include "Topo/TopologicalEdge.h"
 
diff --git a/Engine/Source/Runtime/Datasmith/CADKernel/Base/Private/Mesh/Meshers/BowyerWatsonTriangulator.h b/Engine/Source/Runtime/Datasmith/CADKernel/Base/Private/Mesh/Meshers/BowyerWatsonTriangulator.h
--- a/Engine/Source/Runtime/Datasmith/CADKernel/Base/Private/Mesh/Meshers/BowyerWatsonTriangulator.h
+++ b/Engine/Source/Runtime/Datasmith/CADKernel/Base/Private/Mesh/Meshers/BowyerWatsonTriangulator.h
@@ -4,10 +4,6 @@
 #include "Core/Types.h"
 #include "Math/SlopeUtils.h"
 #include "Mesh/Meshers/IsoTriangulator/IsoCell.h"
-#include "UI/Display.h"
-#ifdef CADKERNEL_DEV
-#include "UI/DefineForDebug.h"
-#endif
 
 namespace UE::CADKernel
 {
@@ -54,36 +50,12 @@ public:
 
 	void Triangulate()
 	{
-		FTimePoint StartTime = FChrono::Now();
-#ifdef DEBUG_BOWYERWATSON
-		F3DDebugSession DelaunayDebugSession(bDisplay, TEXT("Delaunay Algo"));
-#endif // DEBUG_DELAUNAY
-
 		Vertices.Sort([](const TPair<int32, FVector2d>& Vertex1, const TPair<int32, FVector2d>& Vertex2) {return (Vertex1.Value.X + Vertex1.Value.Y) > (Vertex2.Value.X + Vertex2.Value.Y); });
 
-#ifdef DEBUG_BOWYERWATSON
-		if (bDisplay)
-		{
-			F3DDebugSession _(TEXT("Vertices"));
-			for (const TPair<int32, FVector2d>& Vertex : Vertices)
-			{
-				//F3DDebugSession _(TEXT("Vertex"));
-				DisplayPoint2DWithScale(Vertex.Value , EVisuProperty::YellowPoint, Vertex.Key);
-			}
-			//Wait();
-		}
-#endif
 		// initialization of Bowyer & Watson algorithm with a bounding mesh of the vertex cloud 
 		// i.e. 2 triangles defined by the corners of the offset vertices bounding box
 		MakeBoundingMesh();
 
-#ifdef DEBUG_BOWYERWATSON_STEP
-		DisplayTriangles();
-		Wait(bDisplay);
-
-		static int32 StopIndex = 206;
-#endif
-
 		EdgeVertexIndices.Reserve(TriangleSet.Num() * 6);
 
 		// insert each point in the mesh
@@ -105,16 +77,6 @@ public:
 
 			const FVector2d& NewVertex = Vertices[VertexIndex].Value;
 
-#ifdef DEBUG_BOWYERWATSON_STEP
-			F3DDebugSession _(bDisplay, FString::Printf(TEXT("Step %d"), VIndex));
-			if (bDisplay)
-			{
-				F3DDebugSession _(TEXT("New Point"));
-				DisplayPoint(NewVertex * DisplayScale, EVisuProperty::YellowPoint, Vertices[VertexIndex].Key);
-				Wait(VIndex >= StopIndex);
-			}
-#endif
-
 			bool bHaveDoubt = false;
 			
 			// find all triangles whose circumcircles contain the new vertex
@@ -165,13 +127,6 @@ public:
 				}
 			}
 
-#ifdef DEBUG_BOWYERWATSON_STEP
-			if (bHaveDoubt && (VIndex >= StopIndex))
-			{
-				F3DDebugSession A(bDisplay, TEXT("HaveDoubt"));
-				DisplaySelectedTrianglesAndComplements();
-			}
-#endif
 			if (!bHaveDoubt)
 			{
 				bHaveDoubt = !DoSelectedTrianglesFormSinglePartition();
@@ -207,14 +162,6 @@ public:
 				}
 			}
 
-#ifdef DEBUG_BOWYERWATSON_STEP
-			if (VIndex == StopIndex)
-			{
-				DisplaySelectedTrianglesAndComplements(); 
-				Wait();
-			}
-#endif
-
 			if (AdditionalTriangleIndices.Num() > 0)
 			{
 				// Additional triangles are triangles that the new vertex is nearly coincident to the circle.
@@ -279,22 +226,6 @@ public:
 
 						FVector2d CandidateVertex = Vertices[CandidateVertexIndex].Value;
 
-#ifdef DEBUG_BOWYERWATSON_STEP_2
-						if (VIndex == StopIndex)
-						{
-							DisplaySelectedTrianglesBoundary(BoundaryVertexToSlope);
-							{
-								F3DDebugSession A(TEXT("New Sel triangle"));
-								DisplayTriangle(TriangleIndex, EVisuProperty::YellowCurve);
-							}
-							{
-								F3DDebugSession A(TEXT("Candidate Point"));
-								DisplayPoint(Vertices[CandidateVertexIndex].Value * DisplayScale, EVisuProperty::BluePoint);
-							}
-							Wait();
-						}
-#endif
-
 						double NewVertexSlope = ComputeSlope(NewVertex, CandidateVertex);
 
 						int32 StartIndex = -1;
@@ -318,14 +249,6 @@ public:
 						double SlopeNewEndMinusPi = ComputeSlope(NewVertex, EndPoint) - Slope::RightSlope;
 						double SlopeEndCandidate = ComputePositiveSlope(EndPoint, CandidateVertex, SlopeNewEndMinusPi);
 
-#ifdef DEBUG_BOWYERWATSON_STEP_2
-						if (VIndex == StopIndex)
-						{
-							DisplayLocalVoronoiDiagram(CandidateVertex, NewVertex, StartPoint, EndPoint);
-							Wait();
-						}
-#endif
-
 						if (SlopeStartCandidate > 0 && SlopeStartCandidate < Slope::RightSlope && SlopeEndCandidate > Slope::ThreeRightSlope && SlopeEndCandidate < Slope::TwoPiSlope)
 						{
 							bTriangleHasBeenAdded = true;
@@ -339,52 +262,17 @@ public:
 							}
 							BoundaryVertex.Add(CandidateVertexIndex);
 							SelectedTriangleIndices.Add(TriangleIndex);
-
-#ifdef DEBUG_BOWYERWATSON_STEP_2
-							if (VIndex == StopIndex)
-							{
-								DisplaySelectedTriangles();
-								Wait();
-							}
-#endif
 						}
 						TriangleIndex = -1;
 					}
 				}
 			}
 
-#ifdef DEBUG_BOWYERWATSON_STEP
-			if (VIndex >= StopIndex)
-			{
-				DisplaySelectedTriangles();
-				Wait();
-			}
-#endif
-
 			EdgeVertexIndices.Reset(VerticesCount);
  			FindBoundaryEdgesOfSelectedTriangles(EdgeVertexIndices);
 
 			// make the new triangles : Each new triangle is defined by an edge of the boundary and the new vertex 
-#ifdef DEBUG_BOWYERWATSON_STEP
-			if (bDisplay && VIndex >= StopIndex)
 			{
-				F3DDebugSession _(TEXT("To remesh"));
-				for (int32 EdgeIndex = 0; EdgeIndex < EdgeVertexIndices.Num(); EdgeIndex += 2)
-				{
-					if (EdgeVertexIndices[EdgeIndex] < 0)
-					{
-						continue;
-					}
-					DisplaySegment(Vertices[EdgeVertexIndices[EdgeIndex]].Value * DisplayScale, Vertices[EdgeVertexIndices[EdgeIndex + 1]].Value * DisplayScale, 0, EVisuProperty::YellowCurve);
-				}
-				Wait();
-			}
-#endif
-			{
-#ifdef DEBUG_BOWYERWATSON_STEP
-				F3DDebugSession _(bDisplay && VIndex >= StopIndex, TEXT("New Triangles"));
-#endif
-
 				// The deleted triangles are replaced by the new ones
 				int32 EdgeIndex = 0;
 				for (int32 TriangleIndex : SelectedTriangleIndices)
@@ -394,14 +282,6 @@ public:
 						EdgeIndex += 2;
 					}
 					TriangleSet[TriangleIndex].Set(EdgeVertexIndices[EdgeIndex + 1], EdgeVertexIndices[EdgeIndex], VertexIndex, Vertices);
-#ifdef DEBUG_BOWYERWATSON_STEP
-					if (bDisplay && VIndex >= StopIndex)
-					{
-						//F3DDebugSession A(TEXT("Triangle"));
-						DisplayTriangle(TriangleIndex, EVisuProperty::BlueCurve);
-						//Wait();
-					}
-#endif
 
 					EdgeIndex += 2;
 				}
@@ -413,35 +293,10 @@ public:
 						continue;
 					}
 					TriangleSet.Emplace(EdgeVertexIndices[EdgeIndex + 1], EdgeVertexIndices[EdgeIndex], VertexIndex, Vertices);
-#ifdef DEBUG_BOWYERWATSON_STEP
-					if (bDisplay && VIndex >= StopIndex)
-					{
-						//F3DDebugSession A(TEXT("Triangle"));
-						DisplayTriangle(TriangleSet.Num() - 1, EVisuProperty::BlueCurve);
-						//Wait();
-					}
-#endif
 				}
 			}
-
-#ifdef DEBUG_BOWYERWATSON_STEP
-			if (bDisplay)
-			{
-				DisplayTriangles();
-				Wait(VIndex >= StopIndex);
-			}
-#endif
 		}
 
-#ifdef DEBUG_BOWYERWATSON
-		if (bDisplay)
-		{
-			// The final mesh
-			DisplayTriangles();
-			Wait(bDisplay);
-		}
-#endif
-
 		// Find all Edges and their type (inner edge or boundary edge)
 		EdgeVertexIndices.Reset(TriangleSet.Num() * 6);
 
@@ -490,10 +345,6 @@ public:
 				EndVertex = StartVertex;
 			}
 		}
-#ifdef DEBUG_BOWYERWATSON
-		DisplayEdges();
-		//Wait();
-#endif
 
 		// the bounding mesh vertices are removed
 		Vertices.SetNum(VerticesCount);
@@ -806,162 +657,6 @@ private:
 			}
 		}
 	}
-
-#ifdef DEBUG_BOWYERWATSON
-	void DisplayEdges()
-	{
-		if (!bDisplay)
-		{
-			return;
-		}
-
-		F3DDebugSession _(TEXT("Edges"));
-		for (int32 Index = 0; Index < EdgeVertexIndices.Num(); Index += 2)
-		{
-			if (EdgeInstanceCount[Index / 2] < 2)
-			{
-				DisplaySegment(Vertices[EdgeVertexIndices[Index]].Value * DisplayScale, Vertices[EdgeVertexIndices[Index + 1]].Value * DisplayScale, 0, EVisuProperty::YellowCurve);
-			}
-			else
-			{
-				DisplaySegment(Vertices[EdgeVertexIndices[Index]].Value * DisplayScale, Vertices[EdgeVertexIndices[Index + 1]].Value * DisplayScale, 0, EVisuProperty::PurpleCurve);
-			}
-		}
-	};
-
-	void DisplayTriangles()
-	{
-		if (!bDisplay)
-		{
-			return;
-		}
-
-		F3DDebugSession _(TEXT("Triangles"));
-		for (int32 Index = 0; Index < TriangleSet.Num(); Index++)
-		{
-			DisplayTriangle(Index, EVisuProperty::BlueCurve);
-		}
-		//Wait();
-	};
-
-	void DisplaySelectedTriangles()
-	{
-		if (!bDisplay)
-		{
-			return;
-		}
-
-		F3DDebugSession _(TEXT("Selected Triangles"));
-		for (int32 Index : SelectedTriangleIndices)
-		{
-			//F3DDebugSession _(TEXT("Triangle"));
-			DisplayTriangle(Index, EVisuProperty::BlueCurve);
-		}
-		//Wait();
-	};
-
-	void DisplaySelectedTrianglesAndComplements()
-	{
-		if (!bDisplay)
-		{
-			return;
-		}
-
-		DisplaySelectedTriangles(); 
-
-		{
-			F3DDebugSession A(bDisplay, TEXT("Select Triangles Add"));
-			for (int32 TriangleIndex : AdditionalTriangleIndices)
-			{
-				if (TriangleIndex < 0)
-				{
-					continue;
-				}
-				DisplayTriangle(TriangleIndex, EVisuProperty::YellowCurve);
-			}
-		}
-	}
-
-
-	void DisplayTriangle(int32 Index, EVisuProperty Property)
-	{
-		if (!bDisplay)
-		{
-			return;
-		}
-
-		DisplaySegment(Vertices[TriangleSet[Index].VertexIndices[0]].Value * DisplayScale, Vertices[TriangleSet[Index].VertexIndices[1]].Value * DisplayScale, 0, Property);
-		DisplaySegment(Vertices[TriangleSet[Index].VertexIndices[1]].Value * DisplayScale, Vertices[TriangleSet[Index].VertexIndices[2]].Value * DisplayScale, 0, Property);
-		DisplaySegment(Vertices[TriangleSet[Index].VertexIndices[2]].Value * DisplayScale, Vertices[TriangleSet[Index].VertexIndices[0]].Value * DisplayScale, 0, Property);
-		//Wait();
-	};
-
-	void DisplaySelectedTrianglesBoundary(TArray<TPair<int32, double>> VertexToSlop)
-	{
-		F3DDebugSession A(TEXT("Selected Triangles Boundary"));
-
-		FVector2d Previous = Vertices[VertexToSlop[VertexToSlop.Num() - 1].Key].Value;
-		for (int Index = 0; Index < VertexToSlop.Num(); ++Index)
-		{
-			FVector2d VertexPoint = Vertices[VertexToSlop[Index].Key].Value;
-			DisplayPoint(VertexPoint * DisplayScale, EVisuProperty::GreenPoint);
-			DisplaySegment(VertexPoint * DisplayScale, Previous * DisplayScale, 0, EVisuProperty::GreenCurve);
-			Previous = VertexPoint;
-		}
-	}
-
-	void DisplayLocalVoronoiDiagram(const FVector2d& CandidateVertex, const FVector2d& NewVertex, const FVector2d& StartPoint, const FVector2d& EndPoint)
-	{
-		F3DDebugSession B(TEXT("Voronoi Diagram"));
-		{
-			F3DDebugSession A(TEXT("New	Vertex"));
-			DisplayPoint(NewVertex * DisplayScale, EVisuProperty::RedPoint);
-		}
-		{
-			F3DDebugSession A(TEXT("Candidate"));
-			DisplayPoint(CandidateVertex * DisplayScale, EVisuProperty::BluePoint);
-		}
-		{
-			F3DDebugSession A(TEXT("Start"));
-			DisplayPoint(StartPoint * DisplayScale, EVisuProperty::YellowPoint);
-
-			FVector2d Normal = NewVertex - StartPoint;
-			constexpr double HalfPi = -DOUBLE_PI / 2.;
-			Normal.Normalize();
-			FVector2d Tangent = Normal.Rotate(HalfPi) * 1000.;
-			FVector2d TangentPoint = StartPoint + Tangent;
-
-			DisplaySegment(StartPoint * DisplayScale, TangentPoint * DisplayScale, 1, EVisuProperty::YellowCurve);
-			DisplaySegment(NewVertex * DisplayScale, StartPoint * DisplayScale, 1, EVisuProperty::YellowCurve);
-		}
-
-		{
-			F3DDebugSession A(TEXT("Start - Candidate"));
-			DisplaySegment(StartPoint * DisplayScale, CandidateVertex * DisplayScale, 1, EVisuProperty::YellowCurve);
-		}
-
-		{
-			F3DDebugSession A(TEXT("End"));
-			DisplayPoint(EndPoint * DisplayScale, EVisuProperty::OrangePoint);
-
-			FVector2d Normal = NewVertex - EndPoint;
-			constexpr double HalfPi = DOUBLE_PI / 2.;
-			Normal.Normalize();
-			FVector2d Tangent = Normal.Rotate(HalfPi) * 1000.;
-			FVector2d TangentPoint = EndPoint + Tangent;
-
-			DisplaySegment(EndPoint * DisplayScale, TangentPoint * DisplayScale, 1, EVisuProperty::OrangeCurve);
-			DisplaySegment(NewVertex * DisplayScale, EndPoint * DisplayScale, 1, EVisuProperty::OrangeCurve);
-		}
-
-		{
-			F3DDebugSession A(TEXT("End - Candidate"));
-			DisplaySegment(EndPoint * DisplayScale, CandidateVertex * DisplayScale, 1, EVisuProperty::OrangeCurve);
-		}
-	}
-
-#endif
-
 };
 
 }
\ No newline at end of file
diff --git a/Engine/Source/Runtime/Datasmith/CADKernel/Base/Private/Mesh/Meshers/CycleTriangulator.cpp b/Engine/Source/Runtime/Datasmith/CADKernel/Base/Private/Mesh/Meshers/CycleTriangulator.cpp
--- a/Engine/Source/Runtime/Datasmith/CADKernel/Base/Private/Mesh/Meshers/CycleTriangulator.cpp
+++ b/Engine/Source/Runtime/Datasmith/CADKernel/Base/Private/Mesh/Meshers/CycleTriangulator.cpp
@@ -10,11 +10,6 @@
 #include "Mesh/Structure/FaceMesh.h"
 #include "Topo/TopologicalFace.h"
 
-#ifdef CADKERNEL_DEV
-#include "UI/DefineForDebug.h"
-#include "Mesh/Meshers/MesherReport.h"
-#endif
-
 namespace UE::CADKernel
 {
 
@@ -44,9 +39,6 @@ bool FCycleTriangulator::MeshCycle()
 	// Check if the cycle is in self intersecting and fix it. 
 	if (!CanCycleBeMeshed())
 	{
-#ifdef CADKERNEL_DEV
-		FMesherReport::Get().Logs.AddCycleMeshingFailure();
-#endif
 		return false;
 	}
 
@@ -1340,16 +1332,6 @@ bool  FCycleTriangulator::BuildTriangle(TArray<FIsoNode*>& CandidatNodes)
 
 	Mesh.AddTriangle(CandidatNodes[0]->GetGlobalIndex(), CandidatNodes[1]->GetGlobalIndex(), CandidatNodes[2]->GetGlobalIndex());
 
-#ifdef ADD_TRIANGLE_2D
-	if (Grid.bDisplay)
-	{
-		{
-			F3DDebugSession _(FString::Printf(TEXT("Triangle")));
-			Grid.DisplayTriangle(EGridSpace::UniformScaled, *CandidatNodes[0], *CandidatNodes[1], *CandidatNodes[2]);
-		}
-	}
-#endif 
-
 	SortCycleIntersectionToolIfNeeded();
 	return true;
 }
@@ -1385,13 +1367,6 @@ namespace Polygon
 
 void MeshTriangle(const FGrid& Grid, FIsoNode** Nodes, FFaceMesh& Mesh)
 {
-#ifdef ADD_TRIANGLE_2D
-	if (Grid.bDisplay)
-	{
-		F3DDebugSession G(TEXT("Mesh Triangle"));
-		Grid.DisplayTriangle(EGridSpace::UniformScaled, *Nodes[0], *Nodes[1], *Nodes[2]);
-	}
-#endif 
 	Mesh.AddTriangle(Nodes[0]->GetGlobalIndex(), Nodes[1]->GetGlobalIndex(), Nodes[2]->GetGlobalIndex());
 }
 
@@ -1431,16 +1406,6 @@ void MeshQuadrilateral(const FGrid& Grid, FIsoNode** Nodes, FFaceMesh& Mesh)
 		NodeIndices[IndexN] = NodeIndices[IndexN - 1] == 3 ? 0 : NodeIndices[IndexN - 1] + 1;
 	}
 
-#ifdef ADD_TRIANGLE_2D
-	if (Grid.bDisplay)
-	{
-		F3DDebugSession G(TEXT("Mesh Quadrilateral"));
-		Grid.DisplayTriangle(EGridSpace::UniformScaled, *Nodes[NodeIndices[1]], *Nodes[NodeIndices[3]], *Nodes[NodeIndices[0]]);
-		Grid.DisplayTriangle(EGridSpace::UniformScaled, *Nodes[NodeIndices[1]], *Nodes[NodeIndices[2]], *Nodes[NodeIndices[3]]);
-		Grid.DisplayIsoSegment(EGridSpace::UniformScaled, *Nodes[NodeIndices[1]], *Nodes[NodeIndices[3]], 0, EVisuProperty::RedCurve);
-	}
-#endif 
-
 	Mesh.AddTriangle(Nodes[NodeIndices[1]]->GetGlobalIndex(), Nodes[NodeIndices[3]]->GetGlobalIndex(), Nodes[NodeIndices[0]]->GetGlobalIndex());
 	Mesh.AddTriangle(Nodes[NodeIndices[1]]->GetGlobalIndex(), Nodes[NodeIndices[2]]->GetGlobalIndex(), Nodes[NodeIndices[3]]->GetGlobalIndex());
 }
@@ -1612,22 +1577,6 @@ void MeshPentagon(const FGrid& Grid, FIsoNode** Nodes, FFaceMesh& Mesh)
 		break;
 	}
 
-#ifdef ADD_TRIANGLE_2D
-	if (Grid.bDisplay)
-	{
-		F3DDebugSession G(TEXT("Mesh Pentagon"));
-		for (int32 Index = 0; Index < 3; ++Index)
-		{
-			//F3DDebugSession G(TEXT("Triangle"));
-			Grid.DisplayTriangle(EGridSpace::UniformScaled, *Nodes[TrianglesOfPentagon[SelectedTriangles[Index]][0]], *Nodes[TrianglesOfPentagon[SelectedTriangles[Index]][1]], *Nodes[TrianglesOfPentagon[SelectedTriangles[Index]][2]]);
-			Grid.DisplayIsoSegment(EGridSpace::UniformScaled, *Nodes[TrianglesOfPentagon[SelectedTriangles[Index]][0]], *Nodes[TrianglesOfPentagon[SelectedTriangles[Index]][1]], 0, EVisuProperty::BlueCurve);
-			Grid.DisplayIsoSegment(EGridSpace::UniformScaled, *Nodes[TrianglesOfPentagon[SelectedTriangles[Index]][1]], *Nodes[TrianglesOfPentagon[SelectedTriangles[Index]][2]], 0, EVisuProperty::BlueCurve);
-			Grid.DisplayIsoSegment(EGridSpace::UniformScaled, *Nodes[TrianglesOfPentagon[SelectedTriangles[Index]][2]], *Nodes[TrianglesOfPentagon[SelectedTriangles[Index]][0]], 0, EVisuProperty::BlueCurve);
-		}
-		Wait(false);
-	}
-#endif 
-
 	// Add the selected triangles to the mesh
 	for (int32 Index = 0; Index < 3; ++Index)
 	{
@@ -1645,7 +1594,7 @@ bool FCycleTriangulator::CanCycleBeMeshed()
 	{
 		if (CycleIntersectionTool.DoesIntersect(*Segment))
 		{
-			FMessage::Printf(Log, TEXT("A cycle of the surface %d is in self intersecting. The mesh of this sector is canceled.\n"), Grid.GetFace().GetId());
+			UE_LOG(LogCADKernelBase, Warning, TEXT("A cycle of the surface %d is in self intersecting. The mesh of this sector is canceled.\n"), Grid.GetFace().GetId());
 
 #ifdef DEBUG_CAN_CYCLE_BE_FIXED_AND_MESHED
 			if (bDisplay)
@@ -1854,9 +1803,6 @@ bool FCycleTriangulator::BuildSegmentIfNeeded(TArray<FIsoNode*>& CandidatNodes)
 
 		if (!BuildSegmentIfNeeded(StartNode, EndNode))
 		{
-#ifdef CADKERNEL_DEV
-			FMesherReport::GetLogs().AddCycleMeshingFailure();
-#endif
 			return false;
 		}
 		StartNode = EndNode;
@@ -1878,11 +1824,6 @@ bool FCycleTriangulator::BuildSegmentIfNeeded(FIsoNode* NodeA, FIsoNode* NodeB,
 		{
 			if (ABSegment->HasTriangleOnLeft())
 			{
-#ifdef DEBUG_BUILD_SEGMENT_IF_NEEDED
-				F3DDebugSession _(FString::Printf(TEXT("ERROR Segment")));
-				Grid.DisplayIsoSegment(EGridSpace::UniformScaled, *NodeA, *NodeB, 0, EVisuProperty::YellowCurve);
-				Wait();
-#endif
 				return false;
 			}
 			ABSegment->SetHasTriangleOnLeft();
@@ -1891,11 +1832,6 @@ bool FCycleTriangulator::BuildSegmentIfNeeded(FIsoNode* NodeA, FIsoNode* NodeB,
 		{
 			if (ABSegment->HasTriangleOnRight())
 			{
-#ifdef DEBUG_BUILD_SEGMENT_IF_NEEDED
-				F3DDebugSession _(FString::Printf(TEXT("ERROR Segment")));
-				Grid.DisplayIsoSegment(EGridSpace::UniformScaled, *NodeA, *NodeB, 0, EVisuProperty::YellowCurve);
-				Wait();
-#endif
 				return false;
 			}
 			ABSegment->SetHasTriangleOnRight();
@@ -1903,24 +1839,10 @@ bool FCycleTriangulator::BuildSegmentIfNeeded(FIsoNode* NodeA, FIsoNode* NodeB,
 	}
 	else
 	{
-#ifdef ADD_TRIANGLE_2D_
-		if (Grid.bDisplay)
-		{
-			F3DDebugSession _(FString::Printf(TEXT("New Segment")));
-			Grid.DisplayIsoSegment(EGridSpace::UniformScaled, *NodeA, *NodeB, 0, EVisuProperty::YellowCurve);
-			Grid.DisplayIsoNode(EGridSpace::UniformScaled, *NodeB, 0, EVisuProperty::RedPoint);
-		}
-#endif
-
 		FIsoSegment& NewSegment = IsoSegmentFactory.New();
 		NewSegment.Init(*NodeA, *NodeB, ESegmentType::Unknown);
 		if (!NewSegment.ConnectToNode())
 		{
-#ifdef DEBUG_BUILD_SEGMENT_IF_NEEDED
-			F3DDebugSession _(FString::Printf(TEXT("ERROR Segment")));
-			Grid.DisplayIsoSegment(EGridSpace::UniformScaled, *NodeA, *NodeB, 0, EVisuProperty::RedCurve);
-			Wait();
-#endif
 			IsoSegmentFactory.DeleteEntity(&NewSegment);
 			return false;
 		}
diff --git a/Engine/Source/Runtime/Datasmith/CADKernel/Base/Private/Mesh/Meshers/CycleTriangulator.h b/Engine/Source/Runtime/Datasmith/CADKernel/Base/Private/Mesh/Meshers/CycleTriangulator.h
--- a/Engine/Source/Runtime/Datasmith/CADKernel/Base/Private/Mesh/Meshers/CycleTriangulator.h
+++ b/Engine/Source/Runtime/Datasmith/CADKernel/Base/Private/Mesh/Meshers/CycleTriangulator.h
@@ -8,10 +8,6 @@
 #include "Mesh/Meshers/IsoTriangulator/IntersectionSegmentTool.h"
 #include "Mesh/Meshers/IsoTriangulator/IntersectionIsoSegmentTool.h"
 
-#ifdef CADKERNEL_DEV
-#include "UI/DefineForDebug.h"
-#endif
-
 namespace UE::CADKernel
 {
 class FGrid;
diff --git a/Engine/Source/Runtime/Datasmith/CADKernel/Base/Private/Mesh/Meshers/IsoTriangulator.cpp b/Engine/Source/Runtime/Datasmith/CADKernel/Base/Private/Mesh/Meshers/IsoTriangulator.cpp
--- a/Engine/Source/Runtime/Datasmith/CADKernel/Base/Private/Mesh/Meshers/IsoTriangulator.cpp
+++ b/Engine/Source/Runtime/Datasmith/CADKernel/Base/Private/Mesh/Meshers/IsoTriangulator.cpp
@@ -17,7 +17,6 @@
 #include "Mesh/Structure/ThinZone2D.h"
 #include "Topo/TopologicalEdge.h"
 #include "Topo/TopologicalFace.h"
-#include "UI/Display.h"
 #include "Utils/ArrayUtils.h"
 
 namespace UE::CADKernel
@@ -84,11 +83,6 @@ FIsoTriangulator::FIsoTriangulator(FGrid& InGrid, FFaceMesh& OutMesh, const FMes
 bool FIsoTriangulator::Triangulate()
 {
 	EGridSpace DisplaySpace = EGridSpace::UniformScaled;
-#ifdef DEBUG_ISOTRIANGULATOR
-	F3DDebugSession _(bDisplay, FString::Printf(TEXT("Triangulate %d"), Grid.GetFace().GetId()));
-#endif
-
-	FTimePoint StartTime = FChrono::Now();
 
 	// =============================================================================================================
 	// Build the first elements (IsoNodes (i.e. Inner nodes), Loops nodes, and knows segments 
@@ -96,35 +90,16 @@ bool FIsoTriangulator::Triangulate()
 
 	BuildNodes();
 
-#ifdef DEBUG_ISOTRIANGULATOR
-	DisplayIsoNodes(DisplaySpace);
-#endif
-
 	FillMeshNodes();
 	BuildLoopSegments();
 
-#ifdef DEBUG_ISOTRIANGULATOR
-	DisplayLoops(TEXT("FIsoTrianguler::LoopSegments"), true, false);
-	Wait(false);
-#endif
-
 	FLoopCleaner LoopCleaner(*this);
 	if (!LoopCleaner.Run())
 	{
-#ifdef CADKERNEL_DEV
-		FMesherReport::Get().Logs.AddDegeneratedLoop();
-#endif
-		FMessage::Printf(EVerboseLevel::Log, TEXT("The meshing of the surface %d failed due to a degenerated loop\n"), Grid.GetFace().GetId());
+		UE_LOG(LogCADKernelBase, Warning, TEXT("The meshing of the surface %d failed due to a degenerated loop\n"), Grid.GetFace().GetId());
 		return false;
 	}
 
-#ifdef DEBUG_ISOTRIANGULATOR
-	DisplayLoops(TEXT("FIsoTrianguler::LoopSegments after fix intersection"), true, false);
-	Grid.DisplayIsoSegments(TEXT("FIsoTrianguler::LoopSegments orientation"), DisplaySpace, LoopSegments, false, true, false, EVisuProperty::YellowCurve);
-	Grid.DisplayIsoSegments(TEXT("FIsoTrianguler::LoopSegments"), DisplaySpace, LoopSegments, true, false, false, EVisuProperty::YellowCurve);
-	Wait(false);
-#endif
-
 	//Fill Intersection tool
 	LoopSegmentsIntersectionTool.Empty(LoopSegments.Num());
 	LoopSegmentsIntersectionTool.AddSegments(LoopSegments);
@@ -137,30 +112,13 @@ bool FIsoTriangulator::Triangulate()
 
 	FinalToLoops.Append(ThinZoneSegments);
 
-#ifdef DEBUG_ISOTRIANGULATOR
-	Grid.DisplayIsoSegments(TEXT("FIsoTrianguler::ThinZoneSegments"), DisplaySpace, ThinZoneSegments, false);
-	LoopSegmentsIntersectionTool.Display(Grid.bDisplay, TEXT("FIsoTrianguler::IntersectionTool Loop"));
-#endif
 
 	BuildInnerSegments();
 
-#ifdef DEBUG_ISOTRIANGULATOR
-	if (bDisplay)
-	{
-		F3DDebugSession A(TEXT("FIsoTrianguler::FinalInnerSegments"));
-		Grid.DisplayIsoSegments(DisplaySpace, FinalInnerSegments, true, false, false, EVisuProperty::BlueCurve);
-		//Wait();
-	}
-	InnerToOuterIsoSegmentsIntersectionTool.Display(bDisplay, TEXT("FIsoTrianguler::IntersectionTool InnerToOuter"), EVisuProperty::RedCurve);
-#endif
-
 	// =============================================================================================================
 	// =============================================================================================================
 
 	BuildInnerSegmentsIntersectionTool();
-#ifdef DEBUG_ISOTRIANGULATOR
-	InnerSegmentsIntersectionTool.Display(Grid.bDisplay, TEXT("FIsoTrianguler::IntersectionTool Inner"));
-#endif
 
 	// =============================================================================================================
 	// 	   For each cell
@@ -169,22 +127,11 @@ bool FIsoTriangulator::Triangulate()
 	// 	           - build Delaunay connection
 	// 	           - find the shortest segment to connect each connected loop by Delaunay
 	// =============================================================================================================
-#ifdef DEBUG_ISOTRIANGULATOR
-	Grid.DisplayIsoSegments(TEXT("FIsoTrianguler::Final To Loops Before"), DisplaySpace, FinalToLoops, false, false, false, EVisuProperty::YellowCurve);
-#endif
 
 	ConnectCellLoops();
 
-#ifdef DEBUG_ISOTRIANGULATOR
-	Grid.DisplayIsoSegments(TEXT("FIsoTrianguler::Final Iso ToLink Inner To Loops"), DisplaySpace, FinalToLoops, false, false, true, EVisuProperty::YellowCurve);
-#endif
-
 	SelectSegmentsToLinkInnerToLoop();
 
-#ifdef DEBUG_ISOTRIANGULATOR
-	Grid.DisplayIsoSegments(TEXT("FIsoTrianguler::Final Segments"), DisplaySpace, FinalToLoops, false, false, true, EVisuProperty::YellowCurve);
-#endif
-
 	// =============================================================================================================
 	// Make the final tessellation 
 	// =============================================================================================================
@@ -195,26 +142,11 @@ bool FIsoTriangulator::Triangulate()
 	// Finalize the mesh by the tessellation of the inner grid
 	TriangulateInnerNodes();
 
-#ifdef DEBUG_ISOTRIANGULATOR
-	if (bDisplay)
-	{
-		F3DDebugSession A(TEXT("Mesh 3D"));
-		DisplayMesh(Mesh);
-	}
-	//Wait(bDisplay);
-#endif
-
-#ifdef CADKERNEL_DEV
-	//Chronos.PrintTimeElapse();
-#endif
-
 	return true;
 }
 
 void FIsoTriangulator::BuildNodes()
 {
-	FTimePoint StartTime = FChrono::Now();
-
 	LoopNodeCount = 0;
 	for (const TArray<FVector2d>& LoopPoints : Grid.GetLoops2D(EGridSpace::Default2D))
 	{
@@ -258,10 +190,6 @@ void FIsoTriangulator::BuildNodes()
 			GlobalIndexToIsoInnerNodes[Index] = &Node;
 		}
 	}
-
-#ifdef CADKERNEL_DEV
-	Chronos.BuildIsoNodesDuration += FChrono::Elapse(StartTime);
-#endif
 }
 
 void FIsoTriangulator::FillMeshNodes()
@@ -322,8 +250,6 @@ void FIsoTriangulator::FillMeshNodes()
 
 void FIsoTriangulator::BuildLoopSegments()
 {
-	FTimePoint StartTime = FChrono::Now();
-
 	LoopSegments.Reserve(LoopNodeCount);
 
 	int32 LoopIndex = 0;
@@ -342,18 +268,9 @@ void FIsoTriangulator::BuildLoopSegments()
 		}
 		else
 		{
-#ifdef DEBUG_BUILD_SEGMENT_IF_NEEDED
-			F3DDebugSession _(FString::Printf(TEXT("ERROR Segment")));
-			Grid.DisplayIsoSegment(EGridSpace::UniformScaled, Node, Node.GetNextNode(), 0, EVisuProperty::RedCurve);
-			Wait();
-#endif
 			IsoSegmentFactory.DeleteEntity(&Segment);
 		}
 	}
-
-#ifdef CADKERNEL_DEV
-	Chronos.BuildLoopSegmentsDuration += FChrono::Elapse(StartTime);
-#endif
 }
 
 void FIsoTriangulator::GetThinZonesMesh()
@@ -509,11 +426,6 @@ void FIsoTriangulator::GetThinZoneMesh(const TMap<int32, FLoopNode*>& IndexToNod
 		}
 		else
 		{
-#ifdef DEBUG_BUILD_SEGMENT_IF_NEEDED
-			F3DDebugSession _(FString::Printf(TEXT("ERROR Segment")));
-			Grid.DisplayIsoSegment(EGridSpace::UniformScaled, CandidateSegment.StartNode, CandidateSegment.EndNode, 0, EVisuProperty::RedCurve);
-			Wait();
-#endif
 			IsoSegmentFactory.DeleteEntity(&Segment);
 		}
 	}
@@ -521,16 +433,6 @@ void FIsoTriangulator::GetThinZoneMesh(const TMap<int32, FLoopNode*>& IndexToNod
 
 void FIsoTriangulator::BuildInnerSegments()
 {
-#ifdef DEBUG_BUILDINNERSEGMENTS
-	if (bDisplay)
-	{
-		Grid.DisplayInnerPoints(TEXT("BuildInnerSegments::Points"), EGridSpace::UniformScaled);
-	}
-	F3DDebugSession _(bDisplay, TEXT("BuildInnerSegments"));
-#endif
-
-	FTimePoint StartTime = FChrono::Now();
-
 	// Build segments according to the Grid following u then following v
 	// Build segment must not be in intersection with the loop
 	int32 NumU = Grid.GetCuttingCount(EIso::IsoU);
@@ -558,14 +460,6 @@ void FIsoTriangulator::BuildInnerSegments()
 					{
 						if (Node1.X - DOUBLE_SMALL_NUMBER < LoopPoint.X && LoopPoint.X < Node2.X + DOUBLE_SMALL_NUMBER)
 						{
-#ifdef DEBUG_BUILDINNERSEGMENTS
-							if (bDisplay)
-							{
-								F3DDebugSession _(bDisplay, FString::Printf(TEXT("Point")));
-								DisplayPoint2DWithScale(LoopPoint, EVisuProperty::RedPoint);
-								Wait();
-							}
-#endif
 							return true;
 						}
 					}
@@ -582,14 +476,6 @@ void FIsoTriangulator::BuildInnerSegments()
 					{
 						if (Node1.Y - DOUBLE_SMALL_NUMBER < LoopPoint.Y && LoopPoint.Y < Node2.Y + DOUBLE_SMALL_NUMBER)
 						{
-#ifdef DEBUG_BUILDINNERSEGMENTS
-							if (bDisplay)
-							{
-								//F3DDebugSession _(bDisplay, FString::Printf(TEXT("Point")));
-								DisplayPoint2DWithScale(LoopPoint, EVisuProperty::RedPoint);
-								Wait();
-							}
-#endif
 							return true;
 						}
 					}
@@ -619,11 +505,6 @@ void FIsoTriangulator::BuildInnerSegments()
 		}
 		else
 		{
-#ifdef DEBUG_BUILD_SEGMENT_IF_NEEDED
-			F3DDebugSession _(FString::Printf(TEXT("ERROR Segment")));
-			Grid.DisplayIsoSegment(EGridSpace::UniformScaled, Node1, Node2, 0, EVisuProperty::RedCurve);
-			Wait();
-#endif
 			IsoSegmentFactory.DeleteEntity(&Segment);
 		}
 	};
@@ -632,24 +513,12 @@ void FIsoTriangulator::BuildInnerSegments()
 	{
 		if (Grid.IsNodeOusideFaceButClose(IndexNode1) && Grid.IsNodeOusideFaceButClose(IndexNode2))
 		{
-#ifdef DEBUG_BUILDINNERSEGMENTS
-			if (bDisplay)
-			{
-				DisplaySegmentWithScale(Grid.GetInner2DPoint(EGridSpace::UniformScaled, IndexNode1), Grid.GetInner2DPoint(EGridSpace::UniformScaled, IndexNode2), IndexNode1, EVisuProperty::PinkCurve);
-			}
-#endif
 			AddToInnerToOuterSegmentsIntersectionTool(IndexNode1, IndexNode2, InType);
 			return;
 		}
 
 		if (Grid.IsNodeOutsideFace(IndexNode1) && Grid.IsNodeOutsideFace(IndexNode2))
 		{
-#ifdef DEBUG_BUILDINNERSEGMENTS
-			if (bDisplay)
-			{
-				DisplaySegmentWithScale(Grid.GetInner2DPoint(EGridSpace::UniformScaled, IndexNode1), Grid.GetInner2DPoint(EGridSpace::UniformScaled, IndexNode2), IndexNode1, EVisuProperty::RedCurve);
-			}
-#endif
 			return;
 		}
 
@@ -658,23 +527,10 @@ void FIsoTriangulator::BuildInnerSegments()
 			if (LoopSegmentsIntersectionTool.DoesIntersect(Grid.GetInner2DPoint(EGridSpace::UniformScaled, IndexNode1), Grid.GetInner2DPoint(EGridSpace::UniformScaled, IndexNode2))
 				|| AlmostHitsLoop(Grid.GetInner2DPoint(EGridSpace::UniformScaled, IndexNode1), Grid.GetInner2DPoint(EGridSpace::UniformScaled, IndexNode2), InType, Tolerance))
 			{
-#ifdef DEBUG_BUILDINNERSEGMENTS
-				if (bDisplay)
-				{
-					DisplaySegmentWithScale(Grid.GetInner2DPoint(EGridSpace::UniformScaled, IndexNode1), Grid.GetInner2DPoint(EGridSpace::UniformScaled, IndexNode2), IndexNode1, EVisuProperty::PurpleCurve);
-					Wait();
-				}
-#endif
 				AddToInnerToOuterSegmentsIntersectionTool(IndexNode1, IndexNode2, InType);
 			}
 			else
 			{
-#ifdef DEBUG_BUILDINNERSEGMENTS
-				if (bDisplay)
-				{
-					DisplaySegmentWithScale(Grid.GetInner2DPoint(EGridSpace::UniformScaled, IndexNode1), Grid.GetInner2DPoint(EGridSpace::UniformScaled, IndexNode2), IndexNode1, EVisuProperty::OrangeCurve);
-				}
-#endif
 				AddToInnerSegments(IndexNode1, IndexNode2, InType);
 			}
 
@@ -683,33 +539,15 @@ void FIsoTriangulator::BuildInnerSegments()
 
 		if (Grid.IsNodeInsideAndMeshable(IndexNode1) && Grid.IsNodeInsideAndMeshable(IndexNode2))
 		{
-#ifdef DEBUG_BUILDINNERSEGMENTS
-			if (bDisplay)
-			{
-				DisplaySegmentWithScale(Grid.GetInner2DPoint(EGridSpace::UniformScaled, IndexNode1), Grid.GetInner2DPoint(EGridSpace::UniformScaled, IndexNode2), IndexNode1, EVisuProperty::BlueCurve);
-			}
-#endif
 			AddToInnerSegments(IndexNode1, IndexNode2, InType);
 			return;
 		}
 
 		if (Grid.IsNodeInsideButTooCloseToLoop(IndexNode1) && Grid.IsNodeInsideButTooCloseToLoop(IndexNode2))
 		{
-#ifdef DEBUG_BUILDINNERSEGMENTS
-			if (bDisplay)
-			{
-				DisplaySegmentWithScale(Grid.GetInner2DPoint(EGridSpace::UniformScaled, IndexNode1), Grid.GetInner2DPoint(EGridSpace::UniformScaled, IndexNode2), IndexNode1, EVisuProperty::GreenCurve);
-			}
-#endif
 			return;
 		}
 
-#ifdef DEBUG_BUILDINNERSEGMENTS
-		if (bDisplay)
-		{
-			DisplaySegmentWithScale(Grid.GetInner2DPoint(EGridSpace::UniformScaled, IndexNode1), Grid.GetInner2DPoint(EGridSpace::UniformScaled, IndexNode2), IndexNode1, EVisuProperty::YellowCurve);
-		}
-#endif
 		AddToInnerToOuterSegmentsIntersectionTool(IndexNode1, IndexNode2, InType);
 	};
 
@@ -757,21 +595,11 @@ void FIsoTriangulator::BuildInnerSegments()
 		}
 	}
 
-#ifdef DEBUG_BUILDINNERSEGMENTS
-	Wait(bDisplay);
-#endif
-
 	InnerToOuterIsoSegmentsIntersectionTool.Sort();
-
-#ifdef CADKERNEL_DEV
-	Chronos.BuildInnerSegmentsDuration += FChrono::Elapse(StartTime);
-#endif
 }
 
 void FIsoTriangulator::BuildInnerSegmentsIntersectionTool()
 {
-	FTimePoint StartTime = FChrono::Now();
-
 	// Find Boundary Segments Of Inner Triangulation
 	// 
 	// A pixel grid is build. 
@@ -803,43 +631,6 @@ void FIsoTriangulator::BuildInnerSegmentsIntersectionTool()
 	TArray<uint8> Pixel;
 	Pixel.Init(0, Grid.GetTotalCuttingCount());
 
-#ifdef DEBUG_FINDBOUNDARYSEGMENTS
-	if (bDisplay)
-	{
-		F3DDebugSession _(TEXT("FIsoTrianguler::InnerSegment"));
-		for (int32 IndexV = 0, Index = 0; IndexV < NumV; ++IndexV)
-		{
-			for (int32 IndexU = 0; IndexU < NumU; ++IndexU, ++Index)
-			{
-				if (!Grid.IsNodeInsideAndMeshable(Index))
-				{
-					continue;
-				}
-
-				if (GlobalIndexToIsoInnerNodes[Index]->IsLinkedToNextU())
-				{
-					Grid.DisplayIsoSegment(EGridSpace::UniformScaled, *GlobalIndexToIsoInnerNodes[Index], *GlobalIndexToIsoInnerNodes[Index + 1], Index, EVisuProperty::YellowCurve);
-				}
-
-				if (GlobalIndexToIsoInnerNodes[Index]->IsLinkedToPreviousU())
-				{
-					Grid.DisplayIsoSegment(EGridSpace::UniformScaled, *GlobalIndexToIsoInnerNodes[Index], *GlobalIndexToIsoInnerNodes[Index - 1], Index, EVisuProperty::YellowCurve);
-				}
-
-				if (GlobalIndexToIsoInnerNodes[Index]->IsLinkedToNextV())
-				{
-					Grid.DisplayIsoSegment(EGridSpace::UniformScaled, *GlobalIndexToIsoInnerNodes[Index], *GlobalIndexToIsoInnerNodes[Index + NumU], Index, EVisuProperty::YellowCurve);
-				}
-
-				if (GlobalIndexToIsoInnerNodes[Index]->IsLinkedToPreviousV())
-				{
-					Grid.DisplayIsoSegment(EGridSpace::UniformScaled, *GlobalIndexToIsoInnerNodes[Index], *GlobalIndexToIsoInnerNodes[Index - NumU], Index, EVisuProperty::YellowCurve);
-				}
-			}
-		}
-	}
-#endif
-
 	// A pixel is True if one of its boundary segment does not exist 
 	for (int32 IndexV = 0, Index = 0; IndexV < NumV; ++IndexV)
 	{
@@ -905,15 +696,6 @@ void FIsoTriangulator::BuildInnerSegmentsIntersectionTool()
 
 	// initialize the intersection tool
 	InnerSegmentsIntersectionTool.Sort();
-
-#ifdef CADKERNEL_DEV
-	Chronos.FindLoopSegmentOfInnerTriangulationDuration += FChrono::Elapse(StartTime);
-#endif
-
-#ifdef DEBUG_FINDBOUNDARYSEGMENTS
-	DisplayPixels(Pixel);
-	//Wait(bDisplay);
-#endif
 }
 
 // =============================================================================================================
@@ -928,17 +710,6 @@ void FIsoTriangulator::ConnectCellLoops()
 	TArray<FCell> Cells;
 	FindCellContainingBoundaryNodes(Cells);
 
-#ifdef DEBUG_CONNECT_CELL_LOOPS
-	if (bDisplay)
-	{
-		for (FCell& Cell : Cells)
-		{
-			F3DDebugSession _(FString::Printf(TEXT("Cell %d"), Cell.Id));
-			DisplayCell(Cell);
-		}
-	}
-#endif
-
 	for (FCell& Cell : Cells)
 	{
 		InitCellCorners(Cell);
@@ -947,62 +718,26 @@ void FIsoTriangulator::ConnectCellLoops()
 
 	InnerToLoopCandidateSegments.Reserve(Cells.Num() * 2);
 
-#ifdef DEBUG_CONNECT_CELL_LOOPS
-	F3DDebugSession _(bDisplay, ("ConnectCellLoops"));
-#endif
-
 	FinalToLoops.Reserve(LoopNodeCount + InnerNodeCount);
 	for (FCell& Cell : Cells)
 	{
-#ifdef DEBUG_CONNECT_CELL_LOOPS
-		F3DDebugSession _(bDisplay, FString::Printf(TEXT("Cell %d"), Cell.Id));
-		DisplayCell(Cell);
-#endif
-
 		if (Cell.CellLoops.Num())
 		{
 			Cell.FindCandidateToConnectLoopsByNeighborhood();
 
-#ifdef DEBUG_CONNECT_CELL_CORNER_TO_INNER_LOOP
-			if (bDisplay)
-			{
-				DisplayCellConnexions(FString::Printf(TEXT("Cell %d CandidateSegments"), Cell.Id), Cell.LoopConnexions, EVisuProperty::YellowCurve);
-			}
-#endif
-
 			FindCandidateToConnectCellCornerToLoops(Cell);
 
-#ifdef DEBUG_CONNECT_CELL_CORNER_TO_INNER_LOOP
-			if (bDisplay)
-			{
-				DisplayCellConnexions(FString::Printf(TEXT("Cell %d CandidateSegments"), Cell.Id), Cell.LoopConnexions, EVisuProperty::BlueCurve);
-			}
-#endif
-
 			Cell.SelectSegmentToConnectLoops(IsoSegmentFactory);
 			Cell.SelectSegmentToConnectLoopToCorner(IsoSegmentFactory);
 			Cell.CheckAllLoopsConnectedTogetherAndConnect();
-#ifdef DEBUG_CONNECT_CELL_LOOPS
-			Grid.DisplayIsoSegments(TEXT("ConnectCellLoops::Step 0"), EGridSpace::UniformScaled, Cell.FinalSegments, false, false, false, EVisuProperty::YellowCurve);
-#endif
 		}
 
-#ifdef DEBUG_CONNECT_CELL_LOOPS
-		Grid.DisplayIsoSegments(TEXT("ConnectCellLoops::Step 1"), EGridSpace::UniformScaled, Cell.FinalSegments, false, false, false, EVisuProperty::YellowCurve);
-#endif
-
 		FinalToLoops.Append(Cell.FinalSegments);
-
-#ifdef DEBUG_CONNECT_CELL_LOOPS
-		Grid.DisplayIsoSegments(TEXT("ConnectCellLoops::Final Iso ToLink Inner To Loops"), EGridSpace::UniformScaled, FinalToLoops, false, false, false, EVisuProperty::YellowCurve);
-#endif
 	}
 }
 
 void FIsoTriangulator::FindCellContainingBoundaryNodes(TArray<FCell>& Cells)
 {
-	FTimePoint StartTime = FChrono::Now();
-
 	TArray<int32> NodeToCellIndices;
 	TArray<int32> SortedIndex;
 
@@ -1020,10 +755,6 @@ void FIsoTriangulator::FindCellContainingBoundaryNodes(TArray<FCell>& Cells)
 		int32 Index = 0;
 		int32 DeletedNodeCount = 0;
 
-#ifdef DEBUG_BUILD_CELLS
-		F3DDebugSession _(bDisplay, TEXT("Debug Find index"));
-#endif
-
 		for (const FLoopNode& LoopPoint : LoopNodes)
 		{
 			if (!LoopPoint.IsDelete())
@@ -1032,27 +763,10 @@ void FIsoTriangulator::FindCellContainingBoundaryNodes(TArray<FCell>& Cells)
 				ArrayUtils::FindCoordinateIndex(IsoUCoordinates, Coordinate.X, IndexU);
 				ArrayUtils::FindCoordinateIndex(IsoVCoordinates, Coordinate.Y, IndexV);
 
-#ifdef DEBUG_BUILD_CELLS
-				if (bDisplay)
-				{
-					F3DDebugSession S(FString::Printf(TEXT("Index %d"), IndexV * CountU + IndexU));
-					DisplayPoint2DWithScale(LoopPoint.GetPoint(EGridSpace::UniformScaled, Grid), EVisuProperty::BluePoint);
-					DrawCellBoundary(IndexV * CountU + IndexU, EVisuProperty::YellowPoint);
-				}
-#endif
-
 				NodeToCellIndices.Emplace(IndexV * CountU + IndexU);
 			}
 			else
 			{
-#ifdef DEBUG_BUILD_CELLS
-				if (bDisplay)
-				{
-					F3DDebugSession S(TEXT("Deleted Node"));
-					DisplayPoint2DWithScale(LoopPoint.GetPoint(EGridSpace::UniformScaled, Grid), EVisuProperty::RedPoint);
-				}
-#endif
-
 				DeletedNodeCount++;
 				NodeToCellIndices.Emplace(MaxUV);
 			}
@@ -1087,10 +801,6 @@ void FIsoTriangulator::FindCellContainingBoundaryNodes(TArray<FCell>& Cells)
 		TArray<FLoopNode*> CellNodes;
 		CellNodes.Reserve(LoopNodeCount);
 
-#ifdef DEBUG_BUILD_CELLS
-		F3DDebugSession _(bDisplay, TEXT("Build Cells"));
-#endif
-
 		for (int32 Index : SortedIndex)
 		{
 			if (CellIndex != NodeToCellIndices[Index])
@@ -1099,14 +809,6 @@ void FIsoTriangulator::FindCellContainingBoundaryNodes(TArray<FCell>& Cells)
 
 				CellIndex = NodeToCellIndices[Index];
 				CellNodes.Reset(LoopNodeCount);
-
-#ifdef DEBUG_BUILD_CELLS
-				if(bDisplay)
-				{
-					DisplayCell(Cells.Last());
-					Wait();
-				}
-#endif
 			}
 
 			FLoopNode& LoopNode = LoopNodes[Index];
@@ -1116,17 +818,7 @@ void FIsoTriangulator::FindCellContainingBoundaryNodes(TArray<FCell>& Cells)
 			}
 		}
 		Cells.Emplace(CellIndex, CellNodes, *this);
-
-#ifdef DEBUG_BUILD_CELLS
-		if (bDisplay)
-		{
-			DisplayCell(Cells.Last());
-			Wait();
-		}
-#endif
 	}
-
-	FChrono::Elapse(StartTime);
 }
 
 bool FIsoTriangulator::CanCycleBeMeshed(const TArray<FIsoSegment*>& Cycle, FIntersectionSegmentTool& CycleIntersectionTool)
@@ -1137,18 +829,8 @@ bool FIsoTriangulator::CanCycleBeMeshed(const TArray<FIsoSegment*>& Cycle, FInte
 	{
 		if (CycleIntersectionTool.DoesIntersect(*Segment))
 		{
-			FMessage::Printf(Log, TEXT("A cycle of the surface %d is in self intersecting. The mesh of this sector is canceled.\n"), Grid.GetFace().GetId());
+			UE_LOG(LogCADKernelBase, Warning, TEXT("A cycle of the surface %d is in self intersecting. The mesh of this sector is canceled.\n"), Grid.GetFace().GetId());
 
-#ifdef DEBUG_CAN_CYCLE_BE_FIXED_AND_MESHED
-			if (bDisplay)
-			{
-				F3DDebugSession _(TEXT("New Segment"));
-				F3DDebugSession A(TEXT("Segments"));
-				Display(EGridSpace::UniformScaled, *Segment, 0, EVisuProperty::RedPoint);
-				Display(EGridSpace::UniformScaled, *IntersectingSegment, 0, EVisuProperty::BluePoint);
-				Wait();
-			}
-#endif
 			return false;
 		}
 	}
@@ -1193,8 +875,6 @@ void FIsoTriangulator::MeshLargeCycle(const TArray<FIsoSegment*>& Cycle, const T
  */
 void FIsoTriangulator::FindInnerGridCellSurroundingSmallLoop()
 {
-	FTimePoint StartTime = FChrono::Now();
-
 	if (GlobalIndexToIsoInnerNodes.Num() == 0)
 	{
 		// No inner node
@@ -1298,33 +978,15 @@ void FIsoTriangulator::FindInnerGridCellSurroundingSmallLoop()
 			}
 		}
 	}
-
-#ifdef CADKERNEL_DEV
-	Chronos.FindSegmentIsoUVSurroundingSmallLoopDuration += FChrono::Elapse(StartTime);
-#endif
 }
 
 void FIsoTriangulator::TriangulateOverCycle(const EGridSpace Space)
 {
-#ifdef ADD_TRIANGLE_2D
-	F3DDebugSession G(bDisplay, ("Triangulate cycles"));
-#endif 
-
-	FTimePoint StartTime = FChrono::Now();
-
 	TArray<FIsoSegment*> Cycle;
 	Cycle.Reserve(100);
 	TArray<bool> CycleOrientation;
 	CycleOrientation.Reserve(100);
 
-#ifdef FIND_CYCLE
-	int32 OverCycleIndex = 0;
-	if (Grid.bDisplay)
-	{
-		Open3DDebugSession(TEXT("Triangulate Over Cycle"));
-	}
-#endif
-
 	// first the external segments (loop segments) are processed 
 	for (FIsoSegment* Segment : LoopSegments)
 	{
@@ -1337,14 +999,7 @@ void FIsoTriangulator::TriangulateOverCycle(const EGridSpace Space)
 			{
 				continue;
 			}
-#ifdef FIND_CYCLE
-			if (Grid.bDisplay)
-			{
-				F3DDebugSession G(TEXT("Find & mesh cycles"));
-				FString Message = FString::Printf(TEXT("MeshCycle - cycle %d"), OverCycleIndex++);
-				Grid.DisplayIsoSegments(Message, EGridSpace::UniformScaled, Cycle, false);
-			}
-#endif
+
 			MeshCycle(Cycle, CycleOrientation);
 		}
 	}
@@ -1361,15 +1016,7 @@ void FIsoTriangulator::TriangulateOverCycle(const EGridSpace Space)
 			{
 				continue;
 			}
-#ifdef FIND_CYCLE
-			if (Grid.bDisplay)
-			{
-				Open3DDebugSession(TEXT("Find & mesh cycles"));
-				FString Message = FString::Printf(TEXT("MeshCycle - cycle %d"), OverCycleIndex++);
-				Grid.DisplayIsoSegments(Message, EGridSpace::UniformScaled, Cycle, false);
-				Close3DDebugSession();
-			}
-#endif
+
 			MeshCycle(Cycle, CycleOrientation);
 		}
 
@@ -1382,55 +1029,17 @@ void FIsoTriangulator::TriangulateOverCycle(const EGridSpace Space)
 			{
 				continue;
 			}
-#ifdef FIND_CYCLE
-			if (Grid.bDisplay)
-			{
-				Open3DDebugSession(TEXT("Find & mesh cycles"));
-				FString Message = FString::Printf(TEXT("MeshCycle - cycle %d"), OverCycleIndex++);
-				Grid.DisplayIsoSegments(Message, EGridSpace::UniformScaled, Cycle, false);
-				Close3DDebugSession();
-			}
-#endif
+
 			MeshCycle(Cycle, CycleOrientation);
 		}
 	}
-#ifdef FIND_CYCLE
-	Close3DDebugSession();
-#endif
-
-#ifdef CADKERNEL_DEV
-	Chronos.TriangulateOverCycleDuration = FChrono::Elapse(StartTime);
-#endif
 }
 
-#ifdef DEBUG_FIND_CYCLE
-static int32 CycleId = -1;
-static int32 CycleIndex = 0;
-#endif
-
 bool FIsoTriangulator::FindCycle(FIsoSegment* StartSegment, bool LeftSide, TArray<FIsoSegment*>& Cycle, TArray<bool>& CycleOrientation)
 {
-
-#ifdef DEBUG_FIND_CYCLE
-	CycleIndex++;
-
-	if (Grid.bDisplay)
-	{
-		CycleIndex = CycleId;
-		bDisplay = true;
-	}
-#endif
-
 	Cycle.Empty();
 	CycleOrientation.Empty();
 
-#ifdef DEBUG_FIND_CYCLE
-	if (CycleId == CycleIndex)
-	{
-		Open3DDebugSession(TEXT("Cycle"));
-	}
-#endif
-
 	FIsoSegment* Segment = StartSegment;
 	FIsoNode* Node;
 
@@ -1438,31 +1047,11 @@ bool FIsoTriangulator::FindCycle(FIsoSegment* StartSegment, bool LeftSide, TArra
 	{
 		Segment->SetHaveCycleOnLeft();
 		Node = &StartSegment->GetSecondNode();
-#ifdef DEBUG_FIND_CYCLE
-		if (CycleId == CycleIndex)
-		{
-			FIsoNode* EndNode = &StartSegment->GetFirstNode();
-			F3DDebugSession _(TEXT("FirstSegment left"));
-			Grid.DisplayIsoNode(EGridSpace::UniformScaled, *EndNode);
-			Grid.DisplayIsoSegment(EGridSpace::UniformScaled, *StartSegment);
-			//Wait();
-		}
-#endif
 	}
 	else
 	{
 		Segment->SetHaveCycleOnRight();
 		Node = &StartSegment->GetFirstNode();
-#ifdef DEBUG_FIND_CYCLE
-		if (CycleId == CycleIndex)
-		{
-			FIsoNode* EndNode = &StartSegment->GetSecondNode();
-			F3DDebugSession _(TEXT("FirstSegment right"));
-			Grid.DisplayIsoNode(EGridSpace::UniformScaled, *EndNode);
-			Grid.DisplayIsoSegment(EGridSpace::UniformScaled, *StartSegment);
-			//Wait();
-		}
-#endif
 	}
 
 	Cycle.Add(StartSegment);
@@ -1487,14 +1076,6 @@ bool FIsoTriangulator::FindCycle(FIsoSegment* StartSegment, bool LeftSide, TArra
 
 		if (&Segment->GetFirstNode() == Node)
 		{
-#ifdef DEBUG_HAS_CYCLE
-			if (Segment->HasCycleOnLeft())
-			{
-				F3DDebugSession _(TEXT("Segment HasCycleOnLeft"));
-				Grid.DisplayIsoSegment(EGridSpace::UniformScaled, *Segment, 0, EVisuProperty::RedPoint);
-				Wait();
-			}
-#endif
 			if (Segment->HasCycleOnLeft())
 			{
 				return false;
@@ -1502,26 +1083,9 @@ bool FIsoTriangulator::FindCycle(FIsoSegment* StartSegment, bool LeftSide, TArra
 			Segment->SetHaveCycleOnLeft();
 			Node = &Segment->GetSecondNode();
 			CycleOrientation.Add(true);
-#ifdef DEBUG_FIND_CYCLE
-			if (CycleId == CycleIndex)
-			{
-				F3DDebugSession _(TEXT("Next"));
-				Grid.DisplayIsoNode(EGridSpace::UniformScaled, *Node);
-				Grid.DisplayIsoSegment(EGridSpace::UniformScaled, *Segment);
-				//Wait();
-			}
-#endif
 		}
 		else
 		{
-#ifdef DEBUG_HAS_CYCLE
-			if (Segment->HasCycleOnRight())
-			{
-				F3DDebugSession _(TEXT("Segment HasCycleOnRight"));
-				Grid.DisplayIsoSegment(EGridSpace::UniformScaled, *Segment, 0, EVisuProperty::RedPoint);
-				Wait();
-			}
-#endif
 			if (Segment->HasCycleOnRight())
 			{
 				return false;
@@ -1529,27 +1093,8 @@ bool FIsoTriangulator::FindCycle(FIsoSegment* StartSegment, bool LeftSide, TArra
 			Segment->SetHaveCycleOnRight();
 			Node = &Segment->GetFirstNode();
 			CycleOrientation.Add(false);
-#ifdef DEBUG_FIND_CYCLE
-			if (CycleId == CycleIndex)
-			{
-				F3DDebugSession _(TEXT("Next"));
-				Grid.DisplayIsoNode(EGridSpace::UniformScaled, *Node);
-				Grid.DisplayIsoSegment(EGridSpace::UniformScaled, *Segment);
-				//Wait();
-			}
-#endif
-		}
-		if (false)
-		{
-			Wait();
 		}
 	}
-#ifdef DEBUG_FIND_CYCLE
-	if (CycleId == CycleIndex)
-	{
-		Close3DDebugSession();
-	}
-#endif
 	return true;
 }
 
@@ -1585,17 +1130,9 @@ FIsoSegment* FIsoTriangulator::FindNextSegment(EGridSpace Space, const FIsoSegme
 
 void FIsoTriangulator::TriangulateInnerNodes()
 {
-	FTimePoint StartTime = FChrono::Now();
-
 	int32 NumU = Grid.GetCuttingCount(EIso::IsoU);
 	int32 NumV = Grid.GetCuttingCount(EIso::IsoV);
 
-#ifdef ADD_TRIANGLE_2D
-	if (bDisplay)
-	{
-		Open3DDebugSession(TEXT("Inner Mesh 2D"));
-	}
-#endif
 	for (int32 vIndex = 0, Index = 0; vIndex < NumV - 1; vIndex++)
 	{
 		for (int32 uIndex = 0; uIndex < NumU - 1; uIndex++, Index++)
@@ -1639,27 +1176,15 @@ void FIsoTriangulator::TriangulateInnerNodes()
 				continue;
 			}
 
-#ifdef ADD_TRIANGLE_2D
-			Grid.DisplayTriangle(EGridSpace::UniformScaled, *GlobalIndexToIsoInnerNodes[Index], *GlobalIndexToIsoInnerNodes[Index + 1], *GlobalIndexToIsoInnerNodes[OppositIndex]);
-			Grid.DisplayTriangle(EGridSpace::UniformScaled, *GlobalIndexToIsoInnerNodes[OppositIndex], *GlobalIndexToIsoInnerNodes[OppositIndex - 1], *GlobalIndexToIsoInnerNodes[Index]);
-#endif 
-
 			Mesh.AddTriangle(GlobalIndexToIsoInnerNodes[Index]->GetGlobalIndex(), GlobalIndexToIsoInnerNodes[Index + 1]->GetGlobalIndex(), GlobalIndexToIsoInnerNodes[OppositIndex]->GetGlobalIndex());
 			Mesh.AddTriangle(GlobalIndexToIsoInnerNodes[OppositIndex]->GetGlobalIndex(), GlobalIndexToIsoInnerNodes[OppositIndex - 1]->GetGlobalIndex(), GlobalIndexToIsoInnerNodes[Index]->GetGlobalIndex());
 		}
 		Index++;
 	}
-#ifdef ADD_TRIANGLE_2D
-	Close3DDebugSession();
-#endif
 }
 
 void FCell::FindCandidateToConnectLoopsByNeighborhood()
 {
-#ifdef DEBUG_CONNECT_CELL_SUB_LOOPS_BY_NEIGHBORHOOD
-	F3DDebugSession _(Grid.bDisplay, FString::Printf(TEXT("Cell % d: Build Segments to Connect inner loops"), Id));
-#endif
-
 	for (FCellConnexion& LoopConnexion : LoopConnexions)
 	{
 		TryToConnectTwoSubLoopsWithShortestSegment(LoopConnexion);
@@ -1685,16 +1210,6 @@ void FCell::SelectSegmentToConnectLoops(TFactory<FIsoSegment>& SegmentFactory)
 		});
 
 	{
-#ifdef DEBUG_CONNECT_CELL_SUB_LOOPS_BY_NEIGHBORHOOD
-		F3DDebugSession _(Grid.bDisplay, TEXT("Best Path"));
-#endif
-#ifdef DEBUG_CONNECT_CELL_SUB_LOOPS_BY_NEIGHBORHOOD
-		if (Grid.bDisplay)
-		{
-			Grid.DisplayIsoSegments(TEXT("SelectSegmentToConnectLoops: before"), EGridSpace::UniformScaled, FinalSegments, false, false, false, EVisuProperty::YellowCurve);
-			Wait(Grid.bDisplay);
-		}
-#endif
 		const int32 LoopCount = CellLoops.Num();
 		for (FCellConnexion* LoopConnexion : LoopConnexionPtrs)
 		{
@@ -1712,14 +1227,6 @@ void FCell::SelectSegmentToConnectLoops(TFactory<FIsoSegment>& SegmentFactory)
 			}
 		}
 	}
-
-#ifdef DEBUG_CONNECT_CELL_SUB_LOOPS_BY_NEIGHBORHOOD
-	if (Grid.bDisplay)
-	{
-		Grid.DisplayIsoSegments(TEXT("SelectSegmentToConnectLoops: Final"), EGridSpace::UniformScaled, FinalSegments, false, false, false, EVisuProperty::YellowCurve);
-		Wait(Grid.bDisplay);
-	}
-#endif
 }
 
 void FCell::SelectSegmentToConnectLoopToCorner(TFactory<FIsoSegment>& SegmentFactory)
@@ -1740,9 +1247,6 @@ void FCell::SelectSegmentToConnectLoopToCorner(TFactory<FIsoSegment>& SegmentFac
 		});
 
 	{
-#ifdef DEBUG_CONNECT_CELL_SUB_LOOPS_BY_NEIGHBORHOOD
-		F3DDebugSession _(Grid.bDisplay, TEXT("Best Path"));
-#endif
 		const int32 LoopCount = CellLoops.Num();
 		for (FCellConnexion* LoopConnexion : LoopConnexionPtrs)
 		{
@@ -1777,14 +1281,6 @@ void FCell::SelectSegmentToConnectLoopToCorner(TFactory<FIsoSegment>& SegmentFac
 			}
 		}
 	}
-
-#ifdef DEBUG_CONNECT_CELL_SUB_LOOPS_BY_NEIGHBORHOOD
-	if (Grid.bDisplay)
-	{
-		Grid.DisplayIsoSegments(TEXT("Build Segments Connect inner close loops: Candidates"), EGridSpace::UniformScaled, CandidateSegments, false, false, false, EVisuProperty::YellowCurve);
-		Wait(Grid.bDisplay);
-	}
-#endif
 }
 
 void FCellLoop::PropagateAsConnected()
@@ -1835,23 +1331,7 @@ void FCell::TryToCreateSegment(FCellConnexion& LoopConnexion)
 			IntersectionTool.AddSegment(*LoopConnexion.Segment);
 			FinalSegments.Add(LoopConnexion.Segment);
 		}
-		else
-		{
-#ifdef DEBUG_BUILD_SEGMENT_IF_NEEDED
-			F3DDebugSession _(FString::Printf(TEXT("Not expected error Segment")));
-			Grid.DisplayIsoSegment(EGridSpace::UniformScaled, *LoopConnexion.NodeA, *LoopConnexion.NodeB, 0, EVisuProperty::RedCurve);
-			Wait();
-#endif
-		}
 	}
-#ifdef DEBUG_BUILD_SEGMENT_IF_NEEDED2
-	else if(Grid.bDisplay)
-	{
-		F3DDebugSession _(FString::Printf(TEXT("Failed Segment %d %d"), LoopConnexion.Loop1.Id, LoopConnexion.Loop2.Id));
-		Grid.DisplayIsoSegment(EGridSpace::UniformScaled, *LoopConnexion.NodeA, *LoopConnexion.NodeB, 0, EVisuProperty::RedCurve);
-		Wait();
-	}
-#endif
 }
 
 
@@ -1889,9 +1369,6 @@ void FCell::InitLoopConnexions()
 	{
 		TArray<TPair<int32, FVector2d>> LoopBarycenters = GetLoopBarycenters();
 
-#ifdef DEBUG_BOWYERWATSON
-		FBowyerWatsonTriangulator::bDisplay = bDisplay;
-#endif
 		TArray<int32> EdgeVertexIndices;
 		FBowyerWatsonTriangulator BWTriangulator(LoopBarycenters, EdgeVertexIndices);
 		BWTriangulator.Triangulate();
@@ -1929,9 +1406,6 @@ void FCell::TryToConnectTwoSubLoopsWithShortestSegment(FCellConnexion& LoopConne
 	const TArray<FLoopNode*>& LoopA = LoopConnexion.Loop1.Nodes;
 	const TArray<FLoopNode*>& LoopB = LoopConnexion.Loop2.Nodes;
 
-#ifdef DEBUG_CONNECT_CELL_SUB_LOOPS_BY_NEIGHBORHOOD
-	//F3DDebugSession A(Grid.bDisplay, TEXT("ConnectTwoSubLoopsWithShortestSegment"));
-#endif
 	double MinDistanceSquare = HUGE_VALUE_SQUARE;
 
 	for (FLoopNode* NodeA : LoopA)
@@ -1942,10 +1416,6 @@ void FCell::TryToConnectTwoSubLoopsWithShortestSegment(FCellConnexion& LoopConne
 		{
 			const FVector2d& BCoordinates = NodeB->Get2DPoint(EGridSpace::UniformScaled, Grid);
 
-#ifdef DEBUG_CONNECT_CELL_SUB_LOOPS_BY_NEIGHBORHOOD_
-			Grid.DisplayIsoSegment(EGridSpace::UniformScaled, *NodeA, *NodeB, 0, EVisuProperty::BlueCurve);
-#endif
-
 			double SquareDistance = FVector2d::DistSquared(ACoordinates, BCoordinates);
 			if (SquareDistance < MinDistanceSquare)
 			{
@@ -1983,7 +1453,7 @@ void FIsoTriangulator::TryToConnectTwoLoopsWithIsocelesTriangle(FCell& Cell, con
 		// For each segment of the LoopA, find in loopB a vertex that made the best triangle i.e. the triangle the most isosceles.
 		// According to the knowledge of the orientation, only inside triangles are tested
 		//
-		// These computations are done in the UniformScaled space to avoid numerical error due to lenght distortion between U or V space and U or V Length.
+		// These computations are done in the UniformScaled space to avoid numerical error due to length distortion between U or V space and U or V Length.
 		// i.e. if:
 		// (UMax - UMin) / (VMax - VMin) is big 
 		// and 
@@ -1992,21 +1462,6 @@ void FIsoTriangulator::TryToConnectTwoLoopsWithIsocelesTriangle(FCell& Cell, con
 		// To avoid flat triangle, a candidate point must defined a minimal slop with [A, X0] or [B, Xn] to not be aligned with one of them. 
 		//
 
-#ifdef DEBUG_FIND_BEST_TRIANGLE_TO_LINK_LOOPS
-		int32 TriangleIndex = 0;
-		if (bDisplay)
-		{
-			{
-				TriangleIndex++;
-				F3DDebugSession _(FString::Printf(TEXT("Start Segment %d"), TriangleIndex));
-				Grid.DisplayIsoSegment(EGridSpace::UniformScaled, *Segment);
-				Grid.DisplayIsoNode(EGridSpace::UniformScaled, StartNode, 0, EVisuProperty::RedPoint);
-				Grid.DisplayIsoNode(EGridSpace::UniformScaled, EndNode);
-				Wait();
-			}
-		}
-#endif
-
 		FIsoNode* CandidatNode = nullptr;
 		FIsoSegment* StartToCandiatSegment = nullptr;
 		FIsoSegment* EndToCandiatSegment = nullptr;
@@ -2102,15 +1557,6 @@ void FIsoTriangulator::TryToConnectTwoLoopsWithIsocelesTriangle(FCell& Cell, con
 		FIsoNode* Node = FindBestTriangle(Segment);
 		if (Node && !Node->IsDeleteOrThinNode())
 		{
-#ifdef DEBUG_TRY_TO_CONNECT
-			if (bDisplay)
-			{
-				F3DDebugSession _(FString::Printf(TEXT("Triangle")));
-				DisplayTriangle(EGridSpace::UniformScaled, *NodeA1, *NodeA2, *Node);
-				Wait(true);
-
-			}
-#endif 
 			const FVector2d& NodeCoordinates = Node->Get2DPoint(EGridSpace::UniformScaled, Grid);
 			if (!NodeA1->IsDeleteOrThinNode())
 			{
@@ -2127,11 +1573,6 @@ void FIsoTriangulator::TryToConnectTwoLoopsWithIsocelesTriangle(FCell& Cell, con
 
 void FIsoTriangulator::TryToConnectVertexSubLoopWithTheMostIsoSegment(FCell& Cell, const TArray<FLoopNode*>& Loop)
 {
-#ifdef DEBUG_FIND_ISO_SEGMENT_TO_LINK_OUTER_LOOP_NODES
-	F3DDebugSession _(bDisplay, TEXT("TryToConnectVertexSubLoopWithTheMostIsoSegment"));
-	Wait(bDisplay);
-#endif
-
 	const double FlatSlope = 0.10; // ~5 deg: The segment must make an angle less than 10 deg with the Iso
 	double MinSlope = FlatSlope + DOUBLE_SMALL_NUMBER; //.25; // ~10 deg: The segment must make an angle less than 10 deg with the Iso
 
@@ -2198,9 +1639,6 @@ void FIsoTriangulator::TryToConnectVertexSubLoopWithTheMostIsoSegment(FCell& Cel
 		{
 			const FVector2d& BCoordinates = CandidateB->Get2DPoint(EGridSpace::UniformScaled, Grid);
 
-#ifdef DEBUG_FIND_ISO_SEGMENT_TO_LINK_OUTER_LOOP_NODES
-			Grid.DisplayIsoSegment(UniformScaled, *CandidateA, *CandidateB, 0, EVisuProperty::RedCurve);
-#endif
 			GetOrTryToCreateSegment(Cell, CandidateA, ACoordinates, CandidateB, BCoordinates, 0.1);
 			MinSlope = FlatSlope + DOUBLE_SMALL_NUMBER;
 		}
@@ -2210,14 +1648,6 @@ void FIsoTriangulator::TryToConnectVertexSubLoopWithTheMostIsoSegment(FCell& Cel
 
 void FIsoTriangulator::TryToConnectTwoSubLoopsWithTheMostIsoSegment(FCell& Cell, const TArray<FLoopNode*>& LoopA, const TArray<FLoopNode*>& LoopB)
 {
-#ifdef DEBUG_FIND_ISO_SEGMENT_TO_LINK_OUTER_LOOP_NODES
-	F3DDebugSession _(bDisplay, TEXT("TryToConnectVertexSubLoopWithTheMostIsoSegment"));
-	if (bDisplay)
-	{
-		Wait();
-	}
-#endif
-
 	const double FlatSlope = 0.10; // ~5 deg: The segment must make an angle less than 10 deg with the Iso
 
 	for (FLoopNode* CandidateA : LoopA)
@@ -2264,9 +1694,6 @@ void FIsoTriangulator::TryToConnectTwoSubLoopsWithTheMostIsoSegment(FCell& Cell,
 		{
 			const FVector2d& BCoordinates = CandidateB->Get2DPoint(EGridSpace::UniformScaled, Grid);
 
-#ifdef DEBUG_FIND_ISO_SEGMENT_TO_LINK_OUTER_LOOP_NODES
-			Grid.DisplayIsoSegment(EGridSpace::UniformScaled, *CandidateA, *CandidateB, 0, EVisuProperty::BlueCurve);
-#endif			
 			GetOrTryToCreateSegment(Cell, CandidateA, ACoordinates, CandidateB, BCoordinates, 0.1);
 			MinSlope = FlatSlope + DOUBLE_SMALL_NUMBER;
 		}
@@ -2300,14 +1727,6 @@ FIsoSegment* FIsoTriangulator::GetOrTryToCreateSegment(FCell& Cell, FLoopNode* N
 		return nullptr;
 	}
 
-#ifdef DEBUG_TRY_TO_CREATE_SEGMENT
-	{
-		F3DDebugSession _(TEXT("Test"));
-		DisplaySegmentWithScale(NodeA->Get2DPoint(EGridSpace::UniformScaled, Grid), NodeB->Get2DPoint(EGridSpace::UniformScaled, Grid), 0, EVisuProperty::RedCurve);
-		Wait(false);
-	}
-#endif
-
 	// Is Outside and not too flat at NodeA
 	if (NodeA->IsSegmentBeInsideFace(BCoordinates, Grid, FlatAngle))
 	{
@@ -2328,9 +1747,6 @@ FIsoSegment* FIsoTriangulator::GetOrTryToCreateSegment(FCell& Cell, FLoopNode* N
 	Segment.SetCandidate();
 	Cell.CandidateSegments.Add(&Segment);
 
-#ifdef DEBUG_TRY_TO_CREATE_SEGMENT
-	DisplaySegment(ACoordinates, BCoordinates, 0, EVisuProperty::OrangePoint);
-#endif
 	return &Segment;
 };
 
@@ -2360,10 +1776,6 @@ void FIsoTriangulator::FindCandidateToConnectCellCornerToLoops(FCell& Cell)
 		return;
 	}
 
-#ifdef DEBUG_CONNECT_CELL_CORNER_TO_INNER_LOOP
-	F3DDebugSession _(bDisplay, TEXT("With cell corners"));
-#endif
-
 	TFunction<void(FCellLoop&, FCellCorner&)> FindAndTryCreateCandidateSegmentToLinkLoopToCorner = [&](FCellLoop& LoopCell, FCellCorner& CellCorner)
 	{
 		FIsoInnerNode& CornerNode = CellCorner.CornerNode;
@@ -2428,42 +1840,17 @@ void FIsoTriangulator::FindCandidateToConnectCellCornerToLoops(FCell& Cell)
 
 	for (FCellCorner& CellCorner : Cell.CellCorners)
 	{
-#ifdef DEBUG_CONNECT_CELL_CORNER_TO_INNER_LOOP
-		if (bDisplay)
-		{
-			F3DDebugSession _(bDisplay, FString::Printf(TEXT("Cell corner %d"), CellCorner.Id));
-			Grid.DisplayIsoNode(EGridSpace::UniformScaled, CellCorner.CornerNode, 0, EVisuProperty::RedPoint);
-		}
-#endif
 		for (FCellLoop& LoopCell : Cell.CellLoops)
 		{
 			FindAndTryCreateCandidateSegmentToLinkLoopToCorner(LoopCell, CellCorner);
 		}
 	}
-
-#ifdef DEBUG_CONNECT_CELL_CORNER_TO_INNER_LOOP
-	if (bDisplay)
-	{
-		DisplayCellConnexions(TEXT("Cell CandidateSegments"), Cell.LoopConnexions, EVisuProperty::YellowCurve);
-		Wait(false);
-	}
-#endif
-
 }
 
 
 void FIsoTriangulator::SelectSegmentsToLinkInnerToLoop()
 {
-#ifdef DEBUG_ADD_CANDIDATE_SEGMENTS_TO_LINK_INNER_AND_LOOP
-	F3DDebugSession _(Grid.bDisplay, TEXT("AddCandidateSegmentsToLinkInnerAndLoop "));
-	Grid.DisplayIsoSegments(TEXT("CandidateSegments To Link Inner To Loops"), EGridSpace::UniformScaled, InnerToLoopCandidateSegments, false, false, false, EVisuProperty::RedCurve);
-#endif
-
 	LoopSegmentsIntersectionTool.AddSegments(FinalToLoops);
-#ifdef DEBUG_ADD_CANDIDATE_SEGMENTS_TO_LINK_INNER_AND_LOOP
-	LoopSegmentsIntersectionTool.Display(Grid.bDisplay, TEXT("LoopSegmentsIntersectionTool with final to loop"));
-	Wait();
-#endif
 
 	TArray<TPair<double, FIsoSegment*>> LengthOfCandidateSegments;
 	LengthOfCandidateSegments.Reserve(InnerToLoopCandidateSegments.Num());
@@ -2506,11 +1893,6 @@ void FIsoTriangulator::SelectSegmentsToLinkInnerToLoop()
 		Segment->SetFinalMarker();
 		if (!Segment->ConnectToNode())
 		{
-#ifdef DEBUG_BUILD_SEGMENT_IF_NEEDED
-			F3DDebugSession _(FString::Printf(TEXT("ERROR Segment")));
-			Grid.DisplayIsoSegment(EGridSpace::UniformScaled, *Segment, 0, EVisuProperty::RedCurve);
-			Wait();
-#endif
 			IsoSegmentFactory.DeleteEntity(Segment);
 		}
 	}
diff --git a/Engine/Source/Runtime/Datasmith/CADKernel/Base/Private/Mesh/Meshers/IsoTriangulator.h b/Engine/Source/Runtime/Datasmith/CADKernel/Base/Private/Mesh/Meshers/IsoTriangulator.h
--- a/Engine/Source/Runtime/Datasmith/CADKernel/Base/Private/Mesh/Meshers/IsoTriangulator.h
+++ b/Engine/Source/Runtime/Datasmith/CADKernel/Base/Private/Mesh/Meshers/IsoTriangulator.h
@@ -10,12 +10,6 @@
 #include "Mesh/Meshers/IsoTriangulator/IsoNode.h"
 #include "Mesh/Meshers/IsoTriangulator/IsoSegment.h"
 #include "Mesh/Meshers/ParametricMesherConstantes.h"
-#include "UI/Visu.h"
-
-#ifdef CADKERNEL_DEV
-#include "UI/DefineForDebug.h"
-#include "Mesh/Meshers/MesherReport.h"
-#endif
 
 namespace UE::CADKernel
 {
@@ -29,52 +23,6 @@ using FMeshPolygonFunc = TFunction<void(const FGrid&, FIsoNode*[], FFaceMesh&)>;
 
 struct FCellConnexion;
 
-struct CADKERNEL_API FIsoTriangulatorChronos
-{
-	FDuration TriangulateDuration = FChrono::Init();
-	FDuration BuildIsoNodesDuration = FChrono::Init();
-	FDuration BuildLoopSegmentsDuration = FChrono::Init();
-	FDuration BuildInnerSegmentsDuration = FChrono::Init();
-	FDuration FindLoopSegmentOfInnerTriangulationDuration = FChrono::Init();
-	FDuration FindSegmentIsoUVSurroundingSmallLoopDuration = FChrono::Init();
-	FDuration FindIsoSegmentToLinkInnerToLoopDuration = FChrono::Init();
-	FDuration FindInnerSegmentToLinkLoopToLoopDuration = FChrono::Init();
-	FDuration FindSegmentToLinkLoopToLoopDuration = FChrono::Init();
-	FDuration FindSegmentToLinkInnerToLoopDuration = FChrono::Init();
-	FDuration TriangulateOverCycleDuration = FChrono::Init();
-	FDuration TriangulateInnerNodesDuration = FChrono::Init();
-
-	FIsoTriangulatorChronos()
-	{}
-
-	void PrintTimeElapse() const
-	{
-		FDuration IsoTriangulerDuration = FChrono::Init();
-		IsoTriangulerDuration += BuildIsoNodesDuration;
-		IsoTriangulerDuration += BuildLoopSegmentsDuration;
-		IsoTriangulerDuration += BuildInnerSegmentsDuration;
-		IsoTriangulerDuration += FindLoopSegmentOfInnerTriangulationDuration;
-		IsoTriangulerDuration += FindIsoSegmentToLinkInnerToLoopDuration;
-		IsoTriangulerDuration += FindSegmentToLinkLoopToLoopDuration;
-		IsoTriangulerDuration += FindSegmentToLinkInnerToLoopDuration;
-		IsoTriangulerDuration += TriangulateOverCycleDuration;
-		IsoTriangulerDuration += TriangulateInnerNodesDuration;
-
-		FChrono::PrintClockElapse(Log, TEXT(""), TEXT("IsoTrianguler"), IsoTriangulerDuration);
-		FChrono::PrintClockElapse(Log, TEXT("  "), TEXT("Triangulate"), TriangulateDuration);
-		FChrono::PrintClockElapse(Log, TEXT("    "), TEXT("BuildIsoNodes"), BuildIsoNodesDuration);
-		FChrono::PrintClockElapse(Log, TEXT("    "), TEXT("BuildLoopSegments"), BuildLoopSegmentsDuration);
-		FChrono::PrintClockElapse(Log, TEXT("    "), TEXT("BuildInnerSegments"), BuildInnerSegmentsDuration);
-		FChrono::PrintClockElapse(Log, TEXT("    "), TEXT("FindLoopSegmentOfInnerTriangulation"), FindLoopSegmentOfInnerTriangulationDuration);
-		FChrono::PrintClockElapse(Log, TEXT("      "), TEXT("FindSegmentIsoUVSurroundingSmallLoop"), FindSegmentIsoUVSurroundingSmallLoopDuration);
-		FChrono::PrintClockElapse(Log, TEXT("    "), TEXT("Find IsoSegment ToLink InnerToLoop"), FindIsoSegmentToLinkInnerToLoopDuration);
-		FChrono::PrintClockElapse(Log, TEXT("    "), TEXT("Find Segment ToLink LoopToLoop"), FindSegmentToLinkLoopToLoopDuration);
-		FChrono::PrintClockElapse(Log, TEXT("    "), TEXT("Find Segment ToLink InnerToLoop"), FindSegmentToLinkInnerToLoopDuration);
-		FChrono::PrintClockElapse(Log, TEXT("    "), TEXT("Mesh Over Cycle"), TriangulateOverCycleDuration);
-		FChrono::PrintClockElapse(Log, TEXT("    "), TEXT("Mesh Inner Nodes"), TriangulateInnerNodesDuration);
-	}
-};
-
 class FIsoTriangulator
 {
 	friend class FCycleTriangulator;
diff --git a/Engine/Source/Runtime/Datasmith/CADKernel/Base/Private/Mesh/Meshers/IsoTriangulator/IntersectionIsoSegmentTool.h b/Engine/Source/Runtime/Datasmith/CADKernel/Base/Private/Mesh/Meshers/IsoTriangulator/IntersectionIsoSegmentTool.h
--- a/Engine/Source/Runtime/Datasmith/CADKernel/Base/Private/Mesh/Meshers/IsoTriangulator/IntersectionIsoSegmentTool.h
+++ b/Engine/Source/Runtime/Datasmith/CADKernel/Base/Private/Mesh/Meshers/IsoTriangulator/IntersectionIsoSegmentTool.h
@@ -4,7 +4,6 @@
 #include "Core/Types.h"
 #include "Mesh/Meshers/IsoTriangulator/IsoSegment.h"
 #include "Math/Boundary.h"
-#include "UI/Visu.h"
 
 namespace UE::CADKernel
 {
diff --git a/Engine/Source/Runtime/Datasmith/CADKernel/Base/Private/Mesh/Meshers/IsoTriangulator/IntersectionSegmentTool.h b/Engine/Source/Runtime/Datasmith/CADKernel/Base/Private/Mesh/Meshers/IsoTriangulator/IntersectionSegmentTool.h
--- a/Engine/Source/Runtime/Datasmith/CADKernel/Base/Private/Mesh/Meshers/IsoTriangulator/IntersectionSegmentTool.h
+++ b/Engine/Source/Runtime/Datasmith/CADKernel/Base/Private/Mesh/Meshers/IsoTriangulator/IntersectionSegmentTool.h
@@ -4,7 +4,6 @@
 #include "Core/Types.h"
 #include "Math/Boundary.h"
 #include "Math/Geometry.h"
-#include "UI/Visu.h"
 
 namespace UE::CADKernel
 {
diff --git a/Engine/Source/Runtime/Datasmith/CADKernel/Base/Private/Mesh/Meshers/IsoTriangulator/IsoSegment.h b/Engine/Source/Runtime/Datasmith/CADKernel/Base/Private/Mesh/Meshers/IsoTriangulator/IsoSegment.h
--- a/Engine/Source/Runtime/Datasmith/CADKernel/Base/Private/Mesh/Meshers/IsoTriangulator/IsoSegment.h
+++ b/Engine/Source/Runtime/Datasmith/CADKernel/Base/Private/Mesh/Meshers/IsoTriangulator/IsoSegment.h
@@ -3,7 +3,6 @@
 
 #include "Core/Types.h"
 #include "Mesh/Meshers/IsoTriangulator/IsoNode.h"
-#include "UI/Display.h"
 
 namespace UE::CADKernel
 {
diff --git a/Engine/Source/Runtime/Datasmith/CADKernel/Base/Private/Mesh/Meshers/IsoTriangulatorRemoveIntersection.cpp b/Engine/Source/Runtime/Datasmith/CADKernel/Base/Private/Mesh/Meshers/IsoTriangulatorRemoveIntersection.cpp
--- a/Engine/Source/Runtime/Datasmith/CADKernel/Base/Private/Mesh/Meshers/IsoTriangulatorRemoveIntersection.cpp
+++ b/Engine/Source/Runtime/Datasmith/CADKernel/Base/Private/Mesh/Meshers/IsoTriangulatorRemoveIntersection.cpp
@@ -268,7 +268,7 @@ bool FIsoTriangulator::FindLoopIntersectionAndFixIt()
 		// WARNING From this step, RemoveLoopIntersections can delete nodes. 
 		if (!RemoveLoopIntersections(LoopNodesFromStartNode, Intersections, bIsAnOuterLoop))
 		{
-			FMessage::Printf(Log, TEXT("Loop intersections of the surface %d cannot be fixed. The mesh of this surface is canceled.\n"), Grid.GetFace()->GetId());
+			UE_LOG(LogCADKernelBase, Warning, TEXT("Loop intersections of the surface %d cannot be fixed. The mesh of this surface is canceled.\n"), Grid.GetFace()->GetId());
 			return false;
 		}
 
diff --git a/Engine/Source/Runtime/Datasmith/CADKernel/Base/Private/Mesh/Meshers/ParametricFaceMesher.cpp b/Engine/Source/Runtime/Datasmith/CADKernel/Base/Private/Mesh/Meshers/ParametricFaceMesher.cpp
--- a/Engine/Source/Runtime/Datasmith/CADKernel/Base/Private/Mesh/Meshers/ParametricFaceMesher.cpp
+++ b/Engine/Source/Runtime/Datasmith/CADKernel/Base/Private/Mesh/Meshers/ParametricFaceMesher.cpp
@@ -54,15 +54,9 @@ void FParametricFaceMesher::Mesh()
 
 	MeshVerticesOfFace(Face);
 
-	FMessage::Printf(EVerboseLevel::Debug, TEXT("Meshing of surface %d\n"), Face.GetId());
-
-	FProgress _(1, TEXT("Meshing Entities : Mesh Surface"));
-
-	FTimePoint StartTime = FChrono::Now();
-
 	if (!GenerateCloud() || Grid.IsDegenerated())
 	{
-		FMessage::Printf(EVerboseLevel::Log, TEXT("The meshing of the surface %d failed due to a degenerated grid\n"), Face.GetId());
+		UE_LOG(LogCADKernelBase, Warning, TEXT("The meshing of the surface %d failed due to a degenerated grid\n"), Face.GetId());
 
 		Face.IsDegenerated();
 		Face.SetMeshedMarker();
@@ -99,7 +93,6 @@ bool FParametricFaceMesher::GenerateCloud()
 
 	if (bThinZoneMeshing)
 	{
-		FTimePoint StartTime = FChrono::Now();
 		if (Grid.GetFace().HasThinZone())
 		{
 			MeshThinZones();
diff --git a/Engine/Source/Runtime/Datasmith/CADKernel/Base/Private/Mesh/Meshers/ParametricMesher.cpp b/Engine/Source/Runtime/Datasmith/CADKernel/Base/Private/Mesh/Meshers/ParametricMesher.cpp
--- a/Engine/Source/Runtime/Datasmith/CADKernel/Base/Private/Mesh/Meshers/ParametricMesher.cpp
+++ b/Engine/Source/Runtime/Datasmith/CADKernel/Base/Private/Mesh/Meshers/ParametricMesher.cpp
@@ -112,35 +112,18 @@ void FParametricMesher::PreMeshingTasks()
 {
 	TRACE_CPUPROFILER_EVENT_SCOPE(ParametricMesher::PreMeshingTasks);
 
-	FTimePoint ApplyCriteriaStartTime = FChrono::Now();
-
-	FProgress ProgressBar(Faces.Num() * 2, TEXT("Meshing Entities : Apply Surface Criteria"));
-
 	const TArray<TSharedPtr<FCriterion>>& Criteria = GetMeshModel().GetCriteria();
 
 	// ============================================================================================================
 	//      Apply Surface Criteria
 	// ============================================================================================================
 	const double MeshingTolerance = Tolerances.MeshingTolerance;
-#ifdef CADKernelMultiThread
-	TArray<FTopologicalFace*>& LocalFaces = Faces;
-	ParallelFor(Faces.Num(), [&LocalFaces, &Criteria, &MeshingTolerance, &bThinZone = bThinZoneMeshing](int32 Index) {
-
-		FProgress _(1, TEXT("Meshing Entities : Apply Surface Criteria"));
-
-	FTopologicalFace* Face = LocalFaces[Index];
-#else
 	bool bThinZone = bThinZoneMeshing;
 	for (FTopologicalFace* Face : Faces)
 	{ 
-#endif
 		if (Face == nullptr || Face->IsNotMeshable())
 		{
-#ifdef CADKernelMultiThread
-			return;
-#else
 			continue;
-#endif
 		}
 		ApplyFaceCriteria(*Face, Criteria, MeshingTolerance, bThinZone);
 		if (!Face->IsDeletedOrDegenerated())
@@ -148,34 +131,12 @@ void FParametricMesher::PreMeshingTasks()
 			Face->ComputeSurfaceSideProperties();
 		}
 	}
-#ifdef CADKernelMultiThread
-	, EParallelForFlags::Unbalanced);
-#endif
-
-#ifdef CADKERNEL_DEV
-	FMesherReport::GetChronos().ApplyCriteriaDuration = FChrono::Elapse(ApplyCriteriaStartTime);
-
-	for (FTopologicalFace* Face : Faces)
-	{
-		if (Face->IsDegenerated())
-		{
-			FMesherReport::GetLogs().AddDegeneratedGrid();
-		}
-		FMesherReport::GetLogs().AddThinZone(Face->GetThinZones().Num());
-	}
-#endif
 }
 
 void FParametricMesher::MeshEntities()
 {
 	TRACE_CPUPROFILER_EVENT_SCOPE(ParametricMesher::MeshEntities);
 
-	FTimePoint StartTime = FChrono::Now();
-
-	FProgress ProgressBar(Faces.Num() * 2, TEXT("Meshing Entities : Find quad surfaces"));
-
-	FTimePoint MeshingStartTime = FChrono::Now();
-
 	// ============================================================================================================
 	//      Find and sort quad surfaces 
 	// ============================================================================================================
@@ -184,33 +145,16 @@ void FParametricMesher::MeshEntities()
 
 	if (Faces.Num() > 1)
 	{
-		FTimePoint IsolateQuadPatchStartTime = FChrono::Now();
 		TArray<FTopologicalFace*> OtherEntities;
 
-		FMessage::Printf(Log, TEXT("  Isolate QuadPatch\n"));
-
 		IsolateQuadFace(QuadTrimmedSurfaceSet, OtherEntities);
-
-		FMessage::Printf(Log, TEXT("  %d Quad Surfaces found\n"), QuadTrimmedSurfaceSet.Num());
-
-#ifdef CADKERNEL_DEV
-		FMesherReport::GetLogs().AddQuadPatch(QuadTrimmedSurfaceSet.Num());
-		FMesherReport::GetChronos().IsolateQuadPatchDuration = FChrono::Elapse(IsolateQuadPatchStartTime);
-#endif
 	}
 
 	// ============================================================================================================
 	//      Mesh surfaces 
 	// ============================================================================================================
 
-	FTimePoint MeshStartTime = FChrono::Now();
 	MeshSurfaceByFront(QuadTrimmedSurfaceSet);
-
-#ifdef CADKERNEL_DEV
-	FMesherReport::GetChronos().GlobalMeshDuration = FChrono::Elapse(MeshStartTime);
-	FMesherReport::GetChronos().GlobalDuration = FChrono::Elapse(StartTime);
-#endif
-
 }
 
 void FParametricMesher::ApplyFaceCriteria(FTopologicalFace& Face, const TArray<TSharedPtr<FCriterion>>& Criteria, const double MeshingTolerance, bool bThinZoneMeshing)
@@ -240,17 +184,8 @@ void FParametricMesher::ApplyFaceCriteria(FTopologicalFace& Face, const TArray<T
 
 	if (bThinZoneMeshing)
 	{
-		FTimePoint ThinZone2DFinderStartTime = FChrono::Now();
 		Grid.ScaleGrid();
 
-#ifdef DEBUG_ONLY_SURFACE_TO_DEBUG
-		if (Grid.bDisplay)
-		{
-			Grid.DisplayGridPoints(EGridSpace::Default2D);
-			Grid.DisplayGridPoints(EGridSpace::UniformScaled);
-			Grid.DisplayGridPoints(EGridSpace::Scaled);
-		}
-#endif
 		FThinZone2DFinder ThinZoneFinder(Grid, Face);
 
 		// Size (length of segment of the loop sampling) is equal to MinimalElementLength / ElementRatio
@@ -270,10 +205,6 @@ void FParametricMesher::ApplyFaceCriteria(FTopologicalFace& Face, const TArray<T
 			Face.SetHasThinZoneMarker();
 			Face.MoveThinZones(ThinZoneFinder.GetThinZones());
 		}
-
-#ifdef CADKERNEL_DEV
-		FMesherReport::GetChronos().FindThinSurfaceDuration += FChrono::Elapse(ThinZone2DFinderStartTime);
-#endif
 	}
 
 	if (Face.IsDegenerated())
@@ -381,39 +312,6 @@ void FParametricMesher::IsolateQuadFace(TArray<ParametricMesherTool::FCostToFace
 			return (SurfaceA.Cost > SurfaceB.Cost);
 		}
 	);
-
-#ifdef DEBUG_ISOLATEQUADFACE
-	if (QuadSurfaces.Num() > 0)
-	{
-		F3DDebugSession A(TEXT("Quad Entities"));
-		for (const FCostToFace& Quad : QuadSurfaces)
-		{
-			F3DDebugSession _(FString::Printf(TEXT("Face %d %f"), Quad.Face->GetId(), Quad.Cost));
-			Display(*Quad.Face);
-		}
-	}
-
-	if (FlatQuadsAndTriangles.Num() > 0)
-	{
-		F3DDebugSession A(TEXT("Flat Quads & Triangles"));
-		for (const FTopologicalFace* Face : FlatQuadsAndTriangles)
-		{
-			F3DDebugSession _(FString::Printf(TEXT("Face %d"), Face->GetId()));
-			Display(*Face);
-		}
-	}
-
-	if (OtherSurfaces.Num() > 0)
-	{
-		F3DDebugSession A(TEXT("Other Entities"));
-		for (const FTopologicalFace* Face : OtherSurfaces)
-		{
-			F3DDebugSession _(FString::Printf(TEXT("Face %d"), Face->GetId()));
-			Display(*Face);
-		}
-	}
-	Wait();
-#endif
 }
 
 void FParametricMesher::LinkQuadSurfaceForMesh(TArray<ParametricMesherTool::FCostToFace>& QuadTrimmedSurfaceSet, TArray<TArray<FTopologicalFace*>>& OutStrips)
@@ -583,8 +481,6 @@ void FParametricMesher::MeshSurfaceByFront(TArray<ParametricMesherTool::FCostToF
 	// Marker1 : Surfaces added in CandidateSurfacesForMesh
 	// Marker2 : Surfaces added in SecondChoiceOfCandidateSurfacesForMesh
 
-	FMessage::Printf(EVerboseLevel::Debug, TEXT("Start MeshSurfaceByFront\n"));
-
 	for (FTopologicalFace* Face : Faces)
 	{
 		if (Face == nullptr || Face->IsDeletedOrDegenerated())
@@ -606,22 +502,8 @@ void FParametricMesher::MeshSurfaceByFront(TArray<ParametricMesherTool::FCostToF
 
 	TFunction<void(FTopologicalFace&)> MeshFace = [&](FTopologicalFace& Face)
 	{
-#ifdef DISPLAYDEBUGMESHFACEBYFACESTEP
-		{
-			F3DDebugSession A(FString::Printf(TEXT("Surface %d"), Face.GetId()));
-			Display(Face);
-		}
-#endif
 		Mesh(Face);
 
-#ifdef DISPLAYDEBUGMESHFACEBYFACESTEP
-		{
-			F3DDebugSession A(FString::Printf(TEXT("Mesh %d"), Face.GetId()));
-			DisplayMesh(*Face.GetOrCreateMesh(GetMeshModel()));
-			Wait(bStop);
-		}
-#endif
-
 		if (Face.HasMarker1())
 		{
 			CandidateFacesForMesh.RemoveSingle(&Face);
diff --git a/Engine/Source/Runtime/Datasmith/CADKernel/Base/Private/Mesh/Structure/Grid.h b/Engine/Source/Runtime/Datasmith/CADKernel/Base/Private/Mesh/Structure/Grid.h
--- a/Engine/Source/Runtime/Datasmith/CADKernel/Base/Private/Mesh/Structure/Grid.h
+++ b/Engine/Source/Runtime/Datasmith/CADKernel/Base/Private/Mesh/Structure/Grid.h
@@ -8,11 +8,6 @@
 #include "Mesh/MeshEnum.h"
 #include "Mesh/Structure/GridBase.h"
 #include "Mesh/Structure/ThinZone2D.h"
-#include "UI/Display.h"
-
-#ifdef CADKERNEL_DEV
-#include "UI/DefineForDebug.h"
-#endif
 
 
 namespace UE::CADKernel
@@ -420,182 +415,6 @@ public:
 	{
 		return MinimumElementSize;
 	}
-
-	// ======================================================================================================================================================================================================================
-	// Display Methodes   ================================================================================================================================================================================================
-	// ======================================================================================================================================================================================================================
-#ifdef CADKERNEL_DEV
-	void PrintTimeElapse() const;
-#endif
-
-#ifdef CADKERNEL_DEBUG
-
-	void DisplayIsoNode(EGridSpace Space, const int32 PointIndex, FIdent Ident = 0, EVisuProperty Property = EVisuProperty::BluePoint) const;
-	void DisplayIsoNode(EGridSpace Space, const FIsoNode& Node, FIdent Ident = 0, EVisuProperty Property = EVisuProperty::BluePoint) const;
-	void DisplayIsoNodes(EGridSpace Space, const TArray<const FIsoNode*>& Nodes, EVisuProperty Property = EVisuProperty::BluePoint) const;
-
-	void DisplayIsoNodes(const FString& Message, EGridSpace Space, const TArray<const FIsoNode*>& Nodes, EVisuProperty Property = EVisuProperty::BluePoint) const
-	{
-		if (!bDisplay)
-		{
-			return;
-		}
-		F3DDebugSession _(Message);
-		DisplayIsoNodes(Space, Nodes, Property);
-	}
-
-	void DisplayIsoPolyline(const FString& Message, EGridSpace Space, const TArray<const FIsoNode*>& Nodes, EVisuProperty Property = EVisuProperty::BluePoint) const
-	{
-		if (!bDisplay)
-		{
-			return;
-		}
-		F3DDebugSession _(Message);
-		const FIsoNode* LastNode = Nodes.Last();
-		for (const FIsoNode* Node : Nodes)
-		{
-			DisplayIsoSegment(Space, *LastNode, *Node, 0, (EVisuProperty)(Property + 1));
-			LastNode = Node;
-		}
-		DisplayIsoNodes(Space, Nodes, Property);
-	}
-
-	void DisplayIsoSegment(EGridSpace Space, const FIsoSegment& Segment, FIdent Ident = 0, EVisuProperty Property = EVisuProperty::BlueCurve, bool bDisplayOrientation = false) const;
-	void DisplayIsoSegment(EGridSpace Space, const FIsoNode& NodeA, const FIsoNode& NodeB, FIdent Ident = 0, EVisuProperty Property = EVisuProperty::BlueCurve) const;
-
-	void DisplayIsoSegments(EGridSpace Space, const TArray<FIsoSegment*>& Segments, bool bDisplayNode = false, bool bDisplayOrientation = false, bool bSplitBySegment = false, EVisuProperty Property = EVisuProperty::BlueCurve) const;
-	void DisplayIsoSegments(const FString& Message, EGridSpace Space, const TArray<FIsoSegment*>& Segments, bool bDisplayNode = false, bool bDisplayOrientation = false, bool bSplitBySegment = false, EVisuProperty Property = EVisuProperty::BlueCurve) const
-	{
-		if (!bDisplay)
-		{
-			return;
-		}
-		F3DDebugSession _(Message);
-		DisplayIsoSegments(Space, Segments, bDisplayNode, bDisplayOrientation, bSplitBySegment, Property);
-	}
-
-	void DisplayTriangle(EGridSpace Space, const FIsoNode& NodeA, const FIsoNode& NodeB, const FIsoNode& NodeC) const;
-	void DisplayTriangle3D(const FIsoNode& NodeA, const FIsoNode& NodeB, const FIsoNode& NodeC) const;
-
-	void DisplayGridPolyline(EGridSpace Space, const FLoopNode& StartNode, bool bDisplayNode, EVisuProperty Property = EVisuProperty::BlueCurve) const;
-	void DisplayGridPolyline(EGridSpace Space, const TArray<FIsoNode*>& Nodes, bool bDisplayNode, EVisuProperty Property = EVisuProperty::BlueCurve) const;
-	void DisplayGridPolyline(EGridSpace Space, const TArray<FLoopNode*>& Nodes, bool bDisplayNode, EVisuProperty Property = EVisuProperty::BlueCurve) const;
-	void DisplayGridPolyline(EGridSpace Space, const TArray<FLoopNode>& Nodes, bool bDisplayNode, EVisuProperty Property = EVisuProperty::BlueCurve) const;
-	void DisplayGridPolyline(const FString& Message, EGridSpace Space, const TArray<FLoopNode*>& Nodes, bool bDisplayNode, EVisuProperty Property = EVisuProperty::BlueCurve) const
-	{
-		if (!bDisplay)
-		{
-			return;
-		}
-		F3DDebugSession _(Message);
-		DisplayGridPolyline(Space, Nodes, bDisplayNode, Property);
-	}
-
-	void DisplayGridPolyline(const FString& Message, EGridSpace Space, const TArray<FLoopNode>& Nodes, bool bDisplayNode, EVisuProperty Property = EVisuProperty::BlueCurve) const
-	{
-		if (!bDisplay)
-		{
-			return;
-		}
-		F3DDebugSession _(Message);
-		DisplayGridPolyline(Space, Nodes, bDisplayNode, Property);
-	}
-
-	void DisplayNodes(const TCHAR* Message, EGridSpace Space, const TArray<const FIsoNode*>& Nodes, EVisuProperty Property) const;
-
-	virtual void DisplayGridPoints(EGridSpace DisplaySpace) const override;
-	void DisplayInnerPoints(TCHAR* Message, EGridSpace DisplaySpace) const;
-
-	void DisplayGridNormal() const;
-
-	template<typename TPoint>
-	void DisplayInnerDomainPoints(FString Message, const TArray<TPoint>& Points) const
-	{
-		if (!bDisplay)
-		{
-			return;
-		}
-
-		F3DDebugSession _(Message);
-		for (int32 Index = 0; Index < Points.Num(); ++Index)
-		{
-			if (IsNodeInsideAndCloseToLoop(Index))
-			{
-				DisplayPoint2DWithScale(Points[Index], EVisuProperty::OrangePoint, Index);
-			}
-			else if (IsNodeInsideAndMeshable(Index))
-			{
-				DisplayPoint2DWithScale(Points[Index], EVisuProperty::BluePoint, Index);
-			}
-			else if (IsNodeTooCloseToLoop(Index))
-			{
-				DisplayPoint2DWithScale(Points[Index], EVisuProperty::RedPoint, Index);
-			}
-			else
-			{
-				DisplayPoint2DWithScale(Points[Index], IsNodeCloseToLoop(Index) ? EVisuProperty::YellowPoint : EVisuProperty::GreenPoint, Index);
-			}
-		}
-	}
-
-	template<typename TPoint>
-	void DisplayGridLoop(FString Message, const TArray<TPoint>& Loop, bool bDisplayNodes, bool bMakeGroup, bool bDisplaySegments) const
-	{
-		if (!bDisplay)
-		{
-			return;
-		}
-
-		F3DDebugSession _(bMakeGroup, Message);
-
-		const TPoint* FirstSegmentPoint = &Loop[0];
-		if (bDisplayNodes)
-		{
-			F3DDebugSession Q(bDisplaySegments, FString::Printf(TEXT("FirstNode")));
-			DisplayPoint(*FirstSegmentPoint * DisplayScale, EVisuProperty::BluePoint, 0);
-		}
-
-		for (int32 Index = 1; Index < Loop.Num(); ++Index)
-		{
-			const TPoint* SecondSegmentPoint = &Loop[Index];
-			F3DDebugSession B(bDisplaySegments, FString::Printf(TEXT("Segment %d"), Index));
-			DisplaySegment<TPoint>(*FirstSegmentPoint * DisplayScale, *SecondSegmentPoint * DisplayScale);
-			if (bDisplayNodes)
-			{
-				DisplayPoint<TPoint>(*SecondSegmentPoint * DisplayScale, EVisuProperty::BluePoint, Index);
-			}
-			FirstSegmentPoint = SecondSegmentPoint;
-		}
-
-		{
-			F3DDebugSession C(bDisplaySegments, FString::Printf(TEXT("Segment %d"), Loop.Num()));
-			DisplaySegment(*FirstSegmentPoint * DisplayScale, *Loop.begin() * DisplayScale);
-		}
-	}
-
-	template<typename TPoint>
-	void DisplayGridLoops(FString Message, const TArray<TArray<TPoint>>& Loops, bool bDisplayNodes, bool bMakeGroup, bool bDisplaySegments) const
-	{
-		if (!bDisplay)
-		{
-			return;
-		}
-
-		F3DDebugSession _(*Message);
-		int32 LoopIndex = 0;
-		for (const TArray<TPoint>& Loop : Loops)
-		{
-			FString LoopName = FString::Printf(TEXT("Loop %d"), LoopIndex++);
-			DisplayGridLoop(LoopName, Loop, bDisplayNodes, bMakeGroup, bDisplaySegments);
-		}
-	}
-
-	void DisplayGridLoops(FString Message, EGridSpace DisplaySpace, bool bDisplayNodes, bool bMakeGroup, bool bDisplaySegments) const
-	{
-		DisplayGridLoops(Message, FaceLoops2D[DisplaySpace], bDisplayNodes, bMakeGroup, bDisplaySegments);
-	}
-
-#endif
 };
 
 }
diff --git a/Engine/Source/Runtime/Datasmith/CADKernel/Base/Private/Mesh/Structure/GridBase.cpp b/Engine/Source/Runtime/Datasmith/CADKernel/Base/Private/Mesh/Structure/GridBase.cpp
--- a/Engine/Source/Runtime/Datasmith/CADKernel/Base/Private/Mesh/Structure/GridBase.cpp
+++ b/Engine/Source/Runtime/Datasmith/CADKernel/Base/Private/Mesh/Structure/GridBase.cpp
@@ -26,8 +26,6 @@ bool FGridBase::ScaleGrid()
 {
 	const FCoordinateGrid& CoordinateGrid = GetCoordinateGrid();
 
-	FTimePoint StartTime = FChrono::Now();
-
 	TFunction<double(const TArray<double>&)> GetMean = [](const TArray<double>& Lengths)
 	{
 		double MeanLength = 0;
@@ -211,9 +209,7 @@ bool FGridBase::ScaleGrid()
 			}
 		}
 	}
-#ifdef CADKERNEL_DEV
-	Chronos.ScaleGridDuration = FChrono::Elapse(StartTime);
-#endif
+
 	return true;
 }
 
diff --git a/Engine/Source/Runtime/Datasmith/CADKernel/Base/Private/Mesh/Structure/GridBase.h b/Engine/Source/Runtime/Datasmith/CADKernel/Base/Private/Mesh/Structure/GridBase.h
--- a/Engine/Source/Runtime/Datasmith/CADKernel/Base/Private/Mesh/Structure/GridBase.h
+++ b/Engine/Source/Runtime/Datasmith/CADKernel/Base/Private/Mesh/Structure/GridBase.h
@@ -7,11 +7,6 @@
 #include "Geo/GeoPoint.h"
 #include "Math/Point.h"
 #include "Mesh/MeshEnum.h"
-#include "UI/Display.h"
-
-#ifdef CADKERNEL_DEV
-#include "UI/DefineForDebug.h"
-#endif
 
 
 namespace UE::CADKernel
@@ -27,48 +22,6 @@ class FModelMesh;
 class FThinZone2DFinder;
 class FTopologicalFace;
 
-struct FGridChronos
-{
-	FDuration DefineCuttingParametersDuration;
-	FDuration GeneratePointCloudDuration;
-	FDuration ProcessPointCloudDuration;
-	FDuration FindInnerDomainPointsDuration;
-	FDuration Build2DLoopDuration;
-	FDuration RemovePointsClosedToLoopDuration;
-	FDuration FindPointsCloseToLoopDuration;
-	FDuration ScaleGridDuration;
-
-	FGridChronos()
-		: DefineCuttingParametersDuration(FChrono::Init())
-		, GeneratePointCloudDuration(FChrono::Init())
-		, ProcessPointCloudDuration(FChrono::Init())
-		, FindInnerDomainPointsDuration(FChrono::Init())
-		, Build2DLoopDuration(FChrono::Init())
-		, RemovePointsClosedToLoopDuration(FChrono::Init())
-		, FindPointsCloseToLoopDuration(FChrono::Init())
-		, ScaleGridDuration(FChrono::Init())
-	{}
-
-	void PrintTimeElapse() const
-	{
-		FDuration GridDuration = FChrono::Elapse(FChrono::Now());
-		GridDuration += DefineCuttingParametersDuration;
-		GridDuration += GeneratePointCloudDuration;
-		GridDuration += FindInnerDomainPointsDuration;
-		GridDuration += Build2DLoopDuration;
-		GridDuration += RemovePointsClosedToLoopDuration;
-		GridDuration += FindPointsCloseToLoopDuration;
-		FChrono::PrintClockElapse(Log, TEXT(""), TEXT("Grid"), GridDuration);
-		FChrono::PrintClockElapse(Log, TEXT("  "), TEXT("DefineCuttingParameters"), DefineCuttingParametersDuration);
-		FChrono::PrintClockElapse(Log, TEXT("  "), TEXT("GeneratePointCloud"), GeneratePointCloudDuration);
-		FChrono::PrintClockElapse(Log, TEXT("  "), TEXT("GenerateDomainPoints"), ProcessPointCloudDuration);
-		FChrono::PrintClockElapse(Log, TEXT("    "), TEXT("FindInnerDomainPointsDuration"), FindInnerDomainPointsDuration);
-		FChrono::PrintClockElapse(Log, TEXT("    "), TEXT("FindPointsCloseToLoop"), FindPointsCloseToLoopDuration);
-		FChrono::PrintClockElapse(Log, TEXT("    "), TEXT("RemovePointsClosedToLoop"), RemovePointsClosedToLoopDuration);
-		FChrono::PrintClockElapse(Log, TEXT("    "), TEXT("Build2DLoopDuration"), Build2DLoopDuration);
-	}
-};
-
 class FGridBase : public FHaveStates
 {
 
@@ -100,9 +53,6 @@ protected:
 	 */
 	TArray<FVector3f> Normals;
 
-public:
-	FGridChronos Chronos;
-
 public:
 	FGridBase(FTopologicalFace& InFace);
 
diff --git a/Engine/Source/Runtime/Datasmith/CADKernel/Base/Private/Mesh/Structure/LoopCleaner.cpp b/Engine/Source/Runtime/Datasmith/CADKernel/Base/Private/Mesh/Structure/LoopCleaner.cpp
--- a/Engine/Source/Runtime/Datasmith/CADKernel/Base/Private/Mesh/Structure/LoopCleaner.cpp
+++ b/Engine/Source/Runtime/Datasmith/CADKernel/Base/Private/Mesh/Structure/LoopCleaner.cpp
@@ -109,7 +109,7 @@ bool FLoopCleaner::CleanLoops()
 		{
 			if (StartSegmentIndex == 0)
 			{
-				FMessage::Printf(Log, TEXT("The outer loop of the surface %d is degenerated. The mesh of this surface is canceled.\n"), Grid.GetFace().GetId());
+				UE_LOG(LogCADKernelBase, Warning, TEXT("The outer loop of the surface %d is degenerated. The mesh of this surface is canceled.\n"), Grid.GetFace().GetId());
 				return false;
 			}
 			continue;
@@ -119,11 +119,6 @@ bool FLoopCleaner::CleanLoops()
 
 		GetLoopNodeStartingFrom(StartNode, NodesOfLoop);
 
-#ifdef DEBUG_CLEAN_LOOPS		
-		Grid.DisplayGridPolyline(TEXT("Loop: start"), EGridSpace::UniformScaled, NodesOfLoop, true, EVisuProperty::YellowCurve);
-		//Wait(bDisplay);
-#endif
-
 		if (!RemoveLoopPicks())
 		{
 			continue;
@@ -149,7 +144,7 @@ bool FLoopCleaner::CleanLoops()
 		{
 			if (StartSegmentIndex == 0)
 			{
-				FMessage::Printf(Log, TEXT("The outer loop of the surface %d is degenerated. The mesh of this surface is canceled.\n"), Grid.GetFace().GetId());
+				UE_LOG(LogCADKernelBase, Warning, TEXT("The outer loop of the surface %d is degenerated. The mesh of this surface is canceled.\n"), Grid.GetFace().GetId());
 				return false;
 			}
 			continue;
@@ -166,7 +161,7 @@ bool FLoopCleaner::CleanLoops()
 		{
 			if (!RemoveSelfIntersectionsOfLoop())
 			{
-				FMessage::Printf(Log, TEXT("Loop intersections of the surface %d cannot be fixed. The mesh of this surface is canceled.\n"), Grid.GetFace().GetId());
+				UE_LOG(LogCADKernelBase, Warning, TEXT("Loop intersections of the surface %d cannot be fixed. The mesh of this surface is canceled.\n"), Grid.GetFace().GetId());
 				return false;
 			}
 
@@ -174,7 +169,7 @@ bool FLoopCleaner::CleanLoops()
 			FindLoopIntersections();
 			if (Intersections.Num())
 			{
-				FMessage::Printf(Log, TEXT("Loop intersections of the surface %d cannot be fixed. The mesh of this surface is canceled.\n"), Grid.GetFace().GetId());
+				UE_LOG(LogCADKernelBase, Warning, TEXT("Loop intersections of the surface %d cannot be fixed. The mesh of this surface is canceled.\n"), Grid.GetFace().GetId());
 				return false;
 			}
 #ifdef DEBUG_CLEAN_LOOPS		
@@ -185,7 +180,7 @@ bool FLoopCleaner::CleanLoops()
 
 		if (!RemoveLoopPicks())
 		{
-			FMessage::Printf(Log, TEXT("Loop picks of the surface %d cannot be fixed. The mesh of this surface is canceled.\n"), Grid.GetFace().GetId());
+			UE_LOG(LogCADKernelBase, Warning, TEXT("Loop picks of the surface %d cannot be fixed. The mesh of this surface is canceled.\n"), Grid.GetFace().GetId());
 			return false;
 		}
 
@@ -564,10 +559,10 @@ void FLoopCleaner::RemoveIntersectionByMovingOutsideSegmentNodeInside(const FIso
 	SegmentPoints[0] = Segment.GetFirstNode().Get2DPoint(EGridSpace::UniformScaled, Grid);
 	SegmentPoints[1] = Segment.GetSecondNode().Get2DPoint(EGridSpace::UniformScaled, Grid);
 
-	double SquareLenghtIntersectingSegment = FVector2d::DistSquared(IntersectingPoints[0], IntersectingPoints[1]);
-	double SquareLenghtSegment = FVector2d::DistSquared(SegmentPoints[0], SegmentPoints[1]);
+	double SquareLengthIntersectingSegment = FVector2d::DistSquared(IntersectingPoints[0], IntersectingPoints[1]);
+	double SquareLengthSegment = FVector2d::DistSquared(SegmentPoints[0], SegmentPoints[1]);
 
-	if (SquareLenghtSegment > 10 * SquareLenghtIntersectingSegment)
+	if (SquareLengthSegment > 10 * SquareLengthIntersectingSegment)
 	{
 		return RemoveIntersectionByMovingOutsideSegmentNodeInside(Segment, const_cast<FIsoSegment&>(IntersectingSegment), bIsSameInnerLoop);
 	}
diff --git a/Engine/Source/Runtime/Datasmith/CADKernel/Base/Private/Mesh/Structure/ModelMeshAnalyzer.cpp b/Engine/Source/Runtime/Datasmith/CADKernel/Base/Private/Mesh/Structure/ModelMeshAnalyzer.cpp
--- a/Engine/Source/Runtime/Datasmith/CADKernel/Base/Private/Mesh/Structure/ModelMeshAnalyzer.cpp
+++ b/Engine/Source/Runtime/Datasmith/CADKernel/Base/Private/Mesh/Structure/ModelMeshAnalyzer.cpp
@@ -10,7 +10,6 @@
 #include "Mesh/Structure/FaceMesh.h"
 #include "Mesh/Structure/VertexMesh.h"
 #include "Topo/TopologicalEntity.h"
-#include "UI/Display.h"
 
 #include "Templates/UnrealTemplate.h"
 
@@ -211,55 +210,6 @@ bool FModelMeshAnalyzer::CheckOrientation() const
 	return true;
 }
 
-void FModelMeshAnalyzer::Display()
-{
-	F3DDebugSession A(FString::Printf(TEXT("ModelMesh %d"), ModelMesh.GetId()));
-	{
-		F3DDebugSegment B(0);
-		for (const Analyzer::FTriangle& Triangle : Triangles)
-		{
-			DisplayTriangle(Triangle);
-		}
-	}
-
-	{
-		F3DDebugSegment B(ModelMesh.GetId());
-		for (const Analyzer::FEdge& Edge : Edges)
-		{
-			switch (Edge.TriangleCount)
-			{
-			case 1:
-				DisplayEdge(Edge, EVisuProperty::BorderEdge);
-				break;
-
-			case 2:
-				DisplayEdge(Edge, EVisuProperty::EdgeMesh);
-				break;
-
-			default:
-				DisplayEdge(Edge, EVisuProperty::NonManifoldEdge);
-				break;
-			}
-		}
-	}
-
-}
-
-void FModelMeshAnalyzer::DisplayTriangle(const Analyzer::FTriangle& Triangle)
-{
-	TArray<FVector> Vertices;
-	Vertices.Reserve(3);
-	Vertices.Add(NodeCoordinates[Triangle.Vertices[0]]);
-	Vertices.Add(NodeCoordinates[Triangle.Vertices[1]]);
-	Vertices.Add(NodeCoordinates[Triangle.Vertices[2]]);
-	DrawElement(2, Vertices);
-}
-
-void FModelMeshAnalyzer::DisplayEdge(const Analyzer::FEdge& Edge, EVisuProperty Property)
-{
-	DrawSegment( NodeCoordinates[Edge.VertexIndices[0]], NodeCoordinates[Edge.VertexIndices[1]], Property);
-}
-
 FVector Analyzer::FTriangle::ComputeNormal(const TArray<FVector>& NodeCoordinates) const
 {
 	return UE::CADKernel::FTriangle(NodeCoordinates[Vertices[0]], NodeCoordinates[Vertices[1]], NodeCoordinates[Vertices[2]]).ComputeNormal();
diff --git a/Engine/Source/Runtime/Datasmith/CADKernel/Base/Private/Mesh/Structure/ModelMeshAnalyzer.h b/Engine/Source/Runtime/Datasmith/CADKernel/Base/Private/Mesh/Structure/ModelMeshAnalyzer.h
--- a/Engine/Source/Runtime/Datasmith/CADKernel/Base/Private/Mesh/Structure/ModelMeshAnalyzer.h
+++ b/Engine/Source/Runtime/Datasmith/CADKernel/Base/Private/Mesh/Structure/ModelMeshAnalyzer.h
@@ -3,7 +3,6 @@
 
 #include "Core/HaveStates.h"
 #include "Core/Types.h"
-#include "UI/Visu.h"
 
 #include "Math/Point.h"
 
@@ -97,13 +96,6 @@ public:
 	void ComputeMeshGapCount(int32& OutCycleCount, int32& OutChainCount) const;
 	double ComputeMaxAngle() const;
 	bool CheckOrientation() const;
-
-	void Display();
-
-private:
-	void DisplayTriangle(const Analyzer::FTriangle& Triangle);
-	void DisplayEdge(const Analyzer::FEdge& Edge, EVisuProperty Property);
-
 };
 
 }
diff --git a/Engine/Source/Runtime/Datasmith/CADKernel/Base/Private/Mesh/Structure/ThinZone2DFinder.cpp b/Engine/Source/Runtime/Datasmith/CADKernel/Base/Private/Mesh/Structure/ThinZone2DFinder.cpp
--- a/Engine/Source/Runtime/Datasmith/CADKernel/Base/Private/Mesh/Structure/ThinZone2DFinder.cpp
+++ b/Engine/Source/Runtime/Datasmith/CADKernel/Base/Private/Mesh/Structure/ThinZone2DFinder.cpp
@@ -10,7 +10,6 @@
 #include "Topo/TopologicalEdge.h"
 #include "Topo/TopologicalFace.h"
 #include "Topo/TopologicalLoop.h"
-#include "UI/Display.h"
 
 namespace UE::CADKernel
 {
@@ -1261,99 +1260,25 @@ void FThinZone2DFinder::BuildThinZone(const TArray<FEdgeSegment*>& FirstSide, co
 
 bool FThinZone2DFinder::SearchThinZones(double InTolerance)
 {
-#ifdef DEBUG_THIN_ZONES
-	F3DDebugSession A(bDisplay, FString::Printf(TEXT("ThinZone Finder %d"), Grid.GetFace().GetId()));
-#endif
-
-#ifdef DEBUG_SEARCH_THIN_ZONES
-	if (bDisplay)
-	{
-		F3DDebugSession _(TEXT("Thin surface grid"));
-		Grid.DisplayGridPoints(EGridSpace::UniformScaled);
-	}
-#endif
-
 	SetTolerance(InTolerance);
 
-	//FMessage::Printf(DBG, TEXT("Searching thin zones on Surface %d\n", Surface->GetId());
-	FTimePoint StartTime = FChrono::Now();
 	if (!BuildLoopSegments())
 	{
 		Face.SetAsDegenerated();
 		return false;
 	}
-	Chronos.BuildLoopSegmentsTime = FChrono::Elapse(StartTime);
-
-#ifdef DEBUG_SEARCH_THIN_ZONES
-	DisplayLoopSegments();
-#endif
 
-	StartTime = FChrono::Now();
 	FindCloseSegments();
-	Chronos.FindCloseSegmentTime = FChrono::Elapse(StartTime);
-
-#ifdef DEBUG_SEARCH_THIN_ZONES
-	if (bDisplay)
-	{
-		DisplayCloseSegments();
-		Wait(false);
-	}
-#endif
 
-	StartTime = FChrono::Now();
 	LinkCloseSegments();
 
-#ifdef DEBUG_SEARCH_THIN_ZONES
-	if (bDisplay)
-	{
-		ThinZone::DisplayThinZoneSides(ThinZoneSides);
-		Wait(false);
-	}
-#endif
-
 	CheckIfCloseSideOfThinSidesAreNotDegenerated();
 
-#ifdef DEBUG_SEARCH_THIN_ZONES
-	if (bDisplay)
-	{
-		ThinZone::DisplayThinZoneSides(ThinZoneSides);
-		Wait(false);
-	}
-#endif
-
 	ImproveThinSide();
 
-#ifdef DEBUG_SEARCH_THIN_ZONES
-	if (bDisplay)
-	{
-		ThinZone::DisplayThinZoneSides2(ThinZoneSides);
-		Wait(false);
-	}
-#endif
-
 	SplitThinSide();
 
-#ifdef DEBUG_SEARCH_THIN_ZONES
-	if (bDisplay)
-	{
-		ThinZone::DisplayThinZoneSidesAndCloses(ThinZoneSides);
-		Wait(false);
-	}
-#endif
-
-	Chronos.LinkCloseSegmentTime = FChrono::Elapse(StartTime);
-
-	StartTime = FChrono::Now();
 	BuildThinZone();
-	Chronos.BuildThinZoneTime = FChrono::Elapse(StartTime);
-
-#ifdef DEBUG_THIN_ZONES
-	if (bDisplay)
-	{
-		ThinZone::DisplayThinZones(ThinZones);
-		//Wait();
-	}
-#endif
 
 	return (ThinZones.Num() > 0);
 }
diff --git a/Engine/Source/Runtime/Datasmith/CADKernel/Base/Private/Mesh/Structure/ThinZone2DFinder.h b/Engine/Source/Runtime/Datasmith/CADKernel/Base/Private/Mesh/Structure/ThinZone2DFinder.h
--- a/Engine/Source/Runtime/Datasmith/CADKernel/Base/Private/Mesh/Structure/ThinZone2DFinder.h
+++ b/Engine/Source/Runtime/Datasmith/CADKernel/Base/Private/Mesh/Structure/ThinZone2DFinder.h
@@ -11,10 +11,6 @@
 #include "Mesh/Structure/GridBase.h"
 #include "Mesh/Structure/ThinZone2D.h"
 
-#ifdef CADKERNEL_DEBUG
-#include "UI/DefineForDebug.h"
-#endif
-
 namespace UE::CADKernel
 {
 	class FTopologicalLoop;
@@ -23,22 +19,10 @@ namespace UE::CADKernel
 	class FThinZone2D;
 	class FThinZone2DFinder;
 
-	struct FThinZoneChronos
-	{
-		FDuration BuildLoopSegmentsTime;
-		FDuration DisplayBoundarySegmentTime;
-		FDuration FindCloseSegmentTime;
-		FDuration LinkCloseSegmentTime;
-		FDuration BuildThinZoneTime;
-	};
-
 	class FThinZone2DFinder
 	{
 		friend class FThinZone2D;
 
-	public:
-		FThinZoneChronos Chronos;
-
 	protected:
 		const FGridBase& Grid;
 		const FTopologicalFace& Face;
diff --git a/Engine/Source/Runtime/Datasmith/CADKernel/Base/Private/Topo/FaceAnalyzer.cpp b/Engine/Source/Runtime/Datasmith/CADKernel/Base/Private/Topo/FaceAnalyzer.cpp
--- a/Engine/Source/Runtime/Datasmith/CADKernel/Base/Private/Topo/FaceAnalyzer.cpp
+++ b/Engine/Source/Runtime/Datasmith/CADKernel/Base/Private/Topo/FaceAnalyzer.cpp
@@ -6,7 +6,6 @@
 #include "Topo/TopologicalEdge.h"
 #include "Topo/TopologicalFace.h"
 #include "Topo/TopologicalLoop.h"
-#include "UI/Display.h"
 
 namespace UE::CADKernel
 {
@@ -252,24 +251,11 @@ bool FFaceAnalyzer::IsThinFace(double& OutGapSize)
 
 	Topo::FThinFaceContext Context(*Loop);
 
-	FTimePoint StartTime = FChrono::Now();
 	BuildLoopSegments(Context);
-	Chronos.BuildLoopSegmentsTime = FChrono::Elapse(StartTime);
 
-#ifdef CADKERNEL_DEV
-	DisplayLoopSegments(Context);
-#endif
-
-	StartTime = FChrono::Now();
 	FindClosedSegments(Context);
-	Chronos.FindClosedSegmentTime = FChrono::Elapse(StartTime);
-#ifdef CADKERNEL_DEV
-	DisplayCloseSegments(Context);
-#endif
 
-	StartTime = FChrono::Now();
 	Analyze(Context);
-	Chronos.AnalyzeClosedSegmentTime = FChrono::Elapse(StartTime);
 
 	if (Context.OppositSideEdgeLength < MaxOppositSideLength || (Context.ThinSideEdgeLength > Context.OppositSideEdgeLength && Context.MaxSquareDistance < SquareTolerance ) )
 	{
diff --git a/Engine/Source/Runtime/Datasmith/CADKernel/Base/Private/Topo/FaceAnalyzer.h b/Engine/Source/Runtime/Datasmith/CADKernel/Base/Private/Topo/FaceAnalyzer.h
--- a/Engine/Source/Runtime/Datasmith/CADKernel/Base/Private/Topo/FaceAnalyzer.h
+++ b/Engine/Source/Runtime/Datasmith/CADKernel/Base/Private/Topo/FaceAnalyzer.h
@@ -22,13 +22,6 @@ class FTopologicalLoop;
 namespace Topo
 {
 
-struct FFaceAnalyzerChronos
-{
-	FDuration BuildLoopSegmentsTime;
-	FDuration FindClosedSegmentTime;
-	FDuration AnalyzeClosedSegmentTime;
-};
-
 class FEdgeSegment : public FHaveStates
 {
 private:
@@ -182,9 +175,6 @@ struct FThinFaceContext
 class FFaceAnalyzer
 {
 
-public:
-	Topo::FFaceAnalyzerChronos Chronos;
-
 protected:
 	const double Tolerance;
 	const double SquareTolerance;
diff --git a/Engine/Source/Runtime/Datasmith/CADKernel/Base/Private/Topo/Model.cpp b/Engine/Source/Runtime/Datasmith/CADKernel/Base/Private/Topo/Model.cpp
--- a/Engine/Source/Runtime/Datasmith/CADKernel/Base/Private/Topo/Model.cpp
+++ b/Engine/Source/Runtime/Datasmith/CADKernel/Base/Private/Topo/Model.cpp
@@ -7,10 +7,6 @@
 #include "Topo/TopologicalFace.h"
 #include "Topo/TopologicalVertex.h"
 
-#ifdef CADKERNEL_DEV
-#include "Topo/TopologyReport.h"
-#endif
-
 namespace UE::CADKernel
 {
 
@@ -37,20 +33,6 @@ bool FModel::Contains(TSharedPtr<FTopologicalEntity> Entity)
 	}
 }
 
-void FModel::PrintBodyAndShellCount()
-{
-	int32 NbBody = 0;
-	int32 NbShell = 0;
-
-	TArray<TSharedPtr<FBody>> NewBodies;
-	for (TSharedPtr<FBody> Body : Bodies)
-	{
-		NbShell += Body->GetShells().Num();
-		NbBody++;
-	}
-	FMessage::Printf(Log, TEXT("Body count %d shell count %d \n"), NbBody, NbShell);
-}
-
 int32 FModel::FaceCount() const
 {
 	int32 FaceCount = 0;
@@ -144,51 +126,27 @@ void FModel::CheckTopology()
 				{
 					if (SubShells[0].BorderEdgeCount > 0 || SubShells[0].NonManifoldEdgeCount > 0)
 					{
-#ifdef CORETECHBRIDGE_DEBUG
-						FMessage::Printf(Log, TEXT("Body %d shell %d CADId %d is opened and has %d faces "), Body->GetKioId(), Shell->GetKioId(), Shell->GetId(), Shell->FaceCount());
-#else
-						FMessage::Printf(Log, TEXT("Body %d shell %d is opened and has %d faces "), Body->GetId(), Shell->GetId(), Shell->FaceCount());
-#endif
-						FMessage::Printf(Log, TEXT("and has %d border edges and %d nonManifold edges\n"), SubShells[0].BorderEdgeCount, SubShells[0].NonManifoldEdgeCount);
+						UE_LOG(LogCADKernelBase, Verbose, TEXT("Body %d shell %d is opened and has %d faces "), Body->GetId(), Shell->GetId(), Shell->FaceCount());
+						UE_LOG(LogCADKernelBase, Verbose, TEXT("and has %d border edges and %d nonManifold edges\n"), SubShells[0].BorderEdgeCount, SubShells[0].NonManifoldEdgeCount);
 					}
 					else
 					{
-#ifdef CORETECHBRIDGE_DEBUG
-						FMessage::Printf(Log, TEXT("Body %d shell %d CADId %d is closed and has %d faces\n"), Body->GetKioId(), Shell->GetKioId(), Shell->GetId(), Shell->FaceCount());
-#else
-						FMessage::Printf(Log, TEXT("Body %d shell %d is closed and has %d faces\n"), Body->GetId(), Shell->GetId(), Shell->FaceCount());
-#endif
+						UE_LOG(LogCADKernelBase, Verbose, TEXT("Body %d shell %d is closed and has %d faces\n"), Body->GetId(), Shell->GetId(), Shell->FaceCount());
 					}
 				}
 			}
 			else
 			{
-#ifdef CORETECHBRIDGE_DEBUG
-				FMessage::Printf(Log, TEXT("Body %d shell %d CADId %d has %d subshells\n"), Body->GetKioId(), Shell->GetKioId(), Shell->GetId(), SubShells.Num());
-#else
-				FMessage::Printf(Log, TEXT("Body %d shell %d has %d subshells\n"), Body->GetId(), Shell->GetId(), SubShells.Num());
-#endif
-for (const FFaceSubset& Subset : SubShells)
+				UE_LOG(LogCADKernelBase, Verbose, TEXT("Body %d shell %d has %d subshells\n"), Body->GetId(), Shell->GetId(), SubShells.Num());
+				for (const FFaceSubset& Subset : SubShells)
 				{
-					FMessage::Printf(Log, TEXT("     - Subshell of %d faces %d border edges and %d nonManifold edges\n"), Subset. Faces.Num(), Subset.BorderEdgeCount, Subset.NonManifoldEdgeCount);
+					UE_LOG(LogCADKernelBase, Verbose, TEXT("     - Subshell of %d faces %d border edges and %d nonManifold edges\n"), Subset. Faces.Num(), Subset.BorderEdgeCount, Subset.NonManifoldEdgeCount);
 				}
 			}
 		}
 	}
 }
 
-#ifdef CADKERNEL_DEV
-void FModel::FillTopologyReport(FTopologyReport& Report) const
-{
-	for (TSharedPtr<FBody> Body : Bodies)
-	{
-		Body->FillTopologyReport(Report);
-	}
-
-	Report.ResetMarkers();
-}
-#endif
-
 void FModel::Orient()
 {
 	for (TSharedPtr<FBody> Body : Bodies)
diff --git a/Engine/Source/Runtime/Datasmith/CADKernel/Base/Private/Topo/Shell.cpp b/Engine/Source/Runtime/Datasmith/CADKernel/Base/Private/Topo/Shell.cpp
--- a/Engine/Source/Runtime/Datasmith/CADKernel/Base/Private/Topo/Shell.cpp
+++ b/Engine/Source/Runtime/Datasmith/CADKernel/Base/Private/Topo/Shell.cpp
@@ -456,9 +456,7 @@ int32 FShell::Orient()
 	{
 		if (SubshellSwappedFaceCount > (SubshellFaceCount / 2))
 		{
-#ifdef CADERNEL_DEV
-			FMessage::Printf(Log, TEXT("the open shell %s (Id %d) is badly oriented. Its orientation is swapped"), *GetName(), GetId());
-#endif
+			UE_LOG(LogCADKernelBase, Verbose, TEXT("the open shell %s (Id %d) is badly oriented. Its orientation is swapped"), GetName(), GetId());
 			SwapSubshellOrientation();
 			ShellSwappedFaceCount += (SubshellFaceCount - SubshellSwappedFaceCount);
 		}
@@ -499,7 +497,7 @@ int32 FShell::Orient()
 		if (bHasWrongOrientation)
 		{
 #ifdef CADERNEL_DEV
-			FMessage::Printf(Log, TEXT("The closed shell %s (Id %d) is badly oriented. Its orientation is swapped"), *GetName(), GetId());
+			UE_LOG(LogCADKernelBase, Verbose, TEXT("The closed shell %s (Id %d) is badly oriented. Its orientation is swapped"), *GetName(), GetId());
 #endif
 			SwapSubshellOrientation();
 			ShellSwappedFaceCount += (SubshellFaceCount - SubshellSwappedFaceCount);
diff --git a/Engine/Source/Runtime/Datasmith/CADKernel/Base/Private/Topo/TopologicalEdge.cpp b/Engine/Source/Runtime/Datasmith/CADKernel/Base/Private/Topo/TopologicalEdge.cpp
--- a/Engine/Source/Runtime/Datasmith/CADKernel/Base/Private/Topo/TopologicalEdge.cpp
+++ b/Engine/Source/Runtime/Datasmith/CADKernel/Base/Private/Topo/TopologicalEdge.cpp
@@ -300,7 +300,7 @@ void FTopologicalEdge::LinkIfCoincident(FTopologicalEdge& Twin, double EdgeLengt
 		}
 		else
 		{
-			FMessage::Printf(Log, TEXT("Edge %d and Edge %d are to far (%f) to be connected\n"), GetId(), Twin.GetId(), sqrt(SquareDistanceE1V1_E2V1));
+			UE_LOG(LogCADKernelBase, Warning, TEXT("Edge %d and Edge %d are too far (%f) to be connected\n"), GetId(), Twin.GetId(), sqrt(SquareDistanceE1V1_E2V1));
 			return;
 		}
 
@@ -310,7 +310,7 @@ void FTopologicalEdge::LinkIfCoincident(FTopologicalEdge& Twin, double EdgeLengt
 		}
 		else
 		{
-			FMessage::Printf(Log, TEXT("Edge %d and Edge %d are to far (%f) to be connected\n"), GetId(), Twin.GetId(), sqrt(SquareDistanceE1V2_E2V2));
+			UE_LOG(LogCADKernelBase, Warning, TEXT("Edge %d and Edge %d are too far (%f) to be connected\n"), GetId(), Twin.GetId(), sqrt(SquareDistanceE1V2_E2V2));
 			return;
 		}
 	}
@@ -322,7 +322,7 @@ void FTopologicalEdge::LinkIfCoincident(FTopologicalEdge& Twin, double EdgeLengt
 		}
 		else
 		{
-			FMessage::Printf(Log, TEXT("Edge %d and Edge %d are to far (%f) to be connected\n"), GetId(), Twin.GetId(), sqrt(SquareDistanceE1V1_E2V2));
+			UE_LOG(LogCADKernelBase, Warning, TEXT("Edge %d and Edge %d are to far (%f) to be connected\n"), GetId(), Twin.GetId(), sqrt(SquareDistanceE1V1_E2V2));
 			return;
 		}
 
@@ -332,7 +332,7 @@ void FTopologicalEdge::LinkIfCoincident(FTopologicalEdge& Twin, double EdgeLengt
 		}
 		else
 		{
-			FMessage::Printf(Log, TEXT("Edge %d and Edge %d are to far (%f) to be connected\n"), GetId(), Twin.GetId(), sqrt(SquareDistanceE1V2_E2V1));
+			UE_LOG(LogCADKernelBase, Warning, TEXT("Edge %d and Edge %d are to far (%f) to be connected\n"), GetId(), Twin.GetId(), sqrt(SquareDistanceE1V2_E2V1));
 			return;
 		}
 	}
diff --git a/Engine/Source/Runtime/Datasmith/CADKernel/Base/Private/Topo/TopologicalFace.cpp b/Engine/Source/Runtime/Datasmith/CADKernel/Base/Private/Topo/TopologicalFace.cpp
--- a/Engine/Source/Runtime/Datasmith/CADKernel/Base/Private/Topo/TopologicalFace.cpp
+++ b/Engine/Source/Runtime/Datasmith/CADKernel/Base/Private/Topo/TopologicalFace.cpp
@@ -2,7 +2,6 @@
 
 #include "Topo/TopologicalFace.h"
 
-#include "Core/System.h"
 #include "Geo/Curves/RestrictionCurve.h"
 #include "Geo/Curves/SegmentCurve.h"
 #include "Geo/GeoPoint.h"
diff --git a/Engine/Source/Runtime/Datasmith/CADKernel/Base/Private/Topo/TopologicalLoop.cpp b/Engine/Source/Runtime/Datasmith/CADKernel/Base/Private/Topo/TopologicalLoop.cpp
--- a/Engine/Source/Runtime/Datasmith/CADKernel/Base/Private/Topo/TopologicalLoop.cpp
+++ b/Engine/Source/Runtime/Datasmith/CADKernel/Base/Private/Topo/TopologicalLoop.cpp
@@ -9,7 +9,6 @@
 #include "Topo/TopologicalEdge.h"
 #include "Topo/TopologicalLink.h"
 #include "Topo/TopologicalFace.h"
-#include "UI/Display.h"
 #include "UI/Message.h"
 #include "Utils/Util.h"
 
@@ -140,7 +139,7 @@ EOrientation FTopologicalLoop::GetDirection(TSharedPtr<FTopologicalEdge>& InEdge
 		}
 	}
 
-	FMessage::Printf(Debug, TEXT("Edge %d is not in boundary %d Edges\n"), InEdge->GetId(), GetId());
+	UE_LOG(LogCADKernelBase, Verbose, TEXT("Edge %d is not in boundary %d Edges\n"), InEdge->GetId(), GetId());
 	ensureCADKernel(false);
 	return EOrientation::Front;
 }
@@ -506,7 +505,7 @@ bool FTopologicalLoop::Orient()
 #endif
 
 		bSucceed = false;
-		FMessage::Printf(Log, TEXT("WARNING: Loop Orientation of surface %d is doubtful\n"), Face->GetId());
+		UE_LOG(LogCADKernelBase, Warning, TEXT("Loop Orientation of surface %d is doubtful\n"), Face->GetId());
 	}
 
 	if ((WrongOrientationCount > GoodOrientationCount) == bIsExternal)
@@ -733,8 +732,7 @@ void FTopologicalLoop::CheckEdgesOrientation()
 			|| (EIndex == 0 && OrientedEdge.Direction != EOrientation::Front)
 			|| (EIndex == 1 && OrientedEdge.Direction != EOrientation::Back))
 		{
-			FMessage::Printf(EVerboseLevel::Log, TEXT("CheckEdgesOrientation failed for loop %d edge %d\n"), GetId(), OrientedEdge.Entity->GetId());
-			//ensureCADKernel(false);
+			UE_LOG(LogCADKernelBase, Verbose, TEXT("CheckEdgesOrientation failed for loop %d edge %d\n"), GetId(), OrientedEdge.Entity->GetId());
 		}
 	};
 
@@ -792,7 +790,7 @@ void FTopologicalLoop::CheckLoopWithTwoEdgesOrientation()
 	if ((IndexEdge0 == IndexEdge1 && Edge0.Direction == Edge1.Direction)
 		|| (IndexEdge0 != IndexEdge1 && Edge0.Direction != Edge1.Direction))
 	{
-		FMessage::Printf(EVerboseLevel::Log, TEXT("CheckEdgesOrientation failed for loop %d between edge %d and edge %d\n"), GetId(), Edge0.Entity->GetId(), Edge1.Entity->GetId());
+		UE_LOG(LogCADKernelBase, Verbose, TEXT("CheckEdgesOrientation failed for loop %d between edge %d and edge %d\n"), GetId(), Edge0.Entity->GetId(), Edge1.Entity->GetId());
 		//ensureCADKernel(false);
 	}
 }
@@ -912,7 +910,7 @@ void FTopologicalLoop::EnsureLogicalClosing(const double Tolerance3D)
 		const FVector2d Gap2D = EdgeExtremities[ExtremityIndex].Point2D - PreviousExtremity.Point2D;
 		if (SquareGap3D > SquareTolerance3D)
 		{
-			FMessage::Printf(Log, TEXT("Loop %d Gap 3D : %f\n"), Id, sqrt(SquareGap3D));
+			UE_LOG(LogCADKernelBase, Warning, TEXT("Loop %d Gap 3D : %f\n"), Id, sqrt(SquareGap3D));
 			const FVector PreviousTangent = PreviousEdge.Entity->GetTangentAt(*PreviousEdgeEndVertex);
 			const FVector EdgeTangent = OrientedEdge.Entity->GetTangentAt(*EdgeStartVertex);
 
@@ -952,7 +950,7 @@ void FTopologicalLoop::EnsureLogicalClosing(const double Tolerance3D)
 		else if (FMath::Abs(Gap2D.X) > LargeGapTolerance2D.X || FMath::Abs(Gap2D.Y) > LargeGapTolerance2D.Y)
 		{
 			// if the gap is not so large and the extremity of on edge is tangent to the gap segment then we extend the edge to close the gap
-			FMessage::Printf(Debug, TEXT("Loop %d Gap 2D : [%f, %f] vs Tol2D [%f, %f]\n"), Id, FMath::Abs(Gap2D.X), FMath::Abs(Gap2D.Y), Tolerance2D.X, Tolerance2D.Y);
+			UE_LOG(LogCADKernelBase, Verbose, TEXT("Loop %d Gap 2D : [%f, %f] vs Tol2D [%f, %f]\n"), Id, FMath::Abs(Gap2D.X), FMath::Abs(Gap2D.Y), Tolerance2D.X, Tolerance2D.Y);
 
 			// If the gap in 2d is big e.g. side of a degenerated patch => build an edge
 			TSharedPtr<FTopologicalEdge> Edge = FTopologicalEdge::Make(Surface, PreviousExtremity.Point2D, PreviousEdgeEndVertex, EdgeExtremities[ExtremityIndex].Point2D, EdgeStartVertex);
@@ -968,7 +966,7 @@ void FTopologicalLoop::EnsureLogicalClosing(const double Tolerance3D)
 		else if (FMath::Abs(Gap2D.X) > SmallGapTolerance2D.X || FMath::Abs(Gap2D.Y) > SmallGapTolerance2D.Y)
 		{
 			// if the gap is not so large and the extremity of on edge is tangent to the gap segment then we extend the edge to close the gap
-			FMessage::Printf(Log, TEXT("Loop %d Gap 2D : [%f, %f] vs Tol2D [%f, %f]\n"), Id, FMath::Abs(Gap2D.X), FMath::Abs(Gap2D.Y), Tolerance2D.X, Tolerance2D.Y);
+			UE_LOG(LogCADKernelBase, Warning, TEXT("Loop %d Gap 2D : [%f, %f] vs Tol2D [%f, %f]\n"), Id, FMath::Abs(Gap2D.X), FMath::Abs(Gap2D.Y), Tolerance2D.X, Tolerance2D.Y);
 
 			const FVector2d PreviousTangent = PreviousEdge.Entity->GetTangent2DAt(*PreviousEdgeEndVertex);
 			const FVector2d EdgeTangent = OrientedEdge.Entity->GetTangent2DAt(*EdgeStartVertex);
@@ -995,7 +993,7 @@ void FTopologicalLoop::EnsureLogicalClosing(const double Tolerance3D)
 				else
 				{
 					// if the gap is not so large and the extremity of on edge is tangent to the gap segment then we extend the edge to close the gap
-					FMessage::Printf(Log, TEXT("Add degenerated edges for small 2d gap\n"));
+					UE_LOG(LogCADKernelBase, Warning, TEXT("Add degenerated edges for small 2d gap\n"));
 					TSharedPtr<FTopologicalEdge> Edge = FTopologicalEdge::Make(Surface, PreviousExtremity.Point2D, PreviousEdgeEndVertex, EdgeExtremities[ExtremityIndex].Point2D, EdgeStartVertex);
 					if (Edge.IsValid())
 					{
diff --git a/Engine/Source/Runtime/Datasmith/CADKernel/Base/Private/Topo/Topomaker.cpp b/Engine/Source/Runtime/Datasmith/CADKernel/Base/Private/Topo/Topomaker.cpp
--- a/Engine/Source/Runtime/Datasmith/CADKernel/Base/Private/Topo/Topomaker.cpp
+++ b/Engine/Source/Runtime/Datasmith/CADKernel/Base/Private/Topo/Topomaker.cpp
@@ -15,7 +15,6 @@
 #include "Topo/TopologicalLoop.h"
 #include "Topo/TopologicalVertex.h"
 #include "Topo/TopologicalShapeEntity.h"
-#include "UI/Display.h"
 #include "UI/Message.h"
 #include "Utils/Util.h"
 
@@ -400,8 +399,6 @@ void StitchParallelEdges(TArray<FTopologicalVertex*>& VerticesToProcess, double
 	double SquareTolerance = FMath::Square(SewTolerance);
 	double SquareMinEdgeLength = FMath::Square(MinEdgeLength);
 
-	FTimePoint StartTime = FChrono::Now();
-
 	for (int32 VertexI = 0; VertexI < VerticesToProcess.Num(); ++VertexI)
 	{
 		FTopologicalVertex* Vertex = VerticesToProcess[VertexI];
@@ -416,9 +413,6 @@ void StitchParallelEdges(TArray<FTopologicalVertex*>& VerticesToProcess, double
 			Vertex = StitchParallelEdgesFrom(Vertex, SewTolerance, MinEdgeLength, bProhibitSewingEdgesOfSameFace);
 		}
 	}
-
-	FDuration StepDuration = FChrono::Elapse(StartTime);
-	FChrono::PrintClockElapse(EVerboseLevel::Log, TEXT("    "), TEXT("Stitch Parallel Edges"), StepDuration);
 }
 
 /**
@@ -432,13 +426,13 @@ void MergeCoincidentEdges(FTopologicalEdge* Edge, double MinEdgeLength)
 	{
 		if (Edge->GetTwinEntityCount() > 1)
 		{
-			FMessage::Printf(Debug, TEXT("Face %d Edge %d was self connected\n"), Edge->GetFace()->GetId(), Edge->GetId());
+			UE_LOG(LogCADKernelBase, Warning, TEXT("Face %d Edge %d was self connected\n"), Edge->GetFace()->GetId(), Edge->GetId());
 			Edge->GetStartVertex()->UnlinkTo(*Edge->GetEndVertex());
 		}
 		else
 		{
 			Edge->SetAsDegenerated();
-			FMessage::Printf(Debug, TEXT("Face %d Edge %d is set as degenerated\n"), Edge->GetFace()->GetId(), Edge->GetId());
+			UE_LOG(LogCADKernelBase, Warning, TEXT("Face %d Edge %d is set as degenerated\n"), Edge->GetFace()->GetId(), Edge->GetId());
 		}
 		return;
 	}
@@ -489,8 +483,6 @@ void MergeCoincidentEdges(TArray<FTopologicalEdge*>& EdgesToProcess, double MinE
 
 void MergeCoincidentEdges(TArray<FTopologicalVertex*>& VerticesToProcess, double MinEdgeLength)
 {
-	FTimePoint StartTime = FChrono::Now();
-
 	for (FTopologicalVertex* Vertex : VerticesToProcess)
 	{
 		if (Vertex->IsDeleted() || !Vertex->IsActiveEntity() || Vertex->HasMarker1())
@@ -523,9 +515,6 @@ void MergeCoincidentEdges(TArray<FTopologicalVertex*>& VerticesToProcess, double
 	{
 		Vertex->ResetMarker1();
 	}
-
-	FDuration StepDuration = FChrono::Elapse(StartTime);
-	FChrono::PrintClockElapse(EVerboseLevel::Log, TEXT("    "), TEXT("Merge coincident edges"), StepDuration);
 }
 
 void GetVertices(const TArray<FTopologicalEdge*>& InEdges, TArray<FTopologicalVertex*>& OutVertices)
@@ -654,8 +643,6 @@ void FTopomaker::EmptyShells()
 
 void FTopomaker::Sew()
 {
-	FTimePoint StartJoinTime = FChrono::Now();
-
 	TArray<FTopologicalVertex*> BorderVertices;
 	GetBorderVertices(BorderVertices);
 	TopomakerTools::MergeCoincidentVertices(BorderVertices, SewTolerance);
@@ -709,11 +696,6 @@ void FTopomaker::Sew()
 	}
 
 	ResetMarkersOfFaces();
-
-#ifdef CADKERNEL_DEV
-	Report.SewDuration = FChrono::Elapse(StartJoinTime);
-	FChrono::PrintClockElapse(EVerboseLevel::Log, TEXT(""), TEXT("Sew"), Report.SewDuration);
-#endif
 }
 
 void FTopomaker::GetVertices(TArray<FTopologicalVertex*>& Vertices)
@@ -770,7 +752,6 @@ void FTopomaker::GetBorderVertices(TArray<FTopologicalVertex*>& BorderVertices)
 
 void FTopomaker::MergeUnconnectedSuccessiveEdges()
 {
-	FTimePoint StartTime = FChrono::Now();
 	for (TSharedPtr<FTopologicalFace> FacePtr : Faces)
 	{
 		if (!FacePtr.IsValid())
@@ -863,16 +844,10 @@ void FTopomaker::MergeUnconnectedSuccessiveEdges()
 			}
 		}
 	}
-
-	FDuration StepDuration = FChrono::Elapse(StartTime);
-
-	FChrono::PrintClockElapse(EVerboseLevel::Log, TEXT("    "), TEXT("Merge unconnected adjacent edges"), StepDuration);
 }
 
 void FTopomaker::RemoveIsolatedEdges()
 {
-	FTimePoint StartTime = FChrono::Now();
-
 	TArray<FTopologicalEdge*> IsolatedEdges;
 
 	TArray<FTopologicalVertex*> Vertices;
@@ -892,11 +867,6 @@ void FTopomaker::RemoveIsolatedEdges()
 			}
 		}
 	}
-	FDuration StepDuration = FChrono::Elapse(StartTime);
-
-	FChrono::PrintClockElapse(EVerboseLevel::Log, TEXT("    "), TEXT("Remove Isolated Edges"), StepDuration);
-
-	FMessage::Printf(EVerboseLevel::Log, TEXT("\n\nIsolatedEdges count %d\n\n\n"), IsolatedEdges.Num());
 }
 
 void SplitVertexLinkByShell(FTopologicalVertex* InVertex)
@@ -1024,12 +994,6 @@ void FindBorderVertex(const TArray<FTopologicalEdge*>& Edges, TArray<FTopologica
 
 void FTopomaker::RemoveThinFaces()
 {
-	FTimePoint StartTime = FChrono::Now();
-
-#ifdef DEBUG_THIN_FACE
-	F3DDebugSession _(TEXT("RemoveThinFaces"));
-#endif
-
 	TArray<FTopologicalEdge*> BorderEdges;
 	TArray<FTopologicalVertex*> BorderVertices;
 
@@ -1046,13 +1010,6 @@ void FTopomaker::RemoveThinFaces()
 		double GapSize = 0;
 		if (Analyer.IsThinFace(GapSize))
 		{
-#ifdef DEBUG_THIN_FACE
-			{
-				F3DDebugSession _(TEXT("Thin Face"));
-				Display(*Face);
-				Wait();
-			}
-#endif
 			Face->Remove(&BorderEdges);
 			Face->SetDeletedMarker();
 
@@ -1061,17 +1018,8 @@ void FTopomaker::RemoveThinFaces()
 			TopomakerTools::MergeCoincidentEdges(BorderEdges, GapSize * 1.2);
 
 			BorderEdges.Reset();
-
-#ifdef CADKERNEL_DEV
-			Report.AddThinFace();
-#endif
 		}
 	}
-
-#ifdef CADKERNEL_DEV
-	Report.RemoveThinFacesDuration = FChrono::Elapse(StartTime);
-	FChrono::PrintClockElapse(EVerboseLevel::Log, TEXT(""), TEXT(".RemoveThinFaces"), Report.RemoveThinFacesDuration = FChrono::Elapse(StartTime));
-#endif
 }
 
 void FTopomaker::SetSelfConnectedEdgeDegenerated(TArray<FTopologicalVertex*>& Vertices)
@@ -1096,8 +1044,6 @@ void FTopomaker::SetSelfConnectedEdgeDegenerated(TArray<FTopologicalVertex*>& Ve
 
 void FTopomaker::CheckSelfConnectedEdge(double MaxLengthOfDegeneratedEdge, TArray<FTopologicalVertex*>& OutBorderVertices)
 {
-	FTimePoint StartTime = FChrono::Now();
-
 	TFunction<void(FTopologicalVertex&)> AddVertexIfBorder = [&](FTopologicalVertex& Vertex)
 	{
 		TSharedRef<FTopologicalVertex> ActiveVertex = Vertex.GetLinkActiveEntity();
@@ -1107,7 +1053,6 @@ void FTopomaker::CheckSelfConnectedEdge(double MaxLengthOfDegeneratedEdge, TArra
 		}
 	};
 
-	FMessage::Printf(Log, TEXT("    Self connected edges\n"));
 	for (TSharedPtr<FTopologicalFace> Face : Faces)
 	{
 		for (TSharedPtr<FTopologicalLoop> Loop : Face->GetLoops())
@@ -1121,28 +1066,23 @@ void FTopomaker::CheckSelfConnectedEdge(double MaxLengthOfDegeneratedEdge, TArra
 					{
 						if (Edge->GetTwinEntityCount() > 1)
 						{
-							FMessage::Printf(Debug, TEXT("Face %d Edge %d was self connected\n"), Face->GetId(), Edge->GetId());
+							UE_LOG(LogCADKernelBase, Warning, TEXT("Face %d Edge %d was self connected\n"), Face->GetId(), Edge->GetId());
 							Edge->GetStartVertex()->UnlinkTo(*Edge->GetEndVertex());
 						}
 						else
 						{
 							Edge->SetAsDegenerated();
-							FMessage::Printf(Debug, TEXT("Face %d Edge %d is set as degenerated\n"), Face->GetId(), Edge->GetId());
+							UE_LOG(LogCADKernelBase, Warning, TEXT("Face %d Edge %d is set as degenerated\n"), Face->GetId(), Edge->GetId());
 						}
 					}
 				}
 			}
 		}
 	}
-	FDuration StepDuration = FChrono::Elapse(StartTime);
-
-	FChrono::PrintClockElapse(EVerboseLevel::Log, TEXT("    "), TEXT("Unconnect Self connected edges"), StepDuration);
 }
 
 void FTopomaker::SplitIntoConnectedShells()
 {
-	FTimePoint StartTime = FChrono::Now();
-
 	DeleteNonmanifoldLink();
 
 	// Processed1 : Surfaces added in CandidateSurfacesForMesh
@@ -1304,9 +1244,6 @@ void FTopomaker::SplitIntoConnectedShells()
 
 	RemoveEmptyShells();
 	Session.GetModel().RemoveEmptyBodies();
-
-	FDuration StepDuration = FChrono::Elapse(StartTime);
-	FChrono::PrintClockElapse(EVerboseLevel::Log, TEXT(""), TEXT("Split"), StepDuration);
 }
 
 void FTopomaker::RemoveEmptyShells()
@@ -1354,30 +1291,14 @@ void FTopomaker::RemoveEmptyShells()
 
 void FTopomaker::OrientShells()
 {
-	FTimePoint StartTime = FChrono::Now();
-
 	for (FShell* Shell : Shells)
 	{
 		int32 FaceSwapCount = Shell->Orient();
-#ifdef CADKERNEL_DEV
-		Report.AddSwappedFaceCount(FaceSwapCount);
-#endif
 	}
-
-#ifdef CADKERNEL_DEV
-	Report.OrientationDuration = FChrono::Elapse(StartTime);
-	FChrono::PrintClockElapse(EVerboseLevel::Log, TEXT(""), TEXT("Orient"), Report.OrientationDuration);
-#endif
 }
 
 void FTopomaker::RemoveDuplicatedFaces()
 {
-	FTimePoint StartTime = FChrono::Now();
-
-#ifdef DEBUG_REMOVE_DUPLICATED_FACES
-	F3DDebugSession _(TEXT("DOUBLE FACE"));
-#endif
-
 	TArray<FTopologicalFace*> NonManifoldFaces;
 	NonManifoldFaces.Reserve(Faces.Num() / 10);
 	for (TSharedPtr<FTopologicalFace>& FacePtr : Faces)
@@ -1394,14 +1315,7 @@ void FTopomaker::RemoveDuplicatedFaces()
 			{
 				if (Face.IsADuplicatedFace())
 				{
-#ifdef DEBUG_REMOVE_DUPLICATED_FACES
-					F3DDebugSession A(*FString::Printf(TEXT("Duplicated 1 %d"), Face.GetId()));
-					Display(Face);
-#endif
 					Face.Delete();
-#ifdef CADKERNEL_DEV
-					Report.AddDuplicatedFace();
-#endif
 				}
 			}
 			else
@@ -1424,14 +1338,7 @@ void FTopomaker::RemoveDuplicatedFaces()
 		{
 			if (Face->IsADuplicatedFace())
 			{
-#ifdef DEBUG_REMOVE_DUPLICATED_FACES
-				F3DDebugSession A(*FString::Printf(TEXT("Duplicated 2 %d"), Face->GetId()));
-				Display(*Face);
-#endif
 				Face->Delete();
-#ifdef CADKERNEL_DEV
-				Report.AddDuplicatedFace();
-#endif
 			}
 		}
 	}
@@ -1446,36 +1353,9 @@ void FTopomaker::RemoveDuplicatedFaces()
 
 		if (Face->IsADuplicatedFace())
 		{
-#ifdef DEBUG_REMOVE_DUPLICATED_FACES
-			F3DDebugSession A(*FString::Printf(TEXT("Border %d"), Face->GetId()));
-			Display(*Face);
-#endif
 			Face->Delete();
-#ifdef CADKERNEL_DEV
-			Report.AddNearlyDuplicatedFace();
-#endif
 		}
 	}
-
-#ifdef DEBUG_REMOVE_DUPLICATED_FACES
-	// Step 4: display the remaining Face
-	for (FTopologicalFace* Face : NonManifoldFaces)
-	{
-		if (Face->IsDeleted() || !Face->IsANonManifoldFace())
-		{
-			continue;
-		}
-
-		F3DDebugSession A(*FString::Printf(TEXT("Remaining %d"), Face->GetId()));
-		Display(*Face);
-	}
-#endif
-
-#ifdef CADKERNEL_DEV
-	Report.OrientationDuration = FChrono::Elapse(StartTime);
-	FChrono::PrintClockElapse(EVerboseLevel::Log, TEXT(""), TEXT("RemoveDuplicatedFaces"), Report.RemoveDuplicatedFacesDuration);
-#endif
-
 }
 
 void FTopomaker::SetToProcessMarkerOfFaces()
diff --git a/Engine/Source/Runtime/Datasmith/CADKernel/Base/Private/UI/Display.cpp b/Engine/Source/Runtime/Datasmith/CADKernel/Base/Private/UI/Display.cpp
deleted file mode 100644
diff --git a/Engine/Source/Runtime/Datasmith/CADKernel/Base/Private/UI/Message.cpp b/Engine/Source/Runtime/Datasmith/CADKernel/Base/Private/UI/Message.cpp
--- a/Engine/Source/Runtime/Datasmith/CADKernel/Base/Private/UI/Message.cpp
+++ b/Engine/Source/Runtime/Datasmith/CADKernel/Base/Private/UI/Message.cpp
@@ -1,176 +1,5 @@
 // Copyright Epic Games, Inc. All Rights Reserved.
 
-#include "UI/Message.h"
+#include "Message.h"
 
-#include "Core/Chrono.h"
-#include "Core/System.h"
-#include "UI/Display.h"
-#include "Utils/Util.h"
-
-#include "Misc/VarArgs.h"
-
-namespace UE::CADKernel
-{
-const char* VerboseLevelConstNames[] = {
-	("NO_VERBOSE"),
-	("SPY"),
-	("LOG"),
-	("DBG"),
-	nullptr
-};
-
-const char* VerboseConstDescHelp[] = {
-	": nothing is printed",
-	": only commands and arguments will be showed",
-	": SPY + messages",
-	": debug mode, for developers only",
-	nullptr
-};
-
-int32 FMessage::NumberOfIndentation = 0;
-int32 FMessage::OldPercent = -1;
-
-void FMessage::VPrintf(EVerboseLevel Level, const TCHAR* Text, ...)
-{
-	EVerboseLevel VerboseLevel = (EVerboseLevel)(int32)FSystem::Get().GetVerboseLevel();
-	TSharedPtr<FArchive> LogFile = FSystem::Get().GetLogFile();
-	EVerboseLevel SystemLogLevel = FSystem::Get().GetLogLevel();
-	TSharedPtr<FArchive> SpyFile = FSystem::Get().GetSpyFile();
-
-	FString Indentation;
-	TCHAR* Buffer = NULL;
-	int32 Result = -1;
-	if (VerboseLevel >= Level || (LogFile.IsValid() && Level <= EVerboseLevel::Debug) || (SpyFile.IsValid() && Level == EVerboseLevel::Spy))
-	{
-		int32 BufferSize = 1024;
-
-		// do the usual VARARGS shenanigans
-		while (Result == -1)
-		{
-			FMemory::Free(Buffer);
-			Buffer = (TCHAR*)FMemory::Malloc(BufferSize * sizeof(TCHAR));
-			GET_TYPED_VARARGS_RESULT(TCHAR, Buffer, BufferSize, BufferSize - 1, Text, Text, Result);
-			BufferSize *= 2;
-		};
-		Buffer[Result] = 0;
-
-		Indentation.Reserve(NumberOfIndentation * 3);
-		for (int32 iIndent = 0; iIndent < NumberOfIndentation; iIndent++)
-		{
-			Indentation += TEXT(" - ");
-		}
-	}
-
-	if (VerboseLevel >= Level)
-	{
-		FSystem::Get().GetConsole().Print(*Indentation, Level);
-		FSystem::Get().GetConsole().Print(Buffer, Level);
-	}
-
-	if ((LogFile.IsValid() && Level <= SystemLogLevel) || (SpyFile.IsValid() && Level <= Spy))
-	{
-		FTCHARToUTF8 UTF8Indentation(*Indentation, Indentation.Len());
-		FTCHARToUTF8 UTF8String(Buffer, Result);
-		if (LogFile.IsValid() && Level <= SystemLogLevel)
-		{
-			LogFile->Serialize((UTF8CHAR*)UTF8Indentation.Get(), UTF8Indentation.Length() * sizeof(UTF8CHAR));
-			LogFile->Serialize((UTF8CHAR*)UTF8String.Get(), UTF8String.Length() * sizeof(UTF8CHAR));
-			LogFile->Flush();
-		}
-
-		if (SpyFile.IsValid() && Level <= Spy)
-		{
-			SpyFile->Serialize((UTF8CHAR*)UTF8Indentation.Get(), UTF8Indentation.Length() * sizeof(UTF8CHAR));
-			SpyFile->Serialize((UTF8CHAR*)UTF8String.Get(), UTF8String.Length() * sizeof(UTF8CHAR));
-			SpyFile->Flush();
-		}
-	}
-
-	if (Buffer)
-	{
-		FMemory::Free(Buffer);
-	}
-}
-
-void FMessage::VReportPrintF(FString Header, const TCHAR* Text, ...)
-{
-#ifdef CADKERNEL_DEV
-	TSharedPtr<FArchive> ReportFile = FSystem::Get().GetReportFile();
-	TSharedPtr<FArchive> ReportHeaderFile = FSystem::Get().GetReportHeaderFile();
-	if (!ReportFile.IsValid())
-	{
-		return;
-	}
-
-	TCHAR* Buffer = NULL;
-	int32 Result = -1;
-	{
-		int32 BufferSize = 1024;
-
-		// do the usual VARARGS shenanigans
-		while (Result == -1)
-		{
-			FMemory::Free(Buffer);
-			Buffer = (TCHAR*)FMemory::Malloc(BufferSize * sizeof(TCHAR));
-			GET_TYPED_VARARGS_RESULT(TCHAR, Buffer, BufferSize, BufferSize - 1, Text, Text, Result);
-			BufferSize *= 2;
-		};
-		Buffer[Result] = 0;
-	}
-
-	//FSystem::Get().GetConsole().Print(*Header, Log);
-	//FSystem::Get().GetConsole().Print(TEXT(": "), Log);
-	//FSystem::Get().GetConsole().Print(Buffer, Log);
-	//FSystem::Get().GetConsole().Print(TEXT("\n"), Log);
-
-	FTCHARToUTF8 BufferUtf8(Buffer, Result);
-	ReportFile->Serialize((UTF8CHAR*)BufferUtf8.Get(), BufferUtf8.Length() * sizeof(UTF8CHAR));
-	ReportFile->Serialize(",", sizeof(UTF8CHAR));
-	ReportFile->Flush();
-
-	if (ReportHeaderFile.IsValid())
-	{
-		FTCHARToUTF8 HeaderUtf8(*Header, Header.Len());
-		ReportHeaderFile->Serialize((UTF8CHAR*)HeaderUtf8.Get(), HeaderUtf8.Length() * sizeof(UTF8CHAR));
-		ReportHeaderFile->Serialize(",", sizeof(UTF8CHAR));
-		ReportHeaderFile->Flush();
-	}
-
-	if (Buffer)
-	{
-		FMemory::Free(Buffer);
-	}
-#endif
-}
-
-void FChrono::PrintClockElapse(EVerboseLevel Level, const TCHAR* Indent, const TCHAR* Process, FDuration Duration, ETimeUnit Unit)
-{
-#ifdef CADKERNEL_DEV
-	if (FSystem::Get().GetVerboseLevel() < Level)
-	{
-		return;
-	}
-
-	TCHAR* SUnit = TEXT("seconds");
-	long long Time = 0;
-	switch (Unit)
-	{
-	case ETimeUnit::NanoSeconds:
-		SUnit = TEXT("nanoseconds");
-		Time = FChrono::ConvertInto<std::chrono::nanoseconds>(Duration);
-		break;
-	case ETimeUnit::MicroSeconds:
-		SUnit = TEXT("microseconds");
-		Time = FChrono::ConvertInto<std::chrono::microseconds>(Duration);
-		break;
-	case ETimeUnit::MilliSeconds:
-		SUnit = TEXT("milliseconds");
-		Time = FChrono::ConvertInto<std::chrono::milliseconds>(Duration);
-		break;
-	default:
-		Time = FChrono::ConvertInto<std::chrono::seconds>(Duration);
-	}
-	FMessage::Printf(Level, TEXT("%sSpeed test (%s) -----> %lld %s\n"), Indent, Process, Time, SUnit);
-#endif
-}
-}
\ No newline at end of file
+DEFINE_LOG_CATEGORY(LogCADKernelBase);
diff --git a/Engine/Source/Runtime/Datasmith/CADKernel/Base/Private/UI/Message.h b/Engine/Source/Runtime/Datasmith/CADKernel/Base/Private/UI/Message.h
new file mode 100644
--- /dev/null
+++ b/Engine/Source/Runtime/Datasmith/CADKernel/Base/Private/UI/Message.h
@@ -0,0 +1,6 @@
+// Copyright Epic Games, Inc. All Rights Reserved.
+#pragma once
+
+#include "CoreMinimal.h"
+
+DECLARE_LOG_CATEGORY_EXTERN(LogCADKernelBase, Error, All);
diff --git a/Engine/Source/Runtime/Datasmith/CADKernel/Base/Private/UI/Progress.cpp b/Engine/Source/Runtime/Datasmith/CADKernel/Base/Private/UI/Progress.cpp
deleted file mode 100644
diff --git a/Engine/Source/Runtime/Datasmith/CADKernel/Base/Private/UI/Visu.cpp b/Engine/Source/Runtime/Datasmith/CADKernel/Base/Private/UI/Visu.cpp
deleted file mode 100644
diff --git a/Engine/Source/Runtime/Datasmith/CADKernel/Base/Public/Core/CADKernelArchive.h b/Engine/Source/Runtime/Datasmith/CADKernel/Base/Public/Core/CADKernelArchive.h
--- a/Engine/Source/Runtime/Datasmith/CADKernel/Base/Public/Core/CADKernelArchive.h
+++ b/Engine/Source/Runtime/Datasmith/CADKernel/Base/Public/Core/CADKernelArchive.h
@@ -2,7 +2,6 @@
 #pragma once
 
 #include "Core/Types.h"
-#include "UI/Message.h"
 
 #include "HAL/FileManager.h"
 #include "Serialization/Archive.h"
@@ -60,7 +59,6 @@ public:
 		FArchive* Ar = IFileManager::Get().CreateFileWriter(FileName);
 		if (Ar == nullptr)
 		{
-			FMessage::Printf(Log, TEXT("The archive file %s is corrupted\n"), FileName);
 			return TSharedPtr<FCADKernelArchive>();
 		}
 		return MakeShared<FCADKernelArchive>(InSession, Ar);
@@ -71,7 +69,6 @@ public:
 		FArchive* Ar = IFileManager::Get().CreateFileReader(FileName);
 		if (!Ar)
 		{
-			FMessage::Printf(Log, TEXT("The archive file %s is corrupted\n"), FileName);
 			return TSharedPtr<FCADKernelArchive>();
 		}
 		return MakeShared<FCADKernelArchive>(InSession, Ar);
@@ -82,7 +79,6 @@ public:
 		FArchive* Ar = new FMemoryWriter(InRawData, true);
 		if (!Ar)
 		{
-			FMessage::Printf(Log, TEXT("The archive is corrupted\n"));
 			return TSharedPtr<FCADKernelArchive>();
 		}
 		return MakeShared<FCADKernelArchive>(InSession, Ar);
@@ -93,7 +89,6 @@ public:
 		FArchive* Ar = new FMemoryReader(InRawData, true);
 		if (!Ar)
 		{
-			FMessage::Printf(Log, TEXT("The archive buffer is corrupted\n"));
 			return TSharedPtr<FCADKernelArchive>();
 		}
 		return MakeShared<FCADKernelArchive>(InSession, Ar);
diff --git a/Engine/Source/Runtime/Datasmith/CADKernel/Base/Public/Core/Chrono.h b/Engine/Source/Runtime/Datasmith/CADKernel/Base/Public/Core/Chrono.h
--- a/Engine/Source/Runtime/Datasmith/CADKernel/Base/Public/Core/Chrono.h
+++ b/Engine/Source/Runtime/Datasmith/CADKernel/Base/Public/Core/Chrono.h
@@ -1,6 +1,7 @@
 // Copyright Epic Games, Inc. All Rights Reserved.
 #pragma once
 
+#if 0
 #include "Core/Types.h"
 
 #ifdef CADKERNEL_DEV
@@ -70,4 +71,5 @@ public:
 
 	static void PrintClockElapse(EVerboseLevel Level, const TCHAR* Indent, const TCHAR* Process, FDuration Duration, ETimeUnit Unit = ETimeUnit::MicroSeconds);
 };
-}
\ No newline at end of file
+}
+#endif
\ No newline at end of file
diff --git a/Engine/Source/Runtime/Datasmith/CADKernel/Base/Public/Core/System.h b/Engine/Source/Runtime/Datasmith/CADKernel/Base/Public/Core/System.h
deleted file mode 100644
diff --git a/Engine/Source/Runtime/Datasmith/CADKernel/Base/Public/Geo/Curves/Curve.h b/Engine/Source/Runtime/Datasmith/CADKernel/Base/Public/Geo/Curves/Curve.h
--- a/Engine/Source/Runtime/Datasmith/CADKernel/Base/Public/Geo/Curves/Curve.h
+++ b/Engine/Source/Runtime/Datasmith/CADKernel/Base/Public/Geo/Curves/Curve.h
@@ -223,7 +223,7 @@ public:
 	{
 		ensureCADKernel(Dimension == 3);
 
-		NOT_IMPLEMENTED;
+		UE_LOG(LogCADKernelBase, Warning, TEXT("The function \"FCurve::ExtendTo\" at line %d of the file Curve.h is not implemented"), __LINE__);
 		ensureCADKernel(false);
 	}
 
diff --git a/Engine/Source/Runtime/Datasmith/CADKernel/Base/Public/Geo/Sampling/Polyline.h b/Engine/Source/Runtime/Datasmith/CADKernel/Base/Public/Geo/Sampling/Polyline.h
--- a/Engine/Source/Runtime/Datasmith/CADKernel/Base/Public/Geo/Sampling/Polyline.h
+++ b/Engine/Source/Runtime/Datasmith/CADKernel/Base/Public/Geo/Sampling/Polyline.h
@@ -5,7 +5,6 @@
 
 #include "Geo/Sampling/PolylineTools.h"
 #include "Geo/Surfaces/Surface.h"
-#include "UI/Display.h"
 
 #include "Serialization/Archive.h"
 
diff --git a/Engine/Source/Runtime/Datasmith/CADKernel/Base/Public/Geo/Sampling/SurfacicPolyline.h b/Engine/Source/Runtime/Datasmith/CADKernel/Base/Public/Geo/Sampling/SurfacicPolyline.h
--- a/Engine/Source/Runtime/Datasmith/CADKernel/Base/Public/Geo/Sampling/SurfacicPolyline.h
+++ b/Engine/Source/Runtime/Datasmith/CADKernel/Base/Public/Geo/Sampling/SurfacicPolyline.h
@@ -12,7 +12,6 @@
 #include "Math/Boundary.h"
 #include "Math/MatrixH.h"
 #include "Math/Point.h"
-#include "UI/Display.h"
 #include "Utils/IndexOfCoordinateFinder.h"
 
 #include "Serialization/Archive.h"
diff --git a/Engine/Source/Runtime/Datasmith/CADKernel/Base/Public/Math/MatrixH.h b/Engine/Source/Runtime/Datasmith/CADKernel/Base/Public/Math/MatrixH.h
--- a/Engine/Source/Runtime/Datasmith/CADKernel/Base/Public/Math/MatrixH.h
+++ b/Engine/Source/Runtime/Datasmith/CADKernel/Base/Public/Math/MatrixH.h
@@ -261,8 +261,6 @@ public:
 		return FVector(Get(Index, 0), Get(Index, 1), Get(Index, 2));
 	}
 
-	void Print(EVerboseLevel level) const;
-
 	bool IsId() const
 	{
 		for (int32 Row = 0; Row < 4; Row++)
diff --git a/Engine/Source/Runtime/Datasmith/CADKernel/Base/Public/Topo/Model.h b/Engine/Source/Runtime/Datasmith/CADKernel/Base/Public/Topo/Model.h
--- a/Engine/Source/Runtime/Datasmith/CADKernel/Base/Public/Topo/Model.h
+++ b/Engine/Source/Runtime/Datasmith/CADKernel/Base/Public/Topo/Model.h
@@ -112,8 +112,6 @@ public:
 
 	void RemoveEmptyBodies();
 
-	void PrintBodyAndShellCount();
-
 	bool Contains(TSharedPtr<FTopologicalEntity> InEntity);
 
 	/**
diff --git a/Engine/Source/Runtime/Datasmith/CADKernel/Base/Public/Topo/TopologicalFace.h b/Engine/Source/Runtime/Datasmith/CADKernel/Base/Public/Topo/TopologicalFace.h
--- a/Engine/Source/Runtime/Datasmith/CADKernel/Base/Public/Topo/TopologicalFace.h
+++ b/Engine/Source/Runtime/Datasmith/CADKernel/Base/Public/Topo/TopologicalFace.h
@@ -19,10 +19,6 @@
 #include "Mesh/Structure/EdgeSegment.h"
 #include "Mesh/Structure/ThinZone2D.h"
 
-#ifdef CADKERNEL_DEV
-#include "UI/DefineForDebug.h"
-#endif
-
 namespace UE::CADKernel
 {
 
@@ -31,6 +27,7 @@ class FCriteriaGrid;
 class FCriterion;
 class FCurve;
 class FDatabase;
+class FFaceMesh;
 class FModelMesh;
 class FTopologicalVertex;
 struct FSurfacicSampling;
diff --git a/Engine/Source/Runtime/Datasmith/CADKernel/Base/Public/Topo/TopologicalLink.h b/Engine/Source/Runtime/Datasmith/CADKernel/Base/Public/Topo/TopologicalLink.h
--- a/Engine/Source/Runtime/Datasmith/CADKernel/Base/Public/Topo/TopologicalLink.h
+++ b/Engine/Source/Runtime/Datasmith/CADKernel/Base/Public/Topo/TopologicalLink.h
@@ -93,7 +93,7 @@ public:
 					return true;
 				}
 			}
-			FMessage::Error(TEXT("FTopologicalLink::ActivateEntity, the topological entity is not found in the twins entities"));
+			UE_LOG(LogCADKernelBase, Warning, TEXT("FTopologicalLink::ActivateEntity, the topological entity is not found in the twins entities"));
 			return false;
 		};
 
diff --git a/Engine/Source/Runtime/Datasmith/CADKernel/Base/Public/UI/Console.h b/Engine/Source/Runtime/Datasmith/CADKernel/Base/Public/UI/Console.h
deleted file mode 100644
diff --git a/Engine/Source/Runtime/Datasmith/CADKernel/Base/Public/UI/Display.h b/Engine/Source/Runtime/Datasmith/CADKernel/Base/Public/UI/Display.h
deleted file mode 100644
diff --git a/Engine/Source/Runtime/Datasmith/CADKernel/Base/Public/UI/Message.h b/Engine/Source/Runtime/Datasmith/CADKernel/Base/Public/UI/Message.h
deleted file mode 100644
diff --git a/Engine/Source/Runtime/Datasmith/CADKernel/Base/Public/UI/Progress.h b/Engine/Source/Runtime/Datasmith/CADKernel/Base/Public/UI/Progress.h
deleted file mode 100644
diff --git a/Engine/Source/Runtime/Datasmith/CADKernel/Base/Public/UI/Visu.h b/Engine/Source/Runtime/Datasmith/CADKernel/Base/Public/UI/Visu.h
deleted file mode 100644
