Fixed compilation issue in ClangEditor

#issue 1066947
#rnx

diff --git a/Engine/Source/Runtime/Datasmith/CADKernel/Base/Private/Mesh/Criteria/Criterion.cpp b/Engine/Source/Runtime/Datasmith/CADKernel/Base/Private/Mesh/Criteria/Criterion.cpp
--- a/Engine/Source/Runtime/Datasmith/CADKernel/Base/Private/Mesh/Criteria/Criterion.cpp
+++ b/Engine/Source/Runtime/Datasmith/CADKernel/Base/Private/Mesh/Criteria/Criterion.cpp
@@ -89,7 +89,7 @@ TSharedPtr<FCriterion> FCriterion::CreateCriterion(ECriterion Type, double Value
 	case ECriterion::CADCurvature:
 		return FEntity::MakeShared<FCurvatureCriterion>();
 	default:
-		UE_LOG(LogCADKernelBase, Warning, TEXT("CALL of \"FCriterion::CreateCriterion\" NOT EXPECTED at line %d of the file Criterion.cpp\n"), __LINE__);
+		UE_LOG(LogCADKernelBase, Warning, TEXT("CALL of \"FCriterion::CreateCriterion\" NOT EXPECTED"));
 		return TSharedPtr<FCriterion>();
 	}
 }
@@ -107,7 +107,7 @@ double FCriterion::DefaultValue(ECriterion Type)
 	case ECriterion::Angle:
 		return FAngleCriterion::DefaultValue();
 	default:
-		UE_LOG(LogCADKernelBase, Warning, TEXT("CALL of \"FCriterion::DefaultValue\" NOT EXPECTED at line %d of the file Criterion.cpp\n"), __LINE__);
+		UE_LOG(LogCADKernelBase, Warning, TEXT("CALL of \"FCriterion::DefaultValue\" NOT EXPECTED"));
 		return 0.0;
 	}
 }
diff --git a/Engine/Source/Runtime/Datasmith/CADKernel/Base/Public/Geo/Curves/Curve.h b/Engine/Source/Runtime/Datasmith/CADKernel/Base/Public/Geo/Curves/Curve.h
--- a/Engine/Source/Runtime/Datasmith/CADKernel/Base/Public/Geo/Curves/Curve.h
+++ b/Engine/Source/Runtime/Datasmith/CADKernel/Base/Public/Geo/Curves/Curve.h
@@ -222,8 +222,6 @@ public:
 	virtual void ExtendTo(const FVector& DesiredPosition)
 	{
 		ensureCADKernel(Dimension == 3);
-
-		UE_LOG(LogCADKernelBase, Warning, TEXT("The function \"FCurve::ExtendTo\" at line %d of the file Curve.h is not implemented"), __LINE__);
 		ensureCADKernel(false);
 	}
 
