[In Progress] Interchange CAD Importer plugin
- Added support for Interchange translated nodes
- Updated CAD parsers to use newly introduced APIs
- Updated CAD mesh pipeline
- Updated CAD main pipeline
- Reformatted FCADKernelTessellationSettings
Issues:
- Material instances do not work with Wire
- UI glitches in the details panels of the CAD pipelines

#rnx

#virtualized

diff --git a/Engine/Source/Runtime/Datasmith/CADKernel/Editor/Private/LegacyParametricSurfaceData.cpp b/Engine/Source/Runtime/Datasmith/CADKernel/Editor/Private/LegacyParametricSurfaceData.cpp
--- a/Engine/Source/Runtime/Datasmith/CADKernel/Editor/Private/LegacyParametricSurfaceData.cpp
+++ b/Engine/Source/Runtime/Datasmith/CADKernel/Editor/Private/LegacyParametricSurfaceData.cpp
@@ -10,7 +10,7 @@ UParametricSurfaceData* FLegacyParametricSurfaceData::ToCADKernel() const
 	ParametricSurfaceData->SetMeshParameters(MeshParameters.ToCADKernel());
 	ParametricSurfaceData->SetLastTessellationSettings(LastTessellationOptions.ToCADKernel());
 
-	ParametricSurfaceData->SetRawData(RawData, !LastTessellationOptions.bUseCADKernel);
+	ParametricSurfaceData->SetRawData(RawData, LastTessellationOptions.Tessellator == ECADKernelTessellator::TechSoft);
 
 	return ParametricSurfaceData;
 }
@@ -32,7 +32,7 @@ void FLegacyParametricSurfaceData::Serialize(FArchive& Ar)
 	Ar << LastTessellationOptions.MaxEdgeLength;
 	Ar << LastTessellationOptions.NormalTolerance;
 	Ar << LastTessellationOptions.StitchingTechnique;
-	Ar << LastTessellationOptions.bUseCADKernel;
+	Ar << LastTessellationOptions.Tessellator;
 	Ar << LastTessellationOptions.GeometricTolerance;
 	Ar << LastTessellationOptions.StitchingTolerance;
 
diff --git a/Engine/Source/Runtime/Datasmith/CADKernel/Editor/Private/LegacyParametricSurfaceData.h b/Engine/Source/Runtime/Datasmith/CADKernel/Editor/Private/LegacyParametricSurfaceData.h
--- a/Engine/Source/Runtime/Datasmith/CADKernel/Editor/Private/LegacyParametricSurfaceData.h
+++ b/Engine/Source/Runtime/Datasmith/CADKernel/Editor/Private/LegacyParametricSurfaceData.h
@@ -50,8 +50,8 @@ struct FLegacyTessellationOptions
 	float ChordTolerance;
 	float MaxEdgeLength;
 	float NormalTolerance;
-	uint8 StitchingTechnique;
-	bool bUseCADKernel = true;
+	uint8 StitchingTechnique = (uint8)ECADKernelStitchingTechnique::StitchingNone;
+	uint8 Tessellator = (uint8)ECADKernelTessellator::CADKernel;
 	double GeometricTolerance = 0.001;
 	double StitchingTolerance = 0.001;
 
@@ -64,7 +64,7 @@ struct FLegacyTessellationOptions
 			static_cast<ECADKernelStitchingTechnique>(StitchingTechnique),
 			GeometricTolerance,
 			StitchingTolerance,
-			bUseCADKernel
+			static_cast<ECADKernelTessellator>(Tessellator)
 		);
 	}
 };
diff --git a/Engine/Source/Runtime/Datasmith/CADKernel/Editor/Private/ParametricSurfaceData.cpp b/Engine/Source/Runtime/Datasmith/CADKernel/Editor/Private/ParametricSurfaceData.cpp
--- a/Engine/Source/Runtime/Datasmith/CADKernel/Editor/Private/ParametricSurfaceData.cpp
+++ b/Engine/Source/Runtime/Datasmith/CADKernel/Editor/Private/ParametricSurfaceData.cpp
@@ -27,7 +27,7 @@ namespace UE::CADKernel::Editor
 
 			FTessellationContext Context(Data.GetModelParameters(), Data.GetMeshParameters(), Data.GetLastTessellationSettings());
 
-			if (Context.TessellationSettings.bUseCADKernel)
+			if (Context.TessellationSettings.UseCADKernel())
 			{
 				TSharedPtr<FModel> Model = Data.GetModel();
 				if (Model.IsValid())
@@ -52,7 +52,7 @@ namespace UE::CADKernel::Editor
 
 			FTessellationContext Context(Data.GetModelParameters(), Data.GetMeshParameters(), Settings);
 
-			if (Context.TessellationSettings.bUseCADKernel)
+			if (Context.TessellationSettings.UseCADKernel())
 			{
 				TSharedPtr<FModel> Model = Data.GetModel();
 				if (Model.IsValid())
diff --git a/Engine/Source/Runtime/Datasmith/CADKernel/Engine/Private/MeshDescriptionWrapper.cpp b/Engine/Source/Runtime/Datasmith/CADKernel/Engine/Private/MeshDescriptionWrapper.cpp
--- a/Engine/Source/Runtime/Datasmith/CADKernel/Engine/Private/MeshDescriptionWrapper.cpp
+++ b/Engine/Source/Runtime/Datasmith/CADKernel/Engine/Private/MeshDescriptionWrapper.cpp
@@ -154,8 +154,7 @@ namespace UE::CADKernel::MeshUtilities
 		// At least one edge is non-manifold
 		if (IsBoundary[0] == false || IsBoundary[1] == false || IsBoundary[2] == false)
 		{
-			// #cad_debug
-			ensure(false);
+			UE_LOG(LogCADKernelEngine, Verbose, TEXT("FMeshDescriptionWrapper::GetVertexInstances: Detected a non-manifold edge"));
 
 			bool bDuplicate[3]{ false, false, false };
 			if (IsBoundary[0] == false)
diff --git a/Engine/Source/Runtime/Datasmith/CADKernel/Engine/Private/TechSoft/TechSoftUtilities.cpp b/Engine/Source/Runtime/Datasmith/CADKernel/Engine/Private/TechSoft/TechSoftUtilities.cpp
--- a/Engine/Source/Runtime/Datasmith/CADKernel/Engine/Private/TechSoft/TechSoftUtilities.cpp
+++ b/Engine/Source/Runtime/Datasmith/CADKernel/Engine/Private/TechSoft/TechSoftUtilities.cpp
@@ -13,6 +13,7 @@
 
 #include "Misc/FileHelper.h"
 #include "Misc/Paths.h"
+#include "Misc/ScopeLock.h"
 #include "Serialization/JsonSerializer.h"
 #include "Serialization/JsonWriter.h"
 
@@ -193,6 +194,11 @@ namespace UE::CADKernel
 
 		using namespace UE::CADKernel;
 
+		static FCriticalSection TechSoftCriticalSection;
+
+		// #cad_import: Track when TechSoft tessellation is thread safe
+		FScopeLock Lock(&TechSoftCriticalSection);
+
 		// #cad_import: Create set of helpers IsA[Whatever] in FTechSoftLibrary
 		// Here it would FTechSoftLibrary::IsATriangulatedMesh
 		A3DEEntityType Type;
@@ -281,7 +287,14 @@ namespace UE::CADKernel
 		TessellationParameters->m_bKeepUVPoints = A3D_TRUE; // Keep parametric points as texture points.
 
 		// Get the tessellation
+		static FCriticalSection TechSoftCriticalSection;
+
 		A3DStatus Status = A3DRiRepresentationItemComputeTessellation(Representation, TessellationParameters.GetPtr());
+		if (!ensure(Status == A3D_SUCCESS))
+		{
+			return false;
+		}
+
 		TechSoft::TUniqueObject<A3DRiRepresentationItemData> RepresentationItemData(Representation);
 		if (!RepresentationItemData.IsValid())
 		{
@@ -340,9 +353,13 @@ namespace UE::CADKernel
 			}
 		}
 
-		const int32 IsoCurveCount = 7;
+		constexpr int32 IsoCurveCount = 7;
 		const double DeltaU = (Domain->m_sMax.m_dX - Domain->m_sMin.m_dX) / (IsoCurveCount - 1.);
 		const double DeltaV = (Domain->m_sMax.m_dY - Domain->m_sMin.m_dY) / (IsoCurveCount - 1.);
+		if (DeltaU < UE_DOUBLE_SMALL_NUMBER || DeltaV < UE_DOUBLE_SMALL_NUMBER)
+		{
+			return FVector2d::UnitVector;
+		}
 
 		const A3DSurfBase* A3DSurface = TopoFaceData->m_pSurface;
 
@@ -406,13 +423,13 @@ namespace UE::CADKernel
 		LengthVMed /= IsoCurveCount;
 		LengthVMed = LengthVMed * 2 / 3 + LengthVMax / 3;
 
-		// Texture unit is meter, Coord unit from TechSoft is mm, so TextureScale = 0.001 to convert mm into m
+		// Texture unit is meter, Coord unit from TechSoft is mm, so TextureScale = 0.01 to convert cm into m
 		// #cad_import: Verify assumption above and where TextureUnit can come from
 		constexpr double TextureScale = 0.01;
 
 		return {
-			TextureUnit* TextureScale* LengthUMed / (Domain->m_sMax.m_dX - Domain->m_sMin.m_dX),
-			TextureUnit* TextureScale* LengthVMed / (Domain->m_sMax.m_dY - Domain->m_sMin.m_dY)
+			TextureUnit * TextureScale * LengthUMed / (Domain->m_sMax.m_dX - Domain->m_sMin.m_dX),
+			TextureUnit * TextureScale * LengthVMed / (Domain->m_sMax.m_dY - Domain->m_sMin.m_dY)
 		};
 	}
 
diff --git a/Engine/Source/Runtime/Datasmith/CADKernel/Engine/Public/CADKernelEngineDefinitions.h b/Engine/Source/Runtime/Datasmith/CADKernel/Engine/Public/CADKernelEngineDefinitions.h
--- a/Engine/Source/Runtime/Datasmith/CADKernel/Engine/Public/CADKernelEngineDefinitions.h
+++ b/Engine/Source/Runtime/Datasmith/CADKernel/Engine/Public/CADKernelEngineDefinitions.h
@@ -6,7 +6,7 @@
 
 #include "CADKernelEngineDefinitions.generated.h"
 
-UENUM()
+UENUM(BlueprintType)
 enum class ECADKernelModelCoordSystem : uint8
 {
 	ZUp_LeftHanded,
@@ -16,14 +16,21 @@ enum class ECADKernelModelCoordSystem : uint8
 	ZUp_RightHanded_FBXLegacy,
 };
 
-UENUM()
-enum class ECADKernelStitchingTechnique : uint8
+UENUM(BlueprintType)
+enum ECADKernelStitchingTechnique : uint8
 {
 	StitchingNone = 0,
 	StitchingHeal,
 	StitchingSew,
 };
 
+UENUM(BlueprintType)
+enum ECADKernelTessellator : uint8
+{
+	CADKernel = 0,
+	TechSoft,
+};
+
 USTRUCT(BlueprintType)
 struct CADKERNELENGINE_API FCADKernelModelParameters
 {
@@ -79,16 +86,17 @@ struct CADKERNELENGINE_API FCADKernelTessellationSettings
 			ECADKernelStitchingTechnique InStitchingTechnique = ECADKernelStitchingTechnique::StitchingSew,
 			double InGeometricTolerance = 0.001,
 			double InStichingTolerance = 0.001,
-			bool bInUseCADKernel = true
+			ECADKernelTessellator InTessellator = ECADKernelTessellator::CADKernel
 		)
 		: ChordTolerance(InChordTolerance)
 		, MaxEdgeLength(InMaxEdgeLength)
 		, NormalTolerance(InNormalTolerance)
 		, StitchingTechnique(InStitchingTechnique)
-		, bUseCADKernel(bInUseCADKernel)
+		, Tessellator(InTessellator)
 		, GeometricTolerance(InGeometricTolerance)
 		, StitchingTolerance(InStichingTolerance)
 	{
+		Tessellator = bWithTechSoft ? InTessellator : ECADKernelTessellator::CADKernel;
 	}
 
 protected:
@@ -127,18 +135,20 @@ public:
 	 * Healing : Connects surfaces which physically share a boundary but not topologically within an object.
 	 */
 	UPROPERTY(config, EditAnywhere, BlueprintReadWrite, AdvancedDisplay, Category = "Geometry & Tessellation Options", meta = (ToolTip = "Stitching technique applied on model before tessellation. Sewing could impact number of objects."))
-	ECADKernelStitchingTechnique StitchingTechnique;
+	TEnumAsByte<ECADKernelStitchingTechnique> StitchingTechnique;
 
 	/**
-	 * Indicates which tessellator to use CADKernel or TechSoft. By default it is CADKernel.
+	 * Indicates which tessellator to use: CADKernel or TechSoft. By default it is CADKernel.
 	 */
-	UPROPERTY(config, EditAnywhere, BlueprintReadWrite, AdvancedDisplay, Category = "Geometry & Tessellation Options", meta = (DisplayName = "Use CADKernel", Tooltip = "When true uses CADKernel tessellator. Otherwise, TechSoft's tessellator is used. Default is true.", editCondition = "bWithTechSoft == true"))
-	bool bUseCADKernel = true;
-
+	UPROPERTY(config, EditAnywhere, BlueprintReadWrite, AdvancedDisplay, Category = "Geometry & Tessellation Options", meta = (Tooltip = "Select tessellator to use. Default to CADKernel. Disabled if TechSoft is not available.", editCondition = "bWithTechSoft == true"))
+	TEnumAsByte<ECADKernelTessellator> Tessellator = ECADKernelTessellator::CADKernel;
+	
 	/**
-	 * Indicates which tessellator to use CADKernel or TechSoft. By default it is CADKernel.
+	 * T-Junctions are edge configurations where a border vertex lies on an opposing border edge.
+	 * If true, the mesh is analyzed and all T-Junctions found on the mesh will be resolved.
+	 * Only applicable with the TechSoft tessellator. Default is false
 	 */
-	UPROPERTY(config, EditAnywhere, BlueprintReadWrite, AdvancedDisplay, Category = "Geometry & Tessellation Options", meta = (Tooltip = "If true, enables the removal of T-Junctions from the resulting mesh. Default is false."))
+	UPROPERTY(config, EditAnywhere, BlueprintReadWrite, AdvancedDisplay, Category = "Geometry & Tessellation Options", meta = (EditCondition = "Tessellator == ECADKernelTessellator::TechSoft", EditConditionHides, Tooltip = "If true, enables the removal of T-Junctions from the resulting mesh. Default is false."))
 	bool bResolveTJunctions = false;
 
 protected:
@@ -165,9 +175,9 @@ public:
 	 */
 	// #cadkernel_check: Verify that this new logic is properly used everywhere.
 	UPROPERTY()
-	double UnitMultiplier = 0.1;
+	double UnitMultiplier = 1.;
 
-	UPROPERTY(Transient)
+	UPROPERTY(Transient, meta = (EditCondition = false, EditConditionHides))
 	bool bWithTechSoft = bWithHoops;
 
 	static bool bWithHoops;
@@ -178,13 +188,14 @@ public:
 			&& FMath::IsNearlyEqual(MaxEdgeLength, Other.MaxEdgeLength)
 			&& FMath::IsNearlyEqual(NormalTolerance, Other.NormalTolerance)
 			&& StitchingTechnique == Other.StitchingTechnique
+			&& Tessellator == Other.Tessellator
 			&& FMath::IsNearlyEqual(GeometricTolerance, Other.GeometricTolerance)
 			&& FMath::IsNearlyEqual(StitchingTolerance, Other.StitchingTolerance);
 	}
 
 	uint32 GetHash() const
 	{
-		uint32 Hash = uint32(StitchingTechnique);
+		uint32 Hash = HashCombine(GetTypeHash(StitchingTechnique), GetTypeHash(Tessellator));
 		for (double Param : {ChordTolerance, MaxEdgeLength, NormalTolerance, GeometricTolerance, StitchingTolerance})
 		{
 			Hash = HashCombine(Hash, GetTypeHash(Param));
@@ -199,6 +210,8 @@ public:
 
 	bool FromJson(const FString& JsonString) const { return false; }
 	FString ToJson() const { return {}; }
+
+	bool UseCADKernel() const { return Tessellator == ECADKernelTessellator::CADKernel; }
 };
 
 UENUM()
