Fixed crash when using DirectLink

#Synced-CL 49033479

diff --git a/Engine/Source/Runtime/Datasmith/DirectLink/Private/DirectLinkSceneGraphNode.cpp b/Engine/Source/Runtime/Datasmith/DirectLink/Private/DirectLinkSceneGraphNode.cpp
--- a/Engine/Source/Runtime/Datasmith/DirectLink/Private/DirectLinkSceneGraphNode.cpp
+++ b/Engine/Source/Runtime/Datasmith/DirectLink/Private/DirectLinkSceneGraphNode.cpp
@@ -39,7 +39,23 @@ void ISceneGraphNode::UpdateRefs(class IReferenceResolutionProvider& Resolver, c
 		{
 			if (Group.Name == Proxy.Name)
 			{
-				Proxy.View->SetNodes(Resolver, Group.ReferencedIds);
+				// The NodeIds in Group.ReferencedIds may not be valid using the given resolver
+				// In that case, the invalid NodeIds should not be set on Proxy's view
+				// Only valid nodes can be added to it, see TDatasmithReferenceArrayProxy::SetNodes
+				TArray<FSceneGraphId> ValidNodeIds;
+				ValidNodeIds.Reserve(Group.ReferencedIds.Num());
+				for (FSceneGraphId NodeId : Group.ReferencedIds)
+				{
+					if (TSharedPtr<ISceneGraphNode> Node = Resolver.AsSharedPtr(NodeId))
+					{
+						ValidNodeIds.Add(NodeId);
+					}
+				}
+
+				if (!ValidNodeIds.IsEmpty())
+				{
+					Proxy.View->SetNodes(Resolver, ValidNodeIds);
+				}
 				break;
 			}
 		}
